Engine source (Main branch up to CL 2027380)

diff --git a/Engine/Build/build.properties b/Engine/Build/build.properties
--- a/Engine/Build/build.properties
+++ b/Engine/Build/build.properties
@@ -1,3 +1,3 @@
-TimestampForBVT=2014-03-14_16.07.59
-EngineVersion=2026805
+TimestampForBVT=2014-03-15_16.07.01
+EngineVersion=2027327
 BranchName=++depot+UE4
diff --git a/Engine/Source/Programs/AutomationTool/Scripts/GUBP.Automation.cs b/Engine/Source/Programs/AutomationTool/Scripts/GUBP.Automation.cs
--- a/Engine/Source/Programs/AutomationTool/Scripts/GUBP.Automation.cs
+++ b/Engine/Source/Programs/AutomationTool/Scripts/GUBP.Automation.cs
@@ -778,6 +778,11 @@ public class GUBP : BuildCommand
                 {
                     AddPseudodependency(GamePlatformMonolithicsNode.StaticGetFullName(InHostPlatform, bp.Branch.BaseEngineProject, InHostPlatform));
                 }
+                if (GUBP.bBuildRocket && InGameProj.GameName == bp.Branch.BaseEngineProject.GameName)
+                {
+                    AgentSharingGroup = "UE4_" + InTargetPlatform + "_Mono" + StaticGetHostPlatformSuffix(InHostPlatform);
+                    AddPseudodependency(RootEditorHeadersNode.StaticGetFullName(HostPlatform)); // maybe we should start these sooner, but that rather tangles the agent groups
+                }
             }
             if (InGameProj.Options(InHostPlatform).bTestWithShared)  /// compiling templates is only for testing purposes, and we will group them to avoid saturating the farm
             {
@@ -877,6 +882,110 @@ public class GUBP : BuildCommand
         }
     }
 
+    public class RootEditorHeadersNode : HostPlatformNode
+    {
+        public RootEditorHeadersNode(UnrealTargetPlatform InHostPlatform)
+            : base(InHostPlatform)
+        {
+            AgentSharingGroup = "Editor" + StaticGetHostPlatformSuffix(HostPlatform);
+            AddDependency(RootEditorNode.StaticGetFullName(HostPlatform));
+        }
+        public static string StaticGetFullName(UnrealTargetPlatform InHostPlatform)
+        {
+            return "RootEditorHeaders" + StaticGetHostPlatformSuffix(InHostPlatform);
+        }
+        public override string GetFullName()
+        {
+            return StaticGetFullName(HostPlatform);
+        }
+        public override float Priority()
+        {
+            return 1000000.0f; // right after the root editor
+        }
+        public override void DoBuild(GUBP bp)
+        {
+            BuildProducts = new List<string>();
+
+            foreach (var FileToCopy in CommandUtils.FindFiles("*.h", true, CommandUtils.CombinePaths(CmdEnv.LocalRoot, @"Engine\Intermediate\Build\")))
+            {
+                AddBuildProduct(FileToCopy);
+            }
+            var EnginePluginsDirectory = Path.Combine(CommandUtils.CmdEnv.LocalRoot, "Engine/Plugins");
+            var EnginePlugins = new List<PluginInfo>();
+            Plugins.FindPluginsIn(EnginePluginsDirectory, PluginInfo.LoadedFromType.Engine, ref EnginePlugins);
+
+            foreach (var EnginePlugin in EnginePlugins)
+            {
+                foreach (var FileToCopy in CommandUtils.FindFiles("*", true, CommandUtils.CombinePaths(EnginePlugin.Directory, @"Intermediate\Build", HostPlatform.ToString(), "Inc")))
+                {
+                    AddBuildProduct(FileToCopy);
+                }
+            }
+            RemoveOveralppingBuildProducts();
+        }
+    }
+
+    public class GameMonolithicHeadersNode : HostPlatformNode
+    {
+        UnrealTargetPlatform TargetPlatform;
+
+        public GameMonolithicHeadersNode(GUBP bp, UnrealTargetPlatform InHostPlatform, UnrealTargetPlatform InTargetPlatform)
+            : base(InHostPlatform)
+        {
+            if (GUBP.bHackRunIOSCompilesOnMac)
+            {
+                throw new AutomationException("these rocket header node require real mac path.");
+            }
+            TargetPlatform = InTargetPlatform;
+            AgentSharingGroup = "UE4_" + InTargetPlatform + "_Mono" + StaticGetHostPlatformSuffix(InHostPlatform);
+            AddDependency(GamePlatformMonolithicsNode.StaticGetFullName(HostPlatform, bp.Branch.BaseEngineProject, TargetPlatform));
+            if (TargetPlatform == HostPlatform)
+            {
+                AddDependency(RootEditorHeadersNode.StaticGetFullName(HostPlatform)); // we can't tell the headers apart, so we need to remove dups
+            }
+        }
+        public static string StaticGetFullName(UnrealTargetPlatform InHostPlatform, UnrealTargetPlatform InTargetPlatform)
+        {
+            return "UE4_" + InTargetPlatform + "_MonolithicHeaders" + StaticGetHostPlatformSuffix(InHostPlatform);
+        }
+        public override string GetFullName()
+        {
+            return StaticGetFullName(HostPlatform, TargetPlatform);
+        }
+        public override float Priority()
+        {
+            return 1000000.0f; // right after the root editor
+        }
+        public override void DoBuild(GUBP bp)
+        {
+            BuildProducts = new List<string>();
+
+            // these may not be any new build products here, but we will check anyway
+            var EnginePluginsDirectory = Path.Combine(CommandUtils.CmdEnv.LocalRoot, "Engine/Plugins");
+            var EnginePlugins = new List<PluginInfo>();
+            Plugins.FindPluginsIn(EnginePluginsDirectory, PluginInfo.LoadedFromType.Engine, ref EnginePlugins);
+
+            foreach (var EnginePlugin in EnginePlugins)
+            {
+                var PlatformDir = TargetPlatform.ToString();
+                if (TargetPlatform == UnrealTargetPlatform.Android)
+                {
+                    PlatformDir = "Android-armv7";
+                }
+
+                foreach (var FileToCopy in CommandUtils.FindFiles("*", true, CommandUtils.CombinePaths(EnginePlugin.Directory, @"Intermediate\Build", PlatformDir, "Inc")))
+                {
+                    AddBuildProduct(FileToCopy);
+                }
+            }
+            RemoveOveralppingBuildProducts();
+            if (BuildProducts.Count == 0)
+            {
+                SaveRecordOfSuccessAndAddToBuildProducts(); // could be empty
+            }
+        }
+    }
+
     public class AggregateNode : GUBPNode
     {
         public AggregateNode()
@@ -2692,11 +2801,11 @@ public class GUBP : BuildCommand
             {
                 foreach (var Dep in GUBPNodes[NodeToDo].FullNamesOfDependencies)
                 {
-                    Log("           {0}", Dep);
+                    Log("            dep> {0}", Dep);
                 }
                 foreach (var Dep in GUBPNodes[NodeToDo].FullNamesOfPseudosependencies)
                 {
-                    Log("           {0} (pseudo)", Dep);
+                    Log("           pdep> {0}", Dep);
                 }
             }
             if (bShowECDependencies)
@@ -3073,6 +3182,10 @@ public class GUBP : BuildCommand
         {
             AddNode(new ToolsForCompileNode(HostPlatform));
             AddNode(new RootEditorNode(HostPlatform));
+            if (bBuildRocket)
+            {
+                AddNode(new RootEditorHeadersNode(HostPlatform));
+            }
             AddNode(new ToolsNode(HostPlatform));            
 			AddNode(new InternalToolsNode(HostPlatform));
             AddNode(new EditorAndToolsNode(HostPlatform));
@@ -3209,6 +3322,10 @@ public class GUBP : BuildCommand
                                 if (!GUBPNodes.ContainsKey(GamePlatformMonolithicsNode.StaticGetFullName(HostPlatform, Branch.BaseEngineProject, Plat)))
                                 {
                                     AddNode(new GamePlatformMonolithicsNode(this, HostPlatform, Branch.BaseEngineProject, Plat));
+                                    if (bBuildRocket)
+                                    {
+                                        AddNode(new GameMonolithicHeadersNode(this, HostPlatform, Plat));
+                                    }
                                 }
                             }
                         }
@@ -3582,7 +3699,15 @@ if (HostPlatform == UnrealTargetPlatform.Mac) continue; //temp hack till mac aut
             {
                 // rocket is the shared promotable plus some other stuff and nothing else
                 bRelatedToNode = true;
-                NodeSpec = SharedAggregatePromotableNode.StaticGetFullName() + "+" + "Rocket_Aggregate";
+                NodeSpec = "Rocket_Aggregate";
+                if (!ParseParam("RocketValidate"))
+                {
+                    NodeSpec = SharedAggregatePromotableNode.StaticGetFullName() + "+" + NodeSpec;
+                }
+                else
+                {
+                    NodeSpec = "Rocket_MakeBuild+" + NodeSpec;
+                }
             }
             if (!String.IsNullOrEmpty(NodeSpec))
             {
diff --git a/Engine/Source/Runtime/Core/Private/Linux/LinuxPlatformFile.cpp b/Engine/Source/Runtime/Core/Private/Linux/LinuxPlatformFile.cpp
--- a/Engine/Source/Runtime/Core/Private/Linux/LinuxPlatformFile.cpp
+++ b/Engine/Source/Runtime/Core/Private/Linux/LinuxPlatformFile.cpp
@@ -121,11 +121,13 @@ int64 FLinuxPlatformFile::FileSize(const TCHAR* Filename)
 {
 	struct stat FileInfo;
 	FileInfo.st_size = -1;
-	stat(TCHAR_TO_UTF8(*NormalizeFilename(Filename)), &FileInfo);
-	// make sure to return -1 for directories
-	if (S_ISDIR(FileInfo.st_mode))
+	if (stat(TCHAR_TO_UTF8(*NormalizeFilename(Filename)), &FileInfo) != -1)
 	{
-		FileInfo.st_size = -1;
+		// make sure to return -1 for directories
+		if (S_ISDIR(FileInfo.st_mode))
+		{
+			FileInfo.st_size = -1;
+		}
 	}
 	return FileInfo.st_size;
 }
diff --git a/Engine/Source/Runtime/Core/Private/Stats/Stats2.cpp b/Engine/Source/Runtime/Core/Private/Stats/Stats2.cpp
--- a/Engine/Source/Runtime/Core/Private/Stats/Stats2.cpp
+++ b/Engine/Source/Runtime/Core/Private/Stats/Stats2.cpp
@@ -175,6 +175,8 @@ public:
 	FStatGroupEnableManager()
 		: PendingFNames(NULL)
 		, PendingCount(0)
+		, EnableForNewGroups(false)
+		, UseEnableForNewGroups(false)
 	{
 		check(IsInGameThread());
 	}
diff --git a/Engine/Source/Runtime/Engine/Public/TimerManager.h b/Engine/Source/Runtime/Engine/Public/TimerManager.h
--- a/Engine/Source/Runtime/Engine/Public/TimerManager.h
+++ b/Engine/Source/Runtime/Engine/Public/TimerManager.h
@@ -145,6 +145,7 @@ public:
 
 	FTimerManager()
 		: InternalTime(0.0)
+		, LastTickedFrame(static_cast<uint64>(-1))
 	{}
 
 
diff --git a/Engine/Source/Runtime/Renderer/Private/PrimitiveSceneInfo.cpp b/Engine/Source/Runtime/Renderer/Private/PrimitiveSceneInfo.cpp
--- a/Engine/Source/Runtime/Renderer/Private/PrimitiveSceneInfo.cpp
+++ b/Engine/Source/Runtime/Renderer/Private/PrimitiveSceneInfo.cpp
@@ -97,7 +97,8 @@ FPrimitiveSceneInfo::FPrimitiveSceneInfo(UPrimitiveComponent* InComponent,FScene
 	LastVisibilityChangeTime(0.0f),
 	Scene(InScene),
 	PackedIndex(INDEX_NONE),
-	ComponentForDebuggingOnly(InComponent)
+	ComponentForDebuggingOnly(InComponent),
+	bNeedsStaticMeshUpdate(true)
 {
 	check(ComponentForDebuggingOnly);
 	check(PrimitiveComponentId.IsValid());
diff --git a/README.md b/README.md
--- a/README.md
+++ b/README.md
@@ -90,13 +90,13 @@ Visual Studio 2013 and Xcode 5.1 are the supported and recommended compilers to
 
 Visual Studio 2012 is also supported, but you'll need to make a code change and download the **optional** dependencies 
 on the [releases page](https://github.com/EpicGames/UnrealEngine/releases/tag/4.0.1-release).  See the Unreal Engine 
-documentation for more details about using older versions of Visual Studio.
+[documentation](http://docs.unrealengine.com) for more details about using older versions of Visual Studio.
 
 You should probably always work with a versioned **release** branch.  The master branch contains unstable and possibly untested code,
 but it should be a great reference for new developments, or for spot merging bug fixes.  Use it at your own risk.  
 Full engine releases tend to lag behind latest development branch due to QA testing, but we wanted to make sure that 
 you have access to the very latest stuff we're doing.
 
-To build for platforms besides Windows and Mac, please see the Unreal Engine documentation.  Some platforms may require
+To build for platforms besides Windows and Mac, please see the Unreal Engine [documentation](http://docs.unrealengine.com).  Some platforms may require
 you to download the **optional** dependencies as well.
 
