Engine source (Main branch up to CL 2027551)

diff --git a/Engine/Build/build.properties b/Engine/Build/build.properties
--- a/Engine/Build/build.properties
+++ b/Engine/Build/build.properties
@@ -1,3 +1,3 @@
-TimestampForBVT=2014-03-15_16.07.01
-EngineVersion=2027327
+TimestampForBVT=2014-03-16_04.09.05
+EngineVersion=2027525
 BranchName=++depot+UE4
diff --git a/Engine/Source/Programs/AutomationTool/Scripts/GUBP.Automation.cs b/Engine/Source/Programs/AutomationTool/Scripts/GUBP.Automation.cs
--- a/Engine/Source/Programs/AutomationTool/Scripts/GUBP.Automation.cs
+++ b/Engine/Source/Programs/AutomationTool/Scripts/GUBP.Automation.cs
@@ -1662,7 +1662,7 @@ public class GUBP : BuildCommand
         }
         public static string StaticGetFullName(UnrealTargetPlatform InHostPlatform, BranchInfo.BranchUProject InGameProj, string InCookPlatform)
         {
-            return InGameProj.GameName + "_" + InCookPlatform + "_DDC" + StaticGetHostPlatformSuffix(InHostPlatform);
+            return InGameProj.GameName + "_" + InCookPlatform.Replace("+", "_") + "_DDC" + StaticGetHostPlatformSuffix(InHostPlatform);
         }
         public override string GetFullName()
         {
