build script, mark .sh and .command as +x on checkin

diff --git a/Engine/Source/Programs/AutomationTool/BuildUtils.cs b/Engine/Source/Programs/AutomationTool/BuildUtils.cs
--- a/Engine/Source/Programs/AutomationTool/BuildUtils.cs
+++ b/Engine/Source/Programs/AutomationTool/BuildUtils.cs
@@ -100,10 +100,12 @@ namespace AutomationTool
         public static bool IsProbablyAMacOrIOSExe(string Filename)
         {
             return
-                (
+                Path.GetExtension(Filename) == ".sh" ||
+                Path.GetExtension(Filename) == ".command" ||
+                ((
                     CommandUtils.CombinePaths(Filename).ToLower().Contains(CommandUtils.CombinePaths("Binaries", "Mac").ToLower()) ||
                     CommandUtils.CombinePaths(Filename).ToLower().Contains(CommandUtils.CombinePaths("Binaries", "IOS").ToLower())
-                ) && (Path.GetExtension(Filename) == "" || Path.GetExtension(Filename) == ".");
+                ) && (Path.GetExtension(Filename) == "" || Path.GetExtension(Filename) == "."));
         }
 	}
 
