build script, retry on a weird XGE error

diff --git a/Engine/Source/Programs/AutomationTool/ProcessUtils.cs b/Engine/Source/Programs/AutomationTool/ProcessUtils.cs
--- a/Engine/Source/Programs/AutomationTool/ProcessUtils.cs
+++ b/Engine/Source/Programs/AutomationTool/ProcessUtils.cs
@@ -618,6 +618,22 @@ namespace AutomationTool
 			return Result;
 		}
 
+        /// <summary>
+        /// Gets a logfile name for a RunAndLog call
+        /// </summary>
+        /// <param name="Env">Environment to use.</param>
+        /// <param name="App">Executable to run</param>
+        /// <param name="LogName">Name of the logfile ( if null, executable name is used )</param>
+        /// <returns>The log file name.</returns>
+        public static string GetRunAndLogLogName(CommandEnvironment Env, string App, string LogName = null)
+        {
+            if (LogName == null)
+            {
+                LogName = Path.GetFileNameWithoutExtension(App);
+            }
+            return LogUtils.GetUniqueLogName(CombinePaths(Env.LogFolder, LogName));
+        }
+
 		/// <summary>
 		/// Runs external program and writes the output to a logfile.
 		/// </summary>
@@ -625,17 +641,66 @@ namespace AutomationTool
 		/// <param name="App">Executable to run</param>
 		/// <param name="CommandLine">Commandline to pass on to the executable</param>
 		/// <param name="LogName">Name of the logfile ( if null, executable name is used )</param>
-		/// <returns>Whether the program executed successfully or not.</returns>
 		public static void RunAndLog(CommandEnvironment Env, string App, string CommandLine, string LogName = null, int MaxSuccessCode = 0)
 		{
-			if (LogName == null)
-			{
-				LogName = Path.GetFileNameWithoutExtension(App);
-			}
-			string LogFile = LogUtils.GetUniqueLogName(CombinePaths(Env.LogFolder, LogName));
-
-			RunAndLog(App, CommandLine, LogFile, MaxSuccessCode);
+            RunAndLog(App, CommandLine, GetRunAndLogLogName(Env, App, LogName), MaxSuccessCode);
 		}
+        /// <summary>
+        /// Runs external program and writes the output to a logfile.
+        /// </summary>
+        /// <param name="App">Executable to run</param>
+        /// <param name="CommandLine">Commandline to pass on to the executable</param>
+        /// <param name="Logfile">Full path to the logfile, where the application output should be written to.</param>
+        /// <returns>Whether the program executed successfully or not.</returns>
+        public static void RunAndLog(string App, string CommandLine, string Logfile = null, int MaxSuccessCode = 0)
+        {
+            Log("Running {0} {1}", App, CommandLine);
+
+            ProcessResult Result = Run(App, CommandLine);
+            if (Result.Output.Length > 0 && Logfile != null)
+            {
+                WriteToFile(Logfile, Result.Output);
+            }
+            else if (Logfile == null)
+            {
+                Logfile = "[No logfile specified]";
+            }
+            else
+            {
+                Logfile = "[None!, no output produced]";
+            }
+
+            if (Result > MaxSuccessCode || Result < 0)
+            {
+                throw new AutomationException(String.Format("Command failed (Result:{3}): {0} {1}. See logfile for details: '{2}' ",
+                                                App, CommandLine, Path.GetFileName(Logfile), Result.ExitCode));
+            }
+        }
+
+        /// <summary>
+        /// Runs external program and writes the output to a logfile.
+        /// </summary>
+        /// <param name="App">Executable to run</param>
+        /// <param name="CommandLine">Commandline to pass on to the executable</param>
+        /// <param name="Logfile">Full path to the logfile, where the application output should be written to.</param>
+        /// <returns>Whether the program executed successfully or not.</returns>
+        public static string RunAndLog(string App, string CommandLine, out int SuccessCode, string Logfile = null)
+        {
+            Log("Running {0} {1}", App, CommandLine);
+
+            ProcessResult Result = Run(App, CommandLine);
+            SuccessCode = Result.ExitCode;
+            if (Result.Output.Length > 0 && Logfile != null)
+            {
+                WriteToFile(Logfile, Result.Output);
+            }
+            if (!String.IsNullOrEmpty(Result.Output))
+            {
+                return Result.Output;
+            }
+            return "";
+        }
+
 		/// <summary>
 		/// Runs UAT recursively
 		/// </summary>
@@ -739,38 +804,6 @@ namespace AutomationTool
 			return LogFile;
 		}
 
-		/// <summary>
-		/// Runs external program and writes the output to a logfile.
-		/// </summary>
-		/// <param name="App">Executable to run</param>
-		/// <param name="CommandLine">Commandline to pass on to the executable</param>
-		/// <param name="Logfile">Full path to the logfile, where the application output should be written to.</param>
-		/// <returns>Whether the program executed successfully or not.</returns>
-		public static void RunAndLog(string App, string CommandLine, string Logfile = null, int MaxSuccessCode = 0)
-		{
-			Log("Running {0} {1}", App, CommandLine);
-
-			ProcessResult Result = Run(App, CommandLine);
-			if (Result.Output.Length > 0 && Logfile != null)
-			{
-				WriteToFile(Logfile, Result.Output);
-			}
-			else if (Logfile == null)
-			{
-				Logfile = "[No logfile specified]";
-			}
-			else
-			{
-				Logfile = "[None!, no output produced]";
-			}
-
-			if (Result > MaxSuccessCode || Result < 0)
-			{
-				throw new AutomationException(String.Format("Command failed (Result:{3}): {0} {1}. See logfile for details: '{2}' ",
-												App, CommandLine, Path.GetFileName(Logfile), Result.ExitCode));
-			}
-		}
-
 		protected delegate bool ProcessLog(string LogText);
 		/// <summary>
 		/// Keeps reading a log file as it's being written to by another process until it exits.
diff --git a/Engine/Source/Programs/AutomationTool/UE4Build.cs b/Engine/Source/Programs/AutomationTool/UE4Build.cs
--- a/Engine/Source/Programs/AutomationTool/UE4Build.cs
+++ b/Engine/Source/Programs/AutomationTool/UE4Build.cs
@@ -996,21 +996,44 @@ namespace AutomationTool
 				string Args = "\"" + TaskFilePath + "\" /Rebuild /MaxCPUS=200";
 
 				int Retries = DoRetries ? 2 : 1;
-				for (int i = 0; i < Retries; i++)
+                int ConnectionRetries = 4;
+                for (int i = 0; i < Retries; i++)
 				{
 					try
 					{
-						Log("Running XGE *******");
-						PushDir(CombinePaths(CmdEnv.LocalRoot, @"\Engine\Source"));
-						RunAndLog(CmdEnv, XGEConsole, Args);
-						PopDir();
-						Log("XGE Done *******");
-						break;
+                        while (true)
+                        {
+                            Log("Running XGE *******");
+                            PushDir(CombinePaths(CmdEnv.LocalRoot, @"\Engine\Source"));
+                            int SuccesCode;
+                            string LogFile = GetRunAndLogLogName(CmdEnv, XGEConsole);
+                            string Output = RunAndLog(XGEConsole, Args, out SuccesCode, LogFile);
+                            PopDir();
+                            if (ConnectionRetries > 0 && SuccesCode == 4 && Output.Contains("Timed out waiting for reply from local connection"))
+                            {
+                                Log(System.Diagnostics.TraceEventType.Warning, "XGE failure on the local connection timeout");
+                                if (ConnectionRetries < 2)
+                                {
+                                    System.Threading.Thread.Sleep(60000);
+                                }
+                                ConnectionRetries--;
+                                continue;
+                            }
+                            else if (SuccesCode != 0)
+                            {
+                                Log("XGE did not succeed *******");
+                                throw new AutomationException(String.Format("Command failed (Result:{3}): {0} {1}. See logfile for details: '{2}' ",
+                                                                XGEConsole, Args, Path.GetFileName(LogFile), SuccesCode));
+                            }
+                            Log("XGE Done *******");
+                            break;
+                        }
+                        break;
 					}
 					catch (Exception Ex)
 					{
 						Log("XGE failed on try {0}: {1}", i + 1, Ex.ToString());
-						if (i + 1 == Retries)
+						if (i + 1 >= Retries)
 						{
 							return false;
 						}
