Improved packaging for Mac (self-contained bundle, bug fixes)

diff --git a/Engine/Source/Programs/AutomationTool/Mac/MacPlatform.Automation.cs b/Engine/Source/Programs/AutomationTool/Mac/MacPlatform.Automation.cs
--- a/Engine/Source/Programs/AutomationTool/Mac/MacPlatform.Automation.cs
+++ b/Engine/Source/Programs/AutomationTool/Mac/MacPlatform.Automation.cs
@@ -3,6 +3,7 @@ using System;
 using System.Collections.Generic;
 using System.Linq;
 using System.Text;
+using System.IO;
 using AutomationTool;
 using UnrealBuildTool;
 
@@ -25,6 +26,11 @@ public class MacPlatform : Platform
         return "Mac";
     }
 
+	private void StageAppBundle(DeploymentContext SC, string InPath, string NewName)
+	{
+		SC.StageFiles(StagedFileType.NonUFS, InPath, "*", true, null, NewName, false, true, null);
+	}
+
 	public override void GetFilesToDeployOrStage(ProjectParams Params, DeploymentContext SC)
 	{
 		List<string> Exes = GetExecutableNames(SC);
@@ -32,11 +38,11 @@ public class MacPlatform : Platform
 		{
 			if (Exe.StartsWith(CombinePaths(SC.RuntimeProjectRootDir, "Binaries", SC.PlatformDir)))
 			{
-				SC.StageFiles(StagedFileType.NonUFS, CombinePaths(SC.ProjectRoot, "Binaries", SC.PlatformDir, System.IO.Path.GetFileNameWithoutExtension(Exe) + ".app"));
+				StageAppBundle(SC, CombinePaths(SC.ProjectRoot, "Binaries", SC.PlatformDir, Path.GetFileNameWithoutExtension(Exe) + ".app"), SC.ShortProjectName + ".app");
 			}
 			else if (Exe.StartsWith(CombinePaths(SC.RuntimeRootDir, "Engine/Binaries", SC.PlatformDir)))
 			{
-				SC.StageFiles(StagedFileType.NonUFS, CombinePaths(SC.LocalRoot, "Engine/Binaries", SC.PlatformDir, System.IO.Path.GetFileNameWithoutExtension(Exe) + ".app"));
+				StageAppBundle(SC, CombinePaths(SC.LocalRoot, "Engine/Binaries", SC.PlatformDir, Path.GetFileNameWithoutExtension(Exe) + ".app"), SC.ShortProjectName + ".app");
 			}
 		}
 
@@ -50,6 +56,54 @@ public class MacPlatform : Platform
 		PrintRunTime();
 	}
 
+	public override void ProcessArchivedProject(ProjectParams Params, DeploymentContext SC)
+	{
+		string TargetPath = CombinePaths(SC.ArchiveDirectory, SC.ShortProjectName + ".app", "Contents", "UE4");
+		if (DirectoryExists(TargetPath))
+		{
+			Directory.Delete(TargetPath, true);
+		}
+
+		// First, move all files and folders inside he app bundle
+		string[] StagedFiles = Directory.GetFiles(SC.ArchiveDirectory, "*", SearchOption.TopDirectoryOnly);
+		foreach (string FilePath in StagedFiles)
+		{
+			string TargetFilePath = CombinePaths(TargetPath, Path.GetFileName(FilePath));
+			Directory.CreateDirectory(Path.GetDirectoryName(TargetFilePath));
+			File.Move(FilePath, TargetFilePath);
+		}
+
+		string[] StagedDirectories = Directory.GetDirectories(SC.ArchiveDirectory, "*", SearchOption.TopDirectoryOnly);
+		foreach (string DirPath in StagedDirectories)
+		{
+			string DirName = Path.GetFileName(DirPath);
+			if (!DirName.EndsWith(".app"))
+			{
+				string TargetDirPath = CombinePaths(TargetPath, DirName);
+				Directory.CreateDirectory(Path.GetDirectoryName(TargetDirPath));
+				Directory.Move(DirPath, TargetDirPath);
+			}
+		}
+
+		// Update executable name and entry in Info.plist
+		string ExeName = SC.StageExecutables[0];
+		string UE4GamePath = CombinePaths(SC.ArchiveDirectory, SC.ShortProjectName + ".app", "Contents", "MacOS", ExeName);
+		if (File.Exists(UE4GamePath))
+		{
+			string GameExePath = CombinePaths(SC.ArchiveDirectory, SC.ShortProjectName + ".app", "Contents", "MacOS", SC.ShortProjectName);
+			File.Delete(GameExePath);
+			File.Move(UE4GamePath, GameExePath);
+
+			string InfoPlistPath = CombinePaths(SC.ArchiveDirectory, SC.ShortProjectName + ".app", "Contents", "Info.plist");
+			string InfoPlistContents = File.ReadAllText(InfoPlistPath);
+			InfoPlistContents = InfoPlistContents.Replace(ExeName, SC.ShortProjectName);
+			File.Delete(InfoPlistPath);
+			File.WriteAllText(InfoPlistPath, InfoPlistContents);
+		}
+
+		Directory.CreateDirectory(CombinePaths(TargetPath, ExeName.StartsWith("UE4Game") ? "Engine" : SC.ShortProjectName, "Binaries", "Mac"));
+	}
+
 	public override bool IsSupported { get { return true; } }
 
 	public override bool ShouldUseManifestForUBTBuilds()
diff --git a/Engine/Source/Programs/AutomationTool/Platform.cs b/Engine/Source/Programs/AutomationTool/Platform.cs
--- a/Engine/Source/Programs/AutomationTool/Platform.cs
+++ b/Engine/Source/Programs/AutomationTool/Platform.cs
@@ -90,6 +90,15 @@ namespace AutomationTool
 			throw new AutomationException("{0} does not yet implement Packaging.", PlatformType);
 		}
 
+		/// <summary>
+		/// Allow platform to do platform specific work on archived project before it's deployed.
+		/// </summary>
+		/// <param name="Params"></param>
+		/// <param name="SC"></param>
+		public virtual void ProcessArchivedProject(ProjectParams Params, DeploymentContext SC)
+		{
+		}
+
 		/// <summary>
 		/// Deploy the application on the current platform
 		/// </summary>
diff --git a/Engine/Source/Programs/AutomationTool/Scripts/ArchiveCommand.Automation.cs b/Engine/Source/Programs/AutomationTool/Scripts/ArchiveCommand.Automation.cs
--- a/Engine/Source/Programs/AutomationTool/Scripts/ArchiveCommand.Automation.cs
+++ b/Engine/Source/Programs/AutomationTool/Scripts/ArchiveCommand.Automation.cs
@@ -56,6 +56,7 @@ public partial class Project : CommandUtils
 			{
 				CreateArchiveManifest(Params, SC);
 				ApplyArchiveManifest(Params, SC);
+				SC.StageTargetPlatform.ProcessArchivedProject(Params, SC);
 			}
 		}
 		if (Params.DedicatedServer)
@@ -67,6 +68,7 @@ public partial class Project : CommandUtils
 			{
 				CreateArchiveManifest(Params, SC);
 				ApplyArchiveManifest(Params, SC);
+				SC.StageTargetPlatform.ProcessArchivedProject(Params, SC);
 			}
 		}
 	}
diff --git a/Engine/Source/Runtime/Core/Private/Mac/MacPlatformProcess.cpp b/Engine/Source/Runtime/Core/Private/Mac/MacPlatformProcess.cpp
--- a/Engine/Source/Runtime/Core/Private/Mac/MacPlatformProcess.cpp
+++ b/Engine/Source/Runtime/Core/Private/Mac/MacPlatformProcess.cpp
@@ -546,7 +546,15 @@ const TCHAR* FMacPlatformProcess::BaseDir()
 		// If it has .app extension, it's a bundle, otherwise BasePath is a full path to Binaries/Mac (in case of command line tools)
 		if ([[BasePath pathExtension] isEqual: @"app"])
 		{
-			NSString* BundledBinariesPath = [BasePath stringByAppendingPathComponent: @"Contents/UE4/Engine/Binaries/Mac"];
+			NSString* BundledBinariesPath = NULL;
+			if (GGameName[0] != 0)
+			{
+				BundledBinariesPath = [BasePath stringByAppendingPathComponent: [NSString stringWithFormat: @"Contents/UE4/%s/Binaries/Mac", TCHAR_TO_UTF8(GGameName)]];
+			}
+			if (!BundledBinariesPath)
+			{
+				BundledBinariesPath = [BasePath stringByAppendingPathComponent: @"Contents/UE4/Engine/Binaries/Mac"];
+			}
 			if ([FileManager fileExistsAtPath: BundledBinariesPath])
 			{
 				BasePath = BundledBinariesPath;
