#ue4
* fix for paths with spaces in them when using IPP

diff --git a/Engine/Source/Programs/AutomationTool/IOS/IOSPlatform.Automation.cs b/Engine/Source/Programs/AutomationTool/IOS/IOSPlatform.Automation.cs
--- a/Engine/Source/Programs/AutomationTool/IOS/IOSPlatform.Automation.cs
+++ b/Engine/Source/Programs/AutomationTool/IOS/IOSPlatform.Automation.cs
@@ -148,8 +148,8 @@ public class IOSPlatform : Platform
 			}
 
 			IPPArguments += (cookonthefly ? " -cookonthefly" : "");
-			IPPArguments += " -stagedir " + CombinePaths(Params.BaseStageDirectory, "IOS");
-			IPPArguments += " -projectdir " + Path.GetDirectoryName(Params.RawProjectPath);
+			IPPArguments += " -stagedir \"" + CombinePaths(Params.BaseStageDirectory, "IOS") + "\"";
+			IPPArguments += " -projectdir \"" + Path.GetDirectoryName(Params.RawProjectPath) + "\"";
 
 			// rename the .ipa if not code based
 			if (!Params.IsCodeBasedProject)
@@ -529,7 +529,7 @@ public class IOSPlatform : Platform
 			var IPPExe = CombinePaths(CmdEnv.LocalRoot, "Engine/Binaries/DotNET/IOS/IPhonePackager.exe");
 
 			// check for it in the stage directory
-			RunAndLog(CmdEnv, IPPExe, "Deploy " + Path.GetFullPath(StagedIPA) + " -device " + Params.Device.Substring(4));
+			RunAndLog(CmdEnv, IPPExe, "Deploy \"" + Path.GetFullPath(StagedIPA) + "\" -device " + Params.Device.Substring(4));
 		}
         PrintRunTime();
     }
