build script, fix local P4 environment

diff --git a/Engine/Source/Programs/AutomationTool/LocalP4Environment.cs b/Engine/Source/Programs/AutomationTool/LocalP4Environment.cs
--- a/Engine/Source/Programs/AutomationTool/LocalP4Environment.cs
+++ b/Engine/Source/Programs/AutomationTool/LocalP4Environment.cs
@@ -73,7 +73,7 @@ namespace AutomationTool
 			var CLString = CommandUtils.GetEnvVar(EnvVarNames.Changelist);
 			if (String.IsNullOrEmpty(CLString))
 			{
-				DetectCurrentCL(Connection, ClientRootPath);
+                CLString = DetectCurrentCL(Connection, ClientRootPath);
 			}
 			CommandUtils.ConditionallySetEnvVar(EnvVarNames.Changelist, CLString);
 
diff --git a/Engine/Source/Programs/AutomationTool/P4Environment.cs b/Engine/Source/Programs/AutomationTool/P4Environment.cs
--- a/Engine/Source/Programs/AutomationTool/P4Environment.cs
+++ b/Engine/Source/Programs/AutomationTool/P4Environment.cs
@@ -58,6 +58,13 @@ namespace AutomationTool
 			if (((CommandUtils.P4Enabled || CommandUtils.IsBuildMachine) && (ClientRoot == String.Empty || User == String.Empty ||
 				ChangelistString == String.Empty || Client == String.Empty || BuildRootP4 == String.Empty)))
 			{
+                Log.TraceInformation("P4Enabled={0}", CommandUtils.P4Enabled);
+                Log.TraceInformation("ClientRoot={0}",ClientRoot );
+                Log.TraceInformation("User={0}", User);
+                Log.TraceInformation("ChangelistString={0}", ChangelistString);
+                Log.TraceInformation("Client={0}", Client);
+                Log.TraceInformation("BuildRootP4={0}", BuildRootP4);
+
 				throw new AutomationException("BUILD FAILED Perforce Environment is not set up correctly. Please check your environment variables.");
 			}
 
