build script, delete stale DLLs after building the root editor or retrieving it from shared temp storage

diff --git a/Engine/Source/Programs/AutomationTool/Scripts/GUBP.Automation.cs b/Engine/Source/Programs/AutomationTool/Scripts/GUBP.Automation.cs
--- a/Engine/Source/Programs/AutomationTool/Scripts/GUBP.Automation.cs
+++ b/Engine/Source/Programs/AutomationTool/Scripts/GUBP.Automation.cs
@@ -96,6 +96,9 @@ public class GUBP : BuildCommand
             BuildProducts = new List<string>();
             SaveRecordOfSuccessAndAddToBuildProducts();
         }
+        public virtual void PostLoadFromSharedTempStorage(GUBP bp)
+        {
+        }
         public virtual void DoFakeBuild(GUBP bp) // this is used to more rapidly test a build system, it does nothing but save a record of success as a build product
         {
             BuildProducts = new List<string>();
@@ -374,6 +377,9 @@ public class GUBP : BuildCommand
         public virtual void PostBuild(GUBP bp, UE4Build UE4Build)
         {
         }
+        public virtual void PostBuildProducts(GUBP bp)
+        {
+        }
         public virtual bool DeleteBuildProducts()
         {
             return false;
@@ -394,13 +400,13 @@ public class GUBP : BuildCommand
                 {
                     AddBuildProduct(Product);
                 }
-
                 RemoveOveralppingBuildProducts();
                 if (bp.bSignBuildProducts)
                 {
                     // Sign everything we built
                     CodeSign.SignMultipleIfEXEOrDLL(bp, BuildProducts);
                 }
+                PostBuildProducts(bp);
             }
             else
             {
@@ -502,6 +508,72 @@ public class GUBP : BuildCommand
             }
             return Agenda;
         }
+        void DeleteStaleDLLs(GUBP bp)
+        {
+            var Targets = new List<string>{bp.Branch.BaseEngineProject.Properties.Targets[TargetRules.TargetType.Editor].TargetName};
+            foreach (var ProgramTarget in bp.Branch.BaseEngineProject.Properties.Programs)
+            {
+                if (ProgramTarget.Rules.GUBP_AlwaysBuildWithBaseEditor() && ProgramTarget.Rules.SupportsPlatform(HostPlatform))
+                {
+                    Targets.Add(ProgramTarget.TargetName);
+                }
+            }
+
+
+            foreach (var Target in Targets)
+            {
+                var EnginePlatformBinaries = CommandUtils.CombinePaths(CommandUtils.CmdEnv.LocalRoot, "Engine", "Binaries", HostPlatform.ToString());
+                var Wildcard = Target + "-*";
+                Log("************Deleting stale editor DLLs, path {0} wildcard {1}", EnginePlatformBinaries, Wildcard);
+                foreach (var DiskFile in FindFiles(Wildcard, true, EnginePlatformBinaries))
+                {
+                    bool IsBuildProduct = false;
+                    foreach (var Product in BuildProducts)
+                    {
+                        if (Product.Equals(DiskFile, StringComparison.InvariantCultureIgnoreCase))
+                        {
+                            IsBuildProduct = true;
+                            break;
+                        }
+                    }
+                    if (!IsBuildProduct)
+                    {
+                        DeleteFile(DiskFile);
+                    }
+                }
+                var EnginePluginBinaries = CommandUtils.CombinePaths(CommandUtils.CmdEnv.LocalRoot, "Engine", "Plugins");
+                var HostSubstring = CommandUtils.CombinePaths("/", HostPlatform.ToString(), "/");
+                Log("************Deleting stale editor DLLs, path {0} wildcard {1} host {2}", EnginePluginBinaries, Wildcard, HostSubstring);
+                foreach (var DiskFile in FindFiles(Wildcard, true, EnginePluginBinaries))
+                {
+                    if (DiskFile.IndexOf(HostSubstring, StringComparison.InvariantCultureIgnoreCase) < 0)
+                    {
+                        continue;
+                    }
+                    bool IsBuildProduct = false;
+                    foreach (var Product in BuildProducts)
+                    {
+                        if (Product.Equals(DiskFile, StringComparison.InvariantCultureIgnoreCase))
+                        {
+                            IsBuildProduct = true;
+                            break;
+                        }
+                    }
+                    if (!IsBuildProduct)
+                    {
+                        DeleteFile(DiskFile);
+                    }
+                }
+            }
+        }
+        public override void PostLoadFromSharedTempStorage(GUBP bp)
+        {
+            DeleteStaleDLLs(bp);
+        }
+        public override void PostBuildProducts(GUBP bp)
+        {
+            DeleteStaleDLLs(bp);
+        }
     }
 
     public class ToolsNode : CompileNode
@@ -5007,7 +5079,12 @@ if (HostPlatform == UnrealTargetPlatform.Mac) continue; //temp hack till mac aut
                 else
                 {
                     Log("***** Retrieving GUBP Node {0} from {1}", GUBPNodes[NodeToDo].GetFullName(), NodeStoreName);
-                    GUBPNodes[NodeToDo].BuildProducts = RetrieveFromTempStorage(CmdEnv, NodeStoreName, GameNameIfAny, StorageRootIfAny);
+                    bool WasLocal;
+                    GUBPNodes[NodeToDo].BuildProducts = RetrieveFromTempStorage(CmdEnv, NodeStoreName, out WasLocal, GameNameIfAny, StorageRootIfAny);
+                    if (!WasLocal)
+                    {
+                        GUBPNodes[NodeToDo].PostLoadFromSharedTempStorage(this);
+                    }
                 }
             }
             else
diff --git a/Engine/Source/Programs/AutomationTool/Scripts/Tests.Automation.cs b/Engine/Source/Programs/AutomationTool/Scripts/Tests.Automation.cs
--- a/Engine/Source/Programs/AutomationTool/Scripts/Tests.Automation.cs
+++ b/Engine/Source/Programs/AutomationTool/Scripts/Tests.Automation.cs
@@ -236,7 +236,8 @@ class TestTempStorage : BuildCommand
 		DeleteFile(UnitTestFile2);
 		DeleteFile(UnitTestFile3);
 
-		RetrieveFromTempStorage(CmdEnv, "Test");
+        bool WasLocal;
+        RetrieveFromTempStorage(CmdEnv, "Test", out WasLocal);
 		if (!LocalTempStorageExists(CmdEnv, "Test"))
 		{
 			throw new AutomationException("local storage should exist");
@@ -290,8 +291,8 @@ class TestTempStorage : BuildCommand
 		{
 			throw new AutomationException("shared storage should not exist");
 		}
-		RetrieveFromTempStorage(CmdEnv, "Test"); // this should just rely on the local
-		if (!LocalTempStorageExists(CmdEnv, "Test"))
+        RetrieveFromTempStorage(CmdEnv, "Test", out WasLocal); // this should just rely on the local
+		if (!WasLocal || !LocalTempStorageExists(CmdEnv, "Test"))
 		{
 			throw new AutomationException("local storage should exist");
 		}
@@ -328,7 +329,7 @@ class TestTempStorage : BuildCommand
 			bool bFailedProperly = false;
 			try
 			{
-				RetrieveFromTempStorage(CmdEnv, "Test");
+                RetrieveFromTempStorage(CmdEnv, "Test", out WasLocal);
 			}
 			catch (AutomationException)
 			{
@@ -346,7 +347,7 @@ class TestTempStorage : BuildCommand
 			bool bFailedProperly = false;
 			try
 			{
-				RetrieveFromTempStorage(CmdEnv, "Test");
+                RetrieveFromTempStorage(CmdEnv, "Test", out WasLocal);
 			}
 			catch (AutomationException)
 			{
diff --git a/Engine/Source/Programs/AutomationTool/TempStorage.cs b/Engine/Source/Programs/AutomationTool/TempStorage.cs
--- a/Engine/Source/Programs/AutomationTool/TempStorage.cs
+++ b/Engine/Source/Programs/AutomationTool/TempStorage.cs
@@ -699,7 +699,7 @@ namespace AutomationTool
 
         }
 
-        public static List<string> RetrieveFromTempStorage(CommandEnvironment Env, string StorageBlockName, string GameFolder = "", string BaseFolder = "")
+        public static List<string> RetrieveFromTempStorage(CommandEnvironment Env, string StorageBlockName, out bool WasLocal, string GameFolder = "", string BaseFolder = "")
         {
             if (String.IsNullOrEmpty(BaseFolder))
             {
@@ -726,8 +726,10 @@ namespace AutomationTool
                 {
                     throw new AutomationException("Local files in manifest {0} were tampered with.", LocalManifest);
                 }
+                WasLocal = true;
                 return Files;
             }
+            WasLocal = false;
 
             var BlockPath = CombinePaths(SharedTempStorageDirectory(StorageBlockName, GameFolder, false), "/");
             if (!BlockPath.EndsWith("/") && !BlockPath.EndsWith("\\"))
