#UE4 - UAT: Replaced -AllowStdOutTraceEventType with -AllowStdOutLogVerbosity

diff --git a/Engine/Source/Programs/AutomationTool/CommandletUtils.cs b/Engine/Source/Programs/AutomationTool/CommandletUtils.cs
--- a/Engine/Source/Programs/AutomationTool/CommandletUtils.cs
+++ b/Engine/Source/Programs/AutomationTool/CommandletUtils.cs
@@ -158,7 +158,7 @@ namespace AutomationTool
 
 			string LogFile = LogUtils.GetUniqueLogName(CombinePaths(CmdEnv.LogFolder, Commandlet));
 			Log("Commandlet log file is {0}", LogFile);
-            var RunResult = Run(EditorExe, String.Format("{0} -run={1} {2} -abslog={3} -stdout -FORCELOGFLUSH -CrashForUAT -unattended -AllowStdOutTraceEventType {4}", 
+			var RunResult = Run(EditorExe, String.Format("{0} -run={1} {2} -abslog={3} -stdout -FORCELOGFLUSH -CrashForUAT -unattended -AllowStdOutLogVerbosity {4}", 
                 CommandUtils.MakePathSafeToUseWithCommandLine(ProjectName), 
                 Commandlet, 
 				String.IsNullOrEmpty(Parameters) ? "" : Parameters, 
diff --git a/Engine/Source/Programs/AutomationTool/Scripts/CookCommand.Automation.cs b/Engine/Source/Programs/AutomationTool/Scripts/CookCommand.Automation.cs
--- a/Engine/Source/Programs/AutomationTool/Scripts/CookCommand.Automation.cs
+++ b/Engine/Source/Programs/AutomationTool/Scripts/CookCommand.Automation.cs
@@ -70,7 +70,7 @@ public partial class Project : CommandUtils
 
 		try
 		{
-			var CommandletParams = "-unattended -buildmachine -forcelogflush -AllowStdOutTraceEventType -Unversioned -fileopenlog";
+			var CommandletParams = "-unattended -buildmachine -forcelogflush -AllowStdOutLogVerbosity -Unversioned -fileopenlog";
 			if (Params.Compressed)
 			{
 				CommandletParams += " -Compressed";
