#ue4
#ios
* fixed problem with non-development configs with PC IPAs

diff --git a/Engine/Source/Programs/AutomationTool/IOS/IOSPlatform.Automation.cs b/Engine/Source/Programs/AutomationTool/IOS/IOSPlatform.Automation.cs
--- a/Engine/Source/Programs/AutomationTool/IOS/IOSPlatform.Automation.cs
+++ b/Engine/Source/Programs/AutomationTool/IOS/IOSPlatform.Automation.cs
@@ -155,6 +155,10 @@ public class IOSPlatform : Platform
 			if (!Params.IsCodeBasedProject)
 			{
 				ProjectIPA = Path.Combine(Path.GetDirectoryName(Params.RawProjectPath), "Binaries\\IOS", Params.ShortProjectName + ".ipa");
+				if (TargetConfiguration != UnrealTargetConfiguration.Development)
+				{
+					ProjectIPA = Path.Combine(Path.GetDirectoryName(Params.RawProjectPath), "Binaries\\IOS", Params.ShortProjectName + "-" + PlatformType.ToString() + "-" + TargetConfiguration.ToString() + ".ipa");
+				}
 			}
 
             // delete the .ipa to make sure it was made
@@ -509,9 +513,14 @@ public class IOSPlatform : Platform
 
 //			var ProjectStub = Params.ProjectGameExeFilename;	
 			var ProjectIPA = MakeIPAFileName(TargetConfiguration, Params.ProjectGameExeFilename);
+			// rename the .ipa if not code based
 			if (!Params.IsCodeBasedProject)
 			{
 				ProjectIPA = Path.Combine(Path.GetDirectoryName(Params.RawProjectPath), "Binaries\\IOS", Params.ShortProjectName + ".ipa");
+				if (TargetConfiguration != UnrealTargetConfiguration.Development)
+				{
+					ProjectIPA = Path.Combine(Path.GetDirectoryName(Params.RawProjectPath), "Binaries\\IOS", Params.ShortProjectName + "-" + PlatformType.ToString() + "-" + TargetConfiguration.ToString() + ".ipa");
+				}
 			}
 
 			var StagedIPA = SC.StageDirectory + "\\" + Path.GetFileName(ProjectIPA);
