build script, retry on a slightly different XGE error

diff --git a/Engine/Source/Programs/AutomationTool/UE4Build.cs b/Engine/Source/Programs/AutomationTool/UE4Build.cs
--- a/Engine/Source/Programs/AutomationTool/UE4Build.cs
+++ b/Engine/Source/Programs/AutomationTool/UE4Build.cs
@@ -1009,7 +1009,7 @@ namespace AutomationTool
                             string LogFile = GetRunAndLogLogName(CmdEnv, XGEConsole);
                             string Output = RunAndLog(XGEConsole, Args, out SuccesCode, LogFile);
                             PopDir();
-                            if (ConnectionRetries > 0 && SuccesCode == 4 && Output.Contains("Timed out waiting for reply from local connection"))
+                            if (ConnectionRetries > 0 && (SuccesCode == 4 || SuccesCode == 2) && Output.Contains("Timed out waiting for reply from local connection"))
                             {
                                 Log(System.Diagnostics.TraceEventType.Warning, "XGE failure on the local connection timeout");
                                 if (ConnectionRetries < 2)
