build script, remove old rocket build

diff --git a/Engine/Source/Programs/AutomationTool/MCPPublic.cs b/Engine/Source/Programs/AutomationTool/MCPPublic.cs
--- a/Engine/Source/Programs/AutomationTool/MCPPublic.cs
+++ b/Engine/Source/Programs/AutomationTool/MCPPublic.cs
@@ -156,7 +156,17 @@ namespace EpicGames.MCP.Automation
             }
             return (TargetPlatform == UnrealTargetPlatform.Win64 || TargetPlatform == UnrealTargetPlatform.Win32) ? MCPPlatform.Windows : MCPPlatform.Mac;
         }
-
+        /// <summary>
+        /// Determine the platform name (Win32/64 becomes Windows, Mac is Mac, the rest we don't currently understand)
+        /// </summary>
+        static public UnrealTargetPlatform FromMCPPlatform(MCPPlatform TargetPlatform)
+        {
+            if (TargetPlatform != MCPPlatform.Windows && TargetPlatform != MCPPlatform.Mac)
+            {
+                throw new AutomationException("Platform {0} is not properly supported by the MCP backend yet", TargetPlatform);
+            }
+            return (TargetPlatform == MCPPlatform.Windows) ? UnrealTargetPlatform.Win64 : UnrealTargetPlatform.Mac;
+        }
         /// <summary>
         /// Returns the build root path (P:\Builds on build machines usually)
         /// </summary>
