build script, allow OnlyNode to skip dependencies

diff --git a/Engine/Source/Programs/AutomationTool/Scripts/GUBP.Automation.cs b/Engine/Source/Programs/AutomationTool/Scripts/GUBP.Automation.cs
--- a/Engine/Source/Programs/AutomationTool/Scripts/GUBP.Automation.cs
+++ b/Engine/Source/Programs/AutomationTool/Scripts/GUBP.Automation.cs
@@ -5032,26 +5032,38 @@ if (HostPlatform == UnrealTargetPlatform.Mac) continue; //temp hack till mac aut
             GUBPNodes[NodeToDo].AllDependencyBuildProducts = new List<string>();
             if (GUBPNodes[NodeToDo].FullNamesOfDependencies == null)
             {
-                throw new AutomationException("Node {0} was not processed yet?", NodeToDo);
-            }
-            foreach (var Dep in GUBPNodes[NodeToDo].FullNamesOfDependencies)
-            {
-                if (GUBPNodes[Dep].BuildProducts == null)
+                if (ParseParamValue("OnlyNode", "") == "")
                 {
-                    throw new AutomationException("Node {0} was not processed yet2? Processing {1}", Dep, NodeToDo);
+                    throw new AutomationException("Node {0} was not processed yet?", NodeToDo);
                 }
-                foreach (var Prod in GUBPNodes[Dep].BuildProducts)
-                {
-                    GUBPNodes[NodeToDo].AddDependentBuildProduct(Prod);
-                }
-                if (GUBPNodes[Dep].AllDependencyBuildProducts == null)
+            }
+            else
+            {
+                foreach (var Dep in GUBPNodes[NodeToDo].FullNamesOfDependencies)
                 {
-                    throw new AutomationException("Node {0} was not processed yet3?  Processing {1}", Dep, NodeToDo);
+                    if (GUBPNodes[Dep].BuildProducts == null)
+                    {
+                        if (ParseParamValue("OnlyNode", "") == "")
+                        {
+                            throw new AutomationException("Node {0} was not processed yet2? Processing {1}", Dep, NodeToDo);
+                        }
+                    }
+                    else
+                    {
+                        foreach (var Prod in GUBPNodes[Dep].BuildProducts)
+                        {
+                            GUBPNodes[NodeToDo].AddDependentBuildProduct(Prod);
+                        }
+                        if (GUBPNodes[Dep].AllDependencyBuildProducts == null)
+                        {
+                            throw new AutomationException("Node {0} was not processed yet3?  Processing {1}", Dep, NodeToDo);
+                        }
+                        foreach (var Prod in GUBPNodes[Dep].AllDependencyBuildProducts)
+                        {
+                            GUBPNodes[NodeToDo].AddDependentBuildProduct(Prod);
+                        }
+                    }
                 }
-                foreach (var Prod in GUBPNodes[Dep].AllDependencyBuildProducts)
-                {
-                    GUBPNodes[NodeToDo].AddDependentBuildProduct(Prod);
-                } 
             }
             string NodeStoreName = StoreName + "-" + GUBPNodes[NodeToDo].GetFullName();
             
