build script, force incremental for GUBP compiles

diff --git a/Engine/Source/Programs/AutomationTool/Scripts/GUBP.Automation.cs b/Engine/Source/Programs/AutomationTool/Scripts/GUBP.Automation.cs
--- a/Engine/Source/Programs/AutomationTool/Scripts/GUBP.Automation.cs
+++ b/Engine/Source/Programs/AutomationTool/Scripts/GUBP.Automation.cs
@@ -22,9 +22,29 @@ public class GUBP : BuildCommand
     public bool bFake = false;
     public static bool bNoIOSOnPC = false;
     public static bool bBuildRocket = false;
+    public static bool bForceIncrementalCompile = false;
     public static string ECProject = null;
     public static string ForceAllEmailsTo = null;
 
+    Dictionary<string, GUBPNode> GUBPNodes;
+    Dictionary<string, bool> GUBPNodesCompleted;
+    Dictionary<string, string> GUBPNodesControllingTrigger;
+    Dictionary<string, string> GUBPNodesControllingTriggerDotName;
+
+    class NodeHistory
+    {
+        public int LastSucceeded = 0;
+        public List<int> InProgress = new List<int>();
+        public string InProgressString = "";
+        public List<int> Failed = new List<int>();
+        public string FailedString = "";
+        public List<int> AllStarted = new List<int>();
+        public List<int> AllSucceeded = new List<int>();
+        public List<int> AllFailed = new List<int>();
+    };
+
+    Dictionary<string, NodeHistory> GUBPNodesHistory;
+
     public string RocketUBTArgs(bool bUseRocketInsteadOfBuildRocket = false)
     {
         if (bBuildRocket)
@@ -294,7 +314,7 @@ public class GUBP : BuildCommand
 
         public override void DoBuild(GUBP bp)
         {
-            if (IsBuildMachine)
+            if (CommandUtils.P4Enabled && CommandUtils.AllowSubmit)
             {
                 // only update version files on build machines
                 var UE4Build = new UE4Build(bp);
@@ -391,8 +411,9 @@ public class GUBP : BuildCommand
             UE4Build.BuildAgenda Agenda = GetAgenda(bp);
             if (Agenda != null)
             {
+                bool ReallyDeleteBuildProducts = DeleteBuildProducts() && !GUBP.bForceIncrementalCompile;
                 Agenda.DoRetries = false; // these would delete build products
-                UE4Build.Build(Agenda, InDeleteBuildProducts: DeleteBuildProducts(), InUpdateVersionFiles: false, InForceUnity: true);
+                UE4Build.Build(Agenda, InDeleteBuildProducts: ReallyDeleteBuildProducts, InUpdateVersionFiles: false, InForceUnity: true);
                 PostBuild(bp, UE4Build);
 
                 UE4Build.CheckBuildProducts(UE4Build.BuildProductFiles);
@@ -408,7 +429,7 @@ public class GUBP : BuildCommand
                 }
                 PostBuildProducts(bp);
             }
-            else
+            if (Agenda == null || (BuildProducts.Count == 0 && GUBP.bForceIncrementalCompile))
             {
                 SaveRecordOfSuccessAndAddToBuildProducts("Nothing to actually compile");
             }
@@ -510,6 +531,10 @@ public class GUBP : BuildCommand
         }
         void DeleteStaleDLLs(GUBP bp)
         {
+            if (GUBP.bForceIncrementalCompile)
+            {
+                return;
+            }
             var Targets = new List<string>{bp.Branch.BaseEngineProject.Properties.Targets[TargetRules.TargetType.Editor].TargetName};
             foreach (var ProgramTarget in bp.Branch.BaseEngineProject.Properties.Programs)
             {
@@ -2521,27 +2546,6 @@ public class GUBP : BuildCommand
         }
     };
 
-    Dictionary<string, GUBPNode> GUBPNodes;
-    Dictionary<string, bool> GUBPNodesCompleted;
-    Dictionary<string, string> GUBPNodesControllingTrigger;
-    Dictionary<string, string> GUBPNodesControllingTriggerDotName;
-
-    class NodeHistory
-    {
-        public int LastSucceeded = 0;
-        public List<int> InProgress = new List<int>();
-        public string InProgressString = "";
-        public List<int> Failed = new List<int>();
-        public string FailedString = "";
-        public List<int> AllStarted = new List<int>(); 
-        public List<int> AllSucceeded = new List<int>(); 
-        public List<int> AllFailed = new List<int>();
-    };
-
-    Dictionary<string, NodeHistory> GUBPNodesHistory;
-
-
-
     public string AddNode(GUBPNode Node)
     {
         string Name = Node.GetFullName();
@@ -3416,6 +3420,7 @@ public class GUBP : BuildCommand
     [Help("ShowECOnly", "Only show EC nodes.")]
     [Help("ECProject", "From EC, the name of the project, used to get a version number.")]
     [Help("CIS", "This is a CIS run, assign TimeIndex based on the history.")]
+    [Help("ForceIncrementalCompile", "make sure all compiles are incremental")]
 
     public override void ExecuteBuild()
     {
@@ -3436,7 +3441,9 @@ public class GUBP : BuildCommand
         {
             HostPlatforms.Add(UnrealTargetPlatform.Mac);
         }
+
         bBuildRocket = ParseParam("BuildRocket");
+        bForceIncrementalCompile = ParseParam("ForceIncrementalCompile");
 
         StoreName = ParseParamValue("Store");
         string StoreSuffix = ParseParamValue("StoreSuffix", "");
