#UE4 Cleanup of commandlet params building.

diff --git a/Engine/Source/Programs/AutomationTool/Scripts/CookCommand.Automation.cs b/Engine/Source/Programs/AutomationTool/Scripts/CookCommand.Automation.cs
--- a/Engine/Source/Programs/AutomationTool/Scripts/CookCommand.Automation.cs
+++ b/Engine/Source/Programs/AutomationTool/Scripts/CookCommand.Automation.cs
@@ -70,7 +70,7 @@ public partial class Project : CommandUtils
 
 		try
 		{
-			var CommandletParams = "-unattended -buildmachine -forcelogflush -AllowStdOutLogVerbosity -Unversioned -fileopenlog";
+			var CommandletParams = "-buildmachine -Unversioned -fileopenlog";
 			if (Params.Compressed)
 			{
 				CommandletParams += " -Compressed";
