Cross post all builds to both GameDev and Production buildinfo services.
Intended to be temporary until we can promote to copy builds from GameDev up to Production.

diff --git a/Engine/Source/Programs/AutomationTool/MCPPublic.cs b/Engine/Source/Programs/AutomationTool/MCPPublic.cs
--- a/Engine/Source/Programs/AutomationTool/MCPPublic.cs
+++ b/Engine/Source/Programs/AutomationTool/MCPPublic.cs
@@ -11,6 +11,8 @@ using UnrealBuildTool;
 
 namespace EpicGames.MCP.Automation
 {
+	using EpicGames.MCP.Config;
+
     /// <summary>
     /// Utility class to provide commit/rollback functionality via an RAII-like functionality.
     /// Usage is to provide a rollback action that will be called on Dispose if the Commit() method is not called.
@@ -456,6 +458,7 @@ namespace EpicGames.MCP.Automation
             return Handler;
         }
         public abstract string GetClientToken(BuildPatchToolStagingInfo StagingInfo);
+		public abstract string GetClientToken(McpConfigData McpConfig);
         public abstract string SendWebRequest(WebRequest Upload, string Method, string ContentType, byte[] Data);
     }
 
