build more workarounds for XGE badness

diff --git a/Engine/Source/Programs/AutomationTool/UE4Build.cs b/Engine/Source/Programs/AutomationTool/UE4Build.cs
--- a/Engine/Source/Programs/AutomationTool/UE4Build.cs
+++ b/Engine/Source/Programs/AutomationTool/UE4Build.cs
@@ -1009,7 +1009,7 @@ namespace AutomationTool
                             string LogFile = GetRunAndLogLogName(CmdEnv, XGEConsole);
                             string Output = RunAndLog(XGEConsole, Args, out SuccesCode, LogFile);
                             PopDir();
-                            if (ConnectionRetries > 0 && (SuccesCode == 4 || SuccesCode == 2) && Output.Contains("Timed out waiting for reply from local connection"))
+                            if (ConnectionRetries > 0 && (SuccesCode == 4 || SuccesCode == 2) && !Output.Contains("------Project:"))
                             {
                                 Log(System.Diagnostics.TraceEventType.Warning, "XGE failure on the local connection timeout");
                                 if (ConnectionRetries < 2)
@@ -1200,7 +1200,7 @@ namespace AutomationTool
 					{
 						// When building a target for Mac or iOS, use UBT's -flushmac option to clean up the remote builder
 						bool bForceFlushMac = DeleteBuildProducts && (Target.Platform == UnrealBuildTool.UnrealTargetPlatform.Mac || Target.Platform == UnrealBuildTool.UnrealTargetPlatform.IOS);
-						BuildWithUBT(Target.ProjectName, Target.TargetName, Target.Platform, Target.Config.ToString(), Target.UprojectPath, bForceMonolithic, bForceNonUnity, bForceDebugInfo, bForceFlushMac, false, Target.UBTArgs, bForceUnity);
+						BuildWithUBT(Target.ProjectName, Target.TargetName, Target.Platform, Target.Config.ToString(), Target.UprojectPath, bForceMonolithic, bForceNonUnity, bForceDebugInfo, bForceFlushMac, true, Target.UBTArgs, bForceUnity);
 					}
 				}
 
