#ue4
* temporarily disable PC code signing while I fix the mac build to also build a signed stub

diff --git a/Engine/Source/Programs/AutomationTool/IOS/IOSPlatform.Automation.cs b/Engine/Source/Programs/AutomationTool/IOS/IOSPlatform.Automation.cs
--- a/Engine/Source/Programs/AutomationTool/IOS/IOSPlatform.Automation.cs
+++ b/Engine/Source/Programs/AutomationTool/IOS/IOSPlatform.Automation.cs
@@ -37,7 +37,7 @@ public class IOSPlatform : Platform
 	protected bool GetCodeSignDesirability(ProjectParams Params)
 	{
 		//@TODO: Would like to make this true, as it's the common case for everyone else
-		bool bDefaultNeedsSign = true;
+		bool bDefaultNeedsSign = false;
 
 		bool bNeedsSign = false;
 		string EnvVar = InternalUtils.GetEnvironmentVariable(IOSEnvVarNames.CodeSignWhenStaging, bDefaultNeedsSign ? "1" : "0", /*bQuiet=*/ false);
