build script, tweaks to dotnet build products

diff --git a/Engine/Source/Programs/AutomationTool/UE4Build.cs b/Engine/Source/Programs/AutomationTool/UE4Build.cs
--- a/Engine/Source/Programs/AutomationTool/UE4Build.cs
+++ b/Engine/Source/Programs/AutomationTool/UE4Build.cs
@@ -1362,6 +1362,7 @@ namespace AutomationTool
 						"AutomationTool.exe",
 						"AutomationTool.exe.config",
 						"UnrealBuildTool.exe",
+						"UnrealBuildTool.exe.config",
 						"DotNETUtilities.dll",
 					});
 
