build script, disable automated tests

diff --git a/Engine/Source/Programs/AutomationTool/Scripts/GUBP.Automation.cs b/Engine/Source/Programs/AutomationTool/Scripts/GUBP.Automation.cs
--- a/Engine/Source/Programs/AutomationTool/Scripts/GUBP.Automation.cs
+++ b/Engine/Source/Programs/AutomationTool/Scripts/GUBP.Automation.cs
@@ -3421,6 +3421,7 @@ public class GUBP : BuildCommand
     [Help("ECProject", "From EC, the name of the project, used to get a version number.")]
     [Help("CIS", "This is a CIS run, assign TimeIndex based on the history.")]
     [Help("ForceIncrementalCompile", "make sure all compiles are incremental")]
+    [Help("AutomatedTesting", "Allow automated testing, currently disabled.")]
 
     public override void ExecuteBuild()
     {
@@ -3444,6 +3445,7 @@ public class GUBP : BuildCommand
 
         bBuildRocket = ParseParam("BuildRocket");
         bForceIncrementalCompile = ParseParam("ForceIncrementalCompile");
+        bool bAutomatedTesting = ParseParam("AutomatedTesting");
 
         StoreName = ParseParamValue("Store");
         string StoreSuffix = ParseParamValue("StoreSuffix", "");
@@ -3817,6 +3819,7 @@ public class GUBP : BuildCommand
                         throw new AutomationException("We assume that if we have shared promotable, the base engine is in it.");
                     }
                 }
+                if (bAutomatedTesting)
                 {
                     var EditorTests = Branch.BaseEngineProject.Properties.Targets[TargetRules.TargetType.Editor].Rules.GUBP_GetEditorTests_EditorTypeOnly(HostPlatform);
                     var EditorTestNodes = new List<string>();
@@ -3919,6 +3922,7 @@ public class GUBP : BuildCommand
                                 }
                                 var GameTests = Target.Rules.GUBP_GetGameTests_MonolithicOnly(HostPlatform, GetAltHostPlatform(HostPlatform), Plat);
                                 var RequiredPlatforms = new List<UnrealTargetPlatform> { Plat };
+                                if (bAutomatedTesting)
                                 {
                                     var ThisMonoGameTestNodes = new List<string>();
                                     foreach (var Test in GameTests)
@@ -3946,17 +3950,19 @@ public class GUBP : BuildCommand
                                     {
                                         AddNode(new GamePlatformCookedAndCompiledNode(this, HostPlatform, NonCodeProject, Plat, false));
                                     }
-
-                                    if (HostPlatform == UnrealTargetPlatform.Mac) continue; //temp hack till mac automated testing works
-                                    var ThisMonoGameTestNodes = new List<string>();
-                                    foreach (var Test in GameTests)
-                                    {
-                                        var TestName = Test.Key + "_" + Plat.ToString();
-                                        ThisMonoGameTestNodes.Add(AddNode(new UATTestNode(this, HostPlatform, NonCodeProject, TestName, Test.Value, CookedAgentSharingGroup, false, RequiredPlatforms)));
-                                    }
-                                    if (ThisMonoGameTestNodes.Count > 0)
+                                    if (bAutomatedTesting)
                                     {
-                                        GameTestNodes.Add(AddNode(new GameAggregateNode(this, HostPlatform, NonCodeProject, "CookedTests_" + Plat.ToString() + "_" + Kind.ToString() + HostPlatformNode.StaticGetHostPlatformSuffix(HostPlatform), ThisMonoGameTestNodes, 0.0f)));
+                                        if (HostPlatform == UnrealTargetPlatform.Mac) continue; //temp hack till mac automated testing works
+                                        var ThisMonoGameTestNodes = new List<string>();
+                                        foreach (var Test in GameTests)
+                                        {
+                                            var TestName = Test.Key + "_" + Plat.ToString();
+                                            ThisMonoGameTestNodes.Add(AddNode(new UATTestNode(this, HostPlatform, NonCodeProject, TestName, Test.Value, CookedAgentSharingGroup, false, RequiredPlatforms)));
+                                        }
+                                        if (ThisMonoGameTestNodes.Count > 0)
+                                        {
+                                            GameTestNodes.Add(AddNode(new GameAggregateNode(this, HostPlatform, NonCodeProject, "CookedTests_" + Plat.ToString() + "_" + Kind.ToString() + HostPlatformNode.StaticGetHostPlatformSuffix(HostPlatform), ThisMonoGameTestNodes, 0.0f)));
+                                        }
                                     }
                                 }
                             }
@@ -4015,7 +4021,7 @@ public class GUBP : BuildCommand
                 }
 
                 AddNode(new EditorGameNode(this, HostPlatform, CodeProj));
-if (HostPlatform != UnrealTargetPlatform.Mac) //temp hack till mac automated testing works
+                if (bAutomatedTesting && HostPlatform != UnrealTargetPlatform.Mac) //temp hack till mac automated testing works
                 {
                     var EditorTests = CodeProj.Properties.Targets[TargetRules.TargetType.Editor].Rules.GUBP_GetEditorTests_EditorTypeOnly(HostPlatform);
                     var EditorTestNodes = new List<string>();
@@ -4078,46 +4084,53 @@ if (HostPlatform != UnrealTargetPlatform.Mac) //temp hack till mac automated tes
                                 {
                                     AddNode(new GamePlatformCookedAndCompiledNode(this, HostPlatform, CodeProj, Plat, true));
                                 }
-if (HostPlatform == UnrealTargetPlatform.Mac) continue; //temp hack till mac automated testing works
-                                var GameTests = Target.Rules.GUBP_GetGameTests_MonolithicOnly(HostPlatform, GetAltHostPlatform(HostPlatform), Plat);
-                                var RequiredPlatforms = new List<UnrealTargetPlatform> { Plat };
-                                var ThisMonoGameTestNodes = new List<string>();
-
-                                foreach (var Test in GameTests)
+                                if (bAutomatedTesting)
                                 {
-                                    var TestNodeName = Test.Key + "_" + Plat.ToString();
-                                    ThisMonoGameTestNodes.Add(AddNode(new UATTestNode(this, HostPlatform, CodeProj, TestNodeName, Test.Value, CookedAgentSharingGroup, false, RequiredPlatforms)));
-                                }
-                                if (ThisMonoGameTestNodes.Count > 0)
-                                {
-                                    GameTestNodes.Add(AddNode(new GameAggregateNode(this, HostPlatform, CodeProj, "CookedTests_" + Plat.ToString() + "_" + Kind.ToString() + HostPlatformNode.StaticGetHostPlatformSuffix(HostPlatform), ThisMonoGameTestNodes, 0.0f)));
+                                    if (HostPlatform == UnrealTargetPlatform.Mac) continue; //temp hack till mac automated testing works
+                                    var GameTests = Target.Rules.GUBP_GetGameTests_MonolithicOnly(HostPlatform, GetAltHostPlatform(HostPlatform), Plat);
+                                    var RequiredPlatforms = new List<UnrealTargetPlatform> { Plat };
+                                    var ThisMonoGameTestNodes = new List<string>();
+
+                                    foreach (var Test in GameTests)
+                                    {
+                                        var TestNodeName = Test.Key + "_" + Plat.ToString();
+                                        ThisMonoGameTestNodes.Add(AddNode(new UATTestNode(this, HostPlatform, CodeProj, TestNodeName, Test.Value, CookedAgentSharingGroup, false, RequiredPlatforms)));
+                                    }
+                                    if (ThisMonoGameTestNodes.Count > 0)
+                                    {
+                                        GameTestNodes.Add(AddNode(new GameAggregateNode(this, HostPlatform, CodeProj, "CookedTests_" + Plat.ToString() + "_" + Kind.ToString() + HostPlatformNode.StaticGetHostPlatformSuffix(HostPlatform), ThisMonoGameTestNodes, 0.0f)));
+                                    }
+
                                 }
                             }
                         }
                     }
                 }
-                foreach (var ServerPlatform in ServerPlatforms)
+                if (bAutomatedTesting)
                 {
-                    foreach (var GamePlatform in GamePlatforms)
+                    foreach (var ServerPlatform in ServerPlatforms)
                     {
-if (HostPlatform == UnrealTargetPlatform.Mac) continue; //temp hack till mac automated testing works
-                        var Target = CodeProj.Properties.Targets[TargetRules.TargetType.Game];
-                        var ClientServerTests = Target.Rules.GUBP_GetClientServerTests_MonolithicOnly(HostPlatform, GetAltHostPlatform(HostPlatform), ServerPlatform, GamePlatform);
-                        var RequiredPlatforms = new List<UnrealTargetPlatform> { ServerPlatform };
-                        if (ServerPlatform != GamePlatform)
+                        foreach (var GamePlatform in GamePlatforms)
                         {
-                            RequiredPlatforms.Add(GamePlatform);
-                        }
-                        foreach (var Test in ClientServerTests)
-                        {
-                            var TestNodeName = Test.Key + "_" + GamePlatform.ToString() + "_" + ServerPlatform.ToString();
-                            GameTestNodes.Add(AddNode(new UATTestNode(this, HostPlatform, CodeProj, TestNodeName, Test.Value, CookedAgentSharingGroup, false, RequiredPlatforms)));
+                            if (HostPlatform == UnrealTargetPlatform.Mac) continue; //temp hack till mac automated testing works
+                            var Target = CodeProj.Properties.Targets[TargetRules.TargetType.Game];
+                            var ClientServerTests = Target.Rules.GUBP_GetClientServerTests_MonolithicOnly(HostPlatform, GetAltHostPlatform(HostPlatform), ServerPlatform, GamePlatform);
+                            var RequiredPlatforms = new List<UnrealTargetPlatform> { ServerPlatform };
+                            if (ServerPlatform != GamePlatform)
+                            {
+                                RequiredPlatforms.Add(GamePlatform);
+                            }
+                            foreach (var Test in ClientServerTests)
+                            {
+                                var TestNodeName = Test.Key + "_" + GamePlatform.ToString() + "_" + ServerPlatform.ToString();
+                                GameTestNodes.Add(AddNode(new UATTestNode(this, HostPlatform, CodeProj, TestNodeName, Test.Value, CookedAgentSharingGroup, false, RequiredPlatforms)));
+                            }
                         }
                     }
-                }
-                if (GameTestNodes.Count > 0)
-                {
-                    AddNode(new GameAggregateNode(this, HostPlatform, CodeProj, "AllCookedTests", GameTestNodes));
+                    if (GameTestNodes.Count > 0)
+                    {
+                        AddNode(new GameAggregateNode(this, HostPlatform, CodeProj, "AllCookedTests", GameTestNodes));
+                    }
                 }
             }
             AddCustomNodes(HostPlatform);
