build script, some tweaks to GUBP and rocket for local runs

diff --git a/Engine/Source/Programs/AutomationTool/Scripts/GUBP.Automation.cs b/Engine/Source/Programs/AutomationTool/Scripts/GUBP.Automation.cs
--- a/Engine/Source/Programs/AutomationTool/Scripts/GUBP.Automation.cs
+++ b/Engine/Source/Programs/AutomationTool/Scripts/GUBP.Automation.cs
@@ -1026,10 +1026,26 @@ public class GUBP : BuildCommand
         {
             BuildProducts = new List<string>();
 
-            foreach (var FileToCopy in CommandUtils.FindFiles("*.h", true, CommandUtils.CombinePaths(CmdEnv.LocalRoot, @"Engine\Intermediate\Build\", HostPlatform.ToString())))
+            foreach (var FileToCopy in CommandUtils.FindFiles("*.h", true, CommandUtils.CombinePaths(CmdEnv.LocalRoot, @"Engine\Intermediate\Build\", HostPlatform.ToString(), "Inc")))
             {
                 AddBuildProduct(FileToCopy);
             }
+            var Targets = new List<string> { bp.Branch.BaseEngineProject.Properties.Targets[TargetRules.TargetType.Editor].TargetName };
+            foreach (var ProgramTarget in bp.Branch.BaseEngineProject.Properties.Programs)
+            {
+                if (ProgramTarget.Rules.GUBP_AlwaysBuildWithBaseEditor() && ProgramTarget.Rules.SupportsPlatform(HostPlatform))
+                {
+                    Targets.Add(ProgramTarget.TargetName);
+                }
+            }
+            foreach (var Target in Targets)
+            {
+                foreach (var FileToCopy in CommandUtils.FindFiles("*.h", true, CommandUtils.CombinePaths(CmdEnv.LocalRoot, @"Engine\Intermediate\Build\", HostPlatform.ToString(), Target)))
+                {
+                    AddBuildProduct(FileToCopy);
+                }
+            }
+
             var EnginePluginsDirectory = Path.Combine(CommandUtils.CmdEnv.LocalRoot, "Engine/Plugins");
             var EnginePlugins = new List<PluginInfo>();
             Plugins.FindPluginsIn(EnginePluginsDirectory, PluginInfo.LoadedFromType.Engine, ref EnginePlugins);
@@ -1070,7 +1086,7 @@ public class GUBP : BuildCommand
         }
         public override void DoBuild(GUBP bp)
         {
-            if (GUBP.bHackRunIOSCompilesOnMac)
+            if (GUBP.bHackRunIOSCompilesOnMac && TargetPlatform == UnrealTargetPlatform.IOS)
             {
                 throw new AutomationException("these rocket header node require real mac path.");
             }
@@ -3738,8 +3754,8 @@ public class GUBP : BuildCommand
             {
                 AddNode(new RootEditorHeadersNode(HostPlatform));
             }
-            AddNode(new ToolsNode(HostPlatform));            
-			AddNode(new InternalToolsNode(HostPlatform));
+            AddNode(new ToolsNode(HostPlatform));
+            AddNode(new InternalToolsNode(HostPlatform));
             AddNode(new EditorAndToolsNode(HostPlatform));
             AddNode(new NonUnityTestNode(HostPlatform));
 
@@ -4133,7 +4149,6 @@ public class GUBP : BuildCommand
                     }
                 }
             }
-            AddCustomNodes(HostPlatform);
         }
 
         int NumSharedAllHosts = 0;
@@ -4182,6 +4197,10 @@ public class GUBP : BuildCommand
             AddNode(new WaitForSharedPromotionUserInput(this, true));
             AddNode(new SharedLabelPromotableNode(this, true));
         }
+        foreach (var HostPlatform in HostPlatforms)
+        {
+            AddCustomNodes(HostPlatform);
+        }
 
         foreach (var NodeToDo in GUBPNodes)
         {
