build script, split BPT and UEL off as separate nodes, adjust dependencies accordingly, enable timeouts

diff --git a/Engine/Source/Programs/AutomationTool/ProjectUtils.cs b/Engine/Source/Programs/AutomationTool/ProjectUtils.cs
--- a/Engine/Source/Programs/AutomationTool/ProjectUtils.cs
+++ b/Engine/Source/Programs/AutomationTool/ProjectUtils.cs
@@ -543,6 +543,20 @@ namespace AutomationTool
             }
             return null;
         }
+        public SingleTargetProperties FindProgram(string ProgramName)
+        {
+            foreach (var Proj in BaseEngineProject.Properties.Programs)
+            {
+                if (Proj.TargetName.Equals(ProgramName, StringComparison.InvariantCultureIgnoreCase))
+                {
+                    return Proj;
+                }
+            }
+            SingleTargetProperties Result;
+            Result.TargetName = ProgramName;
+            Result.Rules = null;
+            return Result;
+        }
     };
 
 }
diff --git a/Engine/Source/Programs/AutomationTool/Scripts/GUBP.Automation.cs b/Engine/Source/Programs/AutomationTool/Scripts/GUBP.Automation.cs
--- a/Engine/Source/Programs/AutomationTool/Scripts/GUBP.Automation.cs
+++ b/Engine/Source/Programs/AutomationTool/Scripts/GUBP.Automation.cs
@@ -13,6 +13,7 @@ public class GUBP : BuildCommand
 {
     public string StoreName = null;
     public int CL = 0;
+    public int TimeIndex = 0;
     public bool bSignBuildProducts = false;
     public List<UnrealTargetPlatform> ActivePlatforms = null;
     public BranchInfo Branch = null;
@@ -374,6 +375,19 @@ public class GUBP : BuildCommand
         {
             return GUBP.GetAltHostPlatform(HostPlatform);
         }
+        public override int TimeoutInMinutes()
+        {
+            return base.TimeoutInMinutes() + ((HostPlatform == UnrealTargetPlatform.Mac) ? 30 : 0); // Mac is slower and more resource constrained
+        }
+        public override int CISFrequencyQuantumShift(GUBP bp)
+        {
+            int Result = base.CISFrequencyQuantumShift(bp);
+            if (HostPlatform == UnrealTargetPlatform.Mac)
+            {
+                Result += 1; // mac only runs half as much
+            }
+            return Result;
+        }
     }
 
     public class CompileNode : HostPlatformNode
@@ -434,6 +448,10 @@ public class GUBP : BuildCommand
                 SaveRecordOfSuccessAndAddToBuildProducts("Nothing to actually compile");
             }
         }
+        public override int TimeoutInMinutes()
+        {
+            return base.TimeoutInMinutes() + ((HostPlatform == UnrealTargetPlatform.Mac) ? 30 : 0);
+        }
     }
 
     public class ToolsForCompileNode : CompileNode
@@ -457,7 +475,7 @@ public class GUBP : BuildCommand
         }
         public override int AgentMemoryRequirement(GUBP bp)
         {
-            if (bp.ParseParam("Launcher"))
+            if (bp.ParseParam("Launcher") || bp.TimeIndex != 0)
             {
                 return base.AgentMemoryRequirement(bp);
             }
@@ -663,7 +681,9 @@ public class GUBP : BuildCommand
             foreach (var ProgramTarget in bp.Branch.BaseEngineProject.Properties.Programs)
             {
                 bool bInternalOnly;
-                if (ProgramTarget.Rules.GUBP_AlwaysBuildWithTools(HostPlatform, out bInternalOnly) && ProgramTarget.Rules.SupportsPlatform(HostPlatform) && !bInternalOnly)
+                bool SeparateNode;
+
+                if (ProgramTarget.Rules.GUBP_AlwaysBuildWithTools(HostPlatform, out bInternalOnly, out SeparateNode) && ProgramTarget.Rules.SupportsPlatform(HostPlatform) && !bInternalOnly && !SeparateNode)
                 {
                     foreach (var Plat in ProgramTarget.Rules.GUBP_ToolPlatforms(HostPlatform))
                     {
@@ -678,6 +698,52 @@ public class GUBP : BuildCommand
             return Agenda;
         }
     }
+    public class SingleToolsNode : CompileNode
+    {
+        SingleTargetProperties ProgramTarget;
+        public SingleToolsNode(UnrealTargetPlatform InHostPlatform, SingleTargetProperties InProgramTarget)
+            : base(InHostPlatform)
+        {
+            ProgramTarget = InProgramTarget;
+            if (!GUBP.bBuildRocket) // more errors and more performance by just starting before the root editor is done
+            {
+                AddPseudodependency(RootEditorNode.StaticGetFullName(HostPlatform));
+            }
+            AgentSharingGroup = "Tools" + StaticGetHostPlatformSuffix(HostPlatform);
+        }
+        public static string StaticGetFullName(UnrealTargetPlatform InHostPlatform, SingleTargetProperties InProgramTarget)
+        {
+            return "Tools_" + InProgramTarget.TargetName + StaticGetHostPlatformSuffix(InHostPlatform);
+        }
+        public override string GetFullName()
+        {
+            return StaticGetFullName(HostPlatform, ProgramTarget);
+        }
+        public override float Priority()
+        {
+            return base.Priority() + 2;
+        }
+        public override bool DeleteBuildProducts()
+        {
+            return true;
+        }
+        public override UE4Build.BuildAgenda GetAgenda(GUBP bp)
+        {
+            var Agenda = new UE4Build.BuildAgenda();
+
+            string AddArgs = "-nobuilduht -skipactionhistory" + bp.RocketUBTArgs(); ;
+
+            foreach (var Plat in ProgramTarget.Rules.GUBP_ToolPlatforms(HostPlatform))
+            {
+                foreach (var Config in ProgramTarget.Rules.GUBP_ToolConfigs(HostPlatform))
+                {
+                    Agenda.AddTargets(new string[] { ProgramTarget.TargetName }, Plat, Config, InAddArgs: AddArgs);
+                }
+            }
+            return Agenda;
+        }
+    }
+
     public class InternalToolsNode : CompileNode
     {
         public InternalToolsNode(UnrealTargetPlatform InHostPlatform)
@@ -742,7 +808,9 @@ public class GUBP : BuildCommand
             foreach (var ProgramTarget in bp.Branch.BaseEngineProject.Properties.Programs)
             {
                 bool bInternalOnly;
-                if (ProgramTarget.Rules.GUBP_AlwaysBuildWithTools(HostPlatform, out bInternalOnly) && ProgramTarget.Rules.SupportsPlatform(HostPlatform) && bInternalOnly)
+                bool SeparateNode;
+
+                if (ProgramTarget.Rules.GUBP_AlwaysBuildWithTools(HostPlatform, out bInternalOnly, out SeparateNode) && ProgramTarget.Rules.SupportsPlatform(HostPlatform) && bInternalOnly && !SeparateNode)
                 {
                     foreach (var Plat in ProgramTarget.Rules.GUBP_ToolPlatforms(HostPlatform))
                     {
@@ -762,6 +830,53 @@ public class GUBP : BuildCommand
         }
     }
 
+    public class SingleInternalToolsNode : CompileNode
+    {
+        SingleTargetProperties ProgramTarget;
+        public SingleInternalToolsNode(UnrealTargetPlatform InHostPlatform, SingleTargetProperties InProgramTarget)
+            : base(InHostPlatform)
+        {
+            ProgramTarget = InProgramTarget;
+            if (!GUBP.bBuildRocket) // more errors and more performance by just starting before the root editor is done
+            {
+                AddPseudodependency(RootEditorNode.StaticGetFullName(HostPlatform));
+            }
+            AgentSharingGroup = "Tools" + StaticGetHostPlatformSuffix(HostPlatform);
+        }
+        public static string StaticGetFullName(UnrealTargetPlatform InHostPlatform, SingleTargetProperties InProgramTarget)
+        {
+            return "InternalTools_" + InProgramTarget.TargetName + StaticGetHostPlatformSuffix(InHostPlatform);
+        }
+        public override string GetFullName()
+        {
+            return StaticGetFullName(HostPlatform, ProgramTarget);
+        }
+        public override float Priority()
+        {
+            return base.Priority() + 3;
+        }
+        public override bool DeleteBuildProducts()
+        {
+            return true;
+        }
+        public override UE4Build.BuildAgenda GetAgenda(GUBP bp)
+        {
+            var Agenda = new UE4Build.BuildAgenda();
+
+            string AddArgs = "-nobuilduht -skipactionhistory" + bp.RocketUBTArgs(); ;
+
+            foreach (var Plat in ProgramTarget.Rules.GUBP_ToolPlatforms(HostPlatform))
+            {
+                foreach (var Config in ProgramTarget.Rules.GUBP_ToolConfigs(HostPlatform))
+                {
+                    Agenda.AddTargets(new string[] { ProgramTarget.TargetName }, Plat, Config, InAddArgs: AddArgs);
+                }
+            }
+            return Agenda;
+        }
+    }
+
+
     public class EditorPlatformNode : CompileNode
     {
         UnrealTargetPlatform EditorPlatform;
@@ -1770,7 +1885,7 @@ public class GUBP : BuildCommand
         }
         public override int TimeoutInMinutes()
         {
-            return bIsMassive ? 240 : 90;
+            return bIsMassive ? 240 : base.TimeoutInMinutes();
         }
 
         public override string RootIfAnyForTempStorage()
@@ -1912,7 +2027,7 @@ public class GUBP : BuildCommand
         }
         public override int TimeoutInMinutes()
         {
-            return bIsMassive ? 240 : 90;
+            return bIsMassive ? 240 : base.TimeoutInMinutes();
         }
 
         public override void DoBuild(GUBP bp)
@@ -3477,7 +3592,7 @@ public class GUBP : BuildCommand
         bool bSkipTriggers = ParseParam("SkipTriggers");
         bFake = ParseParam("fake");
         bool bFakeEC = ParseParam("FakeEC");
-        int TimeIndex = ParseParamInt("TimeIndex", 0);
+        TimeIndex = ParseParamInt("TimeIndex", 0);
         bHackRunIOSCompilesOnMac = !ParseParam("IOSOnRPCUtility") && !HostPlatforms.Contains(UnrealTargetPlatform.Mac);
 
         bNoIOSOnPC = HostPlatforms.Contains(UnrealTargetPlatform.Mac);
@@ -3756,6 +3871,24 @@ public class GUBP : BuildCommand
             }
             AddNode(new ToolsNode(HostPlatform));
             AddNode(new InternalToolsNode(HostPlatform));
+            foreach (var ProgramTarget in Branch.BaseEngineProject.Properties.Programs)
+            {
+                bool bInternalOnly;
+                bool SeparateNode;
+
+                if (ProgramTarget.Rules.GUBP_AlwaysBuildWithTools(HostPlatform, out bInternalOnly, out SeparateNode) && ProgramTarget.Rules.SupportsPlatform(HostPlatform) && SeparateNode)
+                {
+                    if (bInternalOnly)
+                    {
+                        AddNode(new SingleInternalToolsNode(HostPlatform, ProgramTarget));
+                    }
+                    else
+                    {
+                        AddNode(new SingleToolsNode(HostPlatform, ProgramTarget));
+                    }
+                }
+            }
+
             AddNode(new EditorAndToolsNode(HostPlatform));
             AddNode(new NonUnityTestNode(HostPlatform));
 
@@ -4915,16 +5048,6 @@ public class GUBP : BuildCommand
                         RunCondition = RunCondition + ") true; else false;]\"";
                     }
 
-#if false
-                    if (GUBPNodes[NodeToDo].TimeoutInMinutes() > 0)
-                    {
-                        Args = Args + String.Format(" --timeLimitUnits minutes --timeLimit {0}", GUBPNodes[NodeToDo].TimeoutInMinutes());
-                    }
-#endif
-                    if (Sticky && NodeToDo == LastSticky)
-                    {
-                        Args = Args + " --releaseMode release";
-                    }
                     if (bDoNestedJobstep)
                     {
                         if (bDoFirstNestedJobstep)
@@ -4979,6 +5102,15 @@ public class GUBP : BuildCommand
                     {
                         Args = Args + " --condition " + RunCondition;
                     }
+                    if (GUBPNodes[NodeToDo].TimeoutInMinutes() > 0)
+                    {
+                        Args = Args + String.Format(" --timeLimitUnits minutes --timeLimit {0}", GUBPNodes[NodeToDo].TimeoutInMinutes());
+                    }
+                    if (Sticky && NodeToDo == LastSticky)
+                    {
+                        Args = Args + " --releaseMode release";
+                    }
+
 
                     RunECTool(Args);
                     if (bFakeEC && 
diff --git a/Engine/Source/Programs/CrashReporter/CrashReportClient/CrashReportClient.Target.cs b/Engine/Source/Programs/CrashReporter/CrashReportClient/CrashReportClient.Target.cs
--- a/Engine/Source/Programs/CrashReporter/CrashReportClient/CrashReportClient.Target.cs
+++ b/Engine/Source/Programs/CrashReporter/CrashReportClient/CrashReportClient.Target.cs
@@ -70,9 +70,10 @@ public class CrashReportClientTarget : TargetRules
 		// Do NOT produce additional console app exe
 		OutLinkEnvironmentConfiguration.bBuildAdditionalConsoleApplication = false;
 	}
-    public override bool GUBP_AlwaysBuildWithTools(UnrealTargetPlatform InHostPlatform, out bool bInternalToolOnly)
+    public override bool GUBP_AlwaysBuildWithTools(UnrealTargetPlatform InHostPlatform, out bool bInternalToolOnly, out bool SeparateNode)
     {
         bInternalToolOnly = false;
+        SeparateNode = false;
         return true;
     }
     public override List<UnrealTargetPlatform> GUBP_ToolPlatforms(UnrealTargetPlatform InHostPlatform)
diff --git a/Engine/Source/Programs/CrashReporter/MinidumpDiagnostics/MinidumpDiagnostics.Target.cs b/Engine/Source/Programs/CrashReporter/MinidumpDiagnostics/MinidumpDiagnostics.Target.cs
--- a/Engine/Source/Programs/CrashReporter/MinidumpDiagnostics/MinidumpDiagnostics.Target.cs
+++ b/Engine/Source/Programs/CrashReporter/MinidumpDiagnostics/MinidumpDiagnostics.Target.cs
@@ -62,9 +62,10 @@ public class MinidumpDiagnosticsTarget : TargetRules
 
 		OutCPPEnvironmentConfiguration.Definitions.Add("MINIDUMPDIAGNOSTICS=1");
 	}
-    public override bool GUBP_AlwaysBuildWithTools(UnrealTargetPlatform InHostPlatform, out bool bInternalToolOnly)
+    public override bool GUBP_AlwaysBuildWithTools(UnrealTargetPlatform InHostPlatform, out bool bInternalToolOnly, out bool SeparateNode)
     {
         bInternalToolOnly = true;
+        SeparateNode = false;
         return true;
     }
 }
diff --git a/Engine/Source/Programs/UnrealBuildTool/System/RulesCompiler.cs b/Engine/Source/Programs/UnrealBuildTool/System/RulesCompiler.cs
--- a/Engine/Source/Programs/UnrealBuildTool/System/RulesCompiler.cs
+++ b/Engine/Source/Programs/UnrealBuildTool/System/RulesCompiler.cs
@@ -556,11 +556,13 @@ namespace UnrealBuildTool
         }
         /// <summary>
         /// Return true if this target should always be built with the tools. Usually programs like unrealpak.
+        /// <param name="SeparateNode">If this is set to true, the program will get its own node</param>
         /// </summary>
         /// <returns>true if this target should always be built with the base editor.</returns>
-        public virtual bool GUBP_AlwaysBuildWithTools(UnrealTargetPlatform InHostPlatform, out bool bInternalToolOnly)
+        public virtual bool GUBP_AlwaysBuildWithTools(UnrealTargetPlatform InHostPlatform, out bool bInternalToolOnly, out bool SeparateNode)
         {
             bInternalToolOnly = false;
+            SeparateNode = false;
             return false;
         }
         /// <summary>
diff --git a/Engine/Source/Programs/UnrealFileServer/UnrealFileServer.Target.cs b/Engine/Source/Programs/UnrealFileServer/UnrealFileServer.Target.cs
--- a/Engine/Source/Programs/UnrealFileServer/UnrealFileServer.Target.cs
+++ b/Engine/Source/Programs/UnrealFileServer/UnrealFileServer.Target.cs
@@ -58,9 +58,10 @@ public class UnrealFileServerTarget : TargetRules
 		// UnrealHeaderTool is a console application, not a Windows app (sets entry point to main(), instead of WinMain())
 		OutLinkEnvironmentConfiguration.bIsBuildingConsoleApplication = true;
 	}
-    public override bool GUBP_AlwaysBuildWithTools(UnrealTargetPlatform InHostPlatform, out bool bInternalToolOnly)
+    public override bool GUBP_AlwaysBuildWithTools(UnrealTargetPlatform InHostPlatform, out bool bInternalToolOnly, out bool SeparateNode)
     {
         bInternalToolOnly = false;
+        SeparateNode = false;
         return true;
     }
 }
diff --git a/Engine/Source/Programs/UnrealLightmass/UnrealLightmass.Target.cs b/Engine/Source/Programs/UnrealLightmass/UnrealLightmass.Target.cs
--- a/Engine/Source/Programs/UnrealLightmass/UnrealLightmass.Target.cs
+++ b/Engine/Source/Programs/UnrealLightmass/UnrealLightmass.Target.cs
@@ -74,9 +74,10 @@ public class UnrealLightmassTarget : TargetRules
 		// Disable logging, lightmass will create its own unique logging file
 		OutCPPEnvironmentConfiguration.Definitions.Add("ALLOW_LOG_FILE=0");
 	}
-    public override bool GUBP_AlwaysBuildWithTools(UnrealTargetPlatform InHostPlatform, out bool bInternalToolOnly)
+    public override bool GUBP_AlwaysBuildWithTools(UnrealTargetPlatform InHostPlatform, out bool bInternalToolOnly, out bool SeparateNode)
     {
         bInternalToolOnly = false;
+        SeparateNode = false;
         return true;
     }
 }
diff --git a/Engine/Source/Programs/UnrealPak/UnrealPak.Target.cs b/Engine/Source/Programs/UnrealPak/UnrealPak.Target.cs
--- a/Engine/Source/Programs/UnrealPak/UnrealPak.Target.cs
+++ b/Engine/Source/Programs/UnrealPak/UnrealPak.Target.cs
@@ -52,9 +52,10 @@ public class UnrealPakTarget : TargetRules
 		// UnrealHeaderTool is a console application, not a Windows app (sets entry point to main(), instead of WinMain())
 		OutLinkEnvironmentConfiguration.bIsBuildingConsoleApplication = true;
 	}
-    public override bool GUBP_AlwaysBuildWithTools(UnrealTargetPlatform InHostPlatform, out bool bInternalToolOnly)
+    public override bool GUBP_AlwaysBuildWithTools(UnrealTargetPlatform InHostPlatform, out bool bInternalToolOnly, out bool SeparateNode)
     {
         bInternalToolOnly = false;
+        SeparateNode = false;
         return true;
     }
 }
