build script, fix shared promotion emails

diff --git a/Engine/Source/Programs/AutomationTool/Scripts/GUBP.Automation.cs b/Engine/Source/Programs/AutomationTool/Scripts/GUBP.Automation.cs
--- a/Engine/Source/Programs/AutomationTool/Scripts/GUBP.Automation.cs
+++ b/Engine/Source/Programs/AutomationTool/Scripts/GUBP.Automation.cs
@@ -1608,6 +1608,11 @@ public class GUBP : BuildCommand
         {
             return WaitForPromotionUserInput.StaticGetFullName("Shared", bInLabelPromoted);
         }
+        public override string FailureEMails(GUBP bp, string Branch)
+        {
+            return MergeSpaceStrings(base.FailureEMails(bp, Branch),
+               bp.Branch.BaseEngineProject.Properties.Targets[TargetRules.TargetType.Editor].Rules.GUBP_GetPromotionEMails_EditorTypeOnly(Branch));
+        }
     }
 
     public class LabelPromotableNode : GUBPNode
