build script, no timeout on triggers

diff --git a/Engine/Source/Programs/AutomationTool/Scripts/GUBP.Automation.cs b/Engine/Source/Programs/AutomationTool/Scripts/GUBP.Automation.cs
--- a/Engine/Source/Programs/AutomationTool/Scripts/GUBP.Automation.cs
+++ b/Engine/Source/Programs/AutomationTool/Scripts/GUBP.Automation.cs
@@ -1495,6 +1495,10 @@ public class GUBP : BuildCommand
             }
             return Result;
         }
+        public override int TimeoutInMinutes()
+        {
+            return 0;
+        }
     }
 
     public class WaitForPromotionUserInput : WaitForUserInput
