build script, remove timeouts from top level jobs because it includes precondition times

diff --git a/Engine/Source/Programs/AutomationTool/Scripts/GUBP.Automation.cs b/Engine/Source/Programs/AutomationTool/Scripts/GUBP.Automation.cs
--- a/Engine/Source/Programs/AutomationTool/Scripts/GUBP.Automation.cs
+++ b/Engine/Source/Programs/AutomationTool/Scripts/GUBP.Automation.cs
@@ -5111,10 +5111,13 @@ public class GUBP : BuildCommand
                     {
                         Args = Args + " --condition " + RunCondition;
                     }
+#if false
+                    // this doesn't work because it includes precondition time
                     if (GUBPNodes[NodeToDo].TimeoutInMinutes() > 0)
                     {
                         Args = Args + String.Format(" --timeLimitUnits minutes --timeLimit {0}", GUBPNodes[NodeToDo].TimeoutInMinutes());
                     }
+#endif
                     if (Sticky && NodeToDo == LastSticky)
                     {
                         Args = Args + " --releaseMode release";
