merge fort stuff for no shared promotable

diff --git a/Engine/Source/Programs/AutomationTool/Scripts/GUBP.Automation.cs b/Engine/Source/Programs/AutomationTool/Scripts/GUBP.Automation.cs
--- a/Engine/Source/Programs/AutomationTool/Scripts/GUBP.Automation.cs
+++ b/Engine/Source/Programs/AutomationTool/Scripts/GUBP.Automation.cs
@@ -476,7 +476,7 @@ public class GUBP : BuildCommand
         public override int AgentMemoryRequirement(GUBP bp)
         {
             if (bp.ParseParam("Launcher") || bp.TimeIndex != 0)
-            {
+        {
                 return base.AgentMemoryRequirement(bp);
             }
             return 32;
@@ -501,10 +501,10 @@ public class GUBP : BuildCommand
         {
             if (HostPlatform != UnrealTargetPlatform.Mac)
             {
-            	UE4Build.AddUATFilesToBuildProducts();
-            	UE4Build.AddUBTFilesToBuildProducts();
-        	}
-    	}
+                UE4Build.AddUATFilesToBuildProducts();
+                UE4Build.AddUBTFilesToBuildProducts();
+            }
+        }
     }
 
     public class RootEditorNode : CompileNode
@@ -626,7 +626,7 @@ public class GUBP : BuildCommand
         {
             if (!GUBP.bBuildRocket) // more errors and more performance by just starting before the root editor is done
             {
-                AddPseudodependency(RootEditorNode.StaticGetFullName(HostPlatform));
+            AddPseudodependency(RootEditorNode.StaticGetFullName(HostPlatform));
             }
             AgentSharingGroup = "Tools" + StaticGetHostPlatformSuffix(HostPlatform);
         }
@@ -751,7 +751,7 @@ public class GUBP : BuildCommand
         {
             if (!GUBP.bBuildRocket) // more errors and more performance by just starting before the root editor is done
             {
-                AddPseudodependency(RootEditorNode.StaticGetFullName(HostPlatform));
+            AddPseudodependency(RootEditorNode.StaticGetFullName(HostPlatform));
             }
             AgentSharingGroup = "Tools" + StaticGetHostPlatformSuffix(HostPlatform);
         }
@@ -1004,11 +1004,14 @@ public class GUBP : BuildCommand
             if (InGameProj.GameName != bp.Branch.BaseEngineProject.GameName && GameProj.Properties.Targets.ContainsKey(TargetRules.TargetType.Editor))
             {
                 AddPseudodependency(EditorGameNode.StaticGetFullName(InHostPlatform, GameProj));
-                AddPseudodependency(GamePlatformMonolithicsNode.StaticGetFullName(HostPlatform, bp.Branch.BaseEngineProject, TargetPlatform));
+                if (bp.HasNode(GamePlatformMonolithicsNode.StaticGetFullName(HostPlatform, bp.Branch.BaseEngineProject, TargetPlatform)))
+                {
+                	AddPseudodependency(GamePlatformMonolithicsNode.StaticGetFullName(HostPlatform, bp.Branch.BaseEngineProject, TargetPlatform));
+            	}
             }
             else
             {
-                if (TargetPlatform != InHostPlatform
+                if (TargetPlatform != InHostPlatform && bp.HasNode(GamePlatformMonolithicsNode.StaticGetFullName(InHostPlatform, bp.Branch.BaseEngineProject, InHostPlatform))
                     && !GUBP.bBuildRocket) // more errors and more performance by just starting, for example, IOS without waiting for mac
                 {
                     AddPseudodependency(GamePlatformMonolithicsNode.StaticGetFullName(InHostPlatform, bp.Branch.BaseEngineProject, InHostPlatform));
@@ -1643,7 +1646,7 @@ public class GUBP : BuildCommand
             {
                 //shared promotable has a shorter name
                 CompleteLabelPrefix = "Promotable";
-            }                   
+            }
             if (GUBP.bBuildRocket)
             {
                 CompleteLabelPrefix = "Rocket-" + CompleteLabelPrefix;
@@ -3878,8 +3881,8 @@ public class GUBP : BuildCommand
             {
                 AddNode(new RootEditorHeadersNode(HostPlatform));
             }
-            AddNode(new ToolsNode(HostPlatform));
-            AddNode(new InternalToolsNode(HostPlatform));
+            AddNode(new ToolsNode(HostPlatform));            
+			AddNode(new InternalToolsNode(HostPlatform));
             foreach (var ProgramTarget in Branch.BaseEngineProject.Properties.Programs)
             {
                 bool bInternalOnly;
@@ -3899,7 +3902,6 @@ public class GUBP : BuildCommand
             }
 
             AddNode(new EditorAndToolsNode(HostPlatform));
-            AddNode(new NonUnityTestNode(HostPlatform));
 
             if (bOrthogonalizeEditorPlatforms)
             {
@@ -3937,7 +3939,7 @@ public class GUBP : BuildCommand
                         throw new AutomationException("We assume that if we have shared promotable, the base engine is in it. Some games are marked as shared and promotable, but the base engine is not.");
                     }
                 }
-                else
+                else if (NumShared > 0)
                 {
                     DoASharedPromotable = true;
                 }
@@ -3950,6 +3952,8 @@ public class GUBP : BuildCommand
 
             if (DoASharedPromotable)
             {
+                AddNode(new NonUnityTestNode(HostPlatform));
+
                 var AgentSharingGroup = "Shared_EditorTests" + HostPlatformNode.StaticGetHostPlatformSuffix(HostPlatform);
 
                 Dictionary<string, List<UnrealTargetPlatform>> NonCodeProjectNames;
@@ -4009,6 +4013,12 @@ public class GUBP : BuildCommand
                     {
                         var Target = Branch.BaseEngineProject.Properties.Targets[Kind];
                         var Platforms = Target.Rules.GUBP_GetPlatforms_MonolithicOnly(HostPlatform);
+                        if (Platforms.Contains(HostPlatform))
+                        {
+                            // we want the host platform first since some some pseudodependencies look to see if the shared promotable exists.
+                            Platforms.Remove(HostPlatform);
+                            Platforms.Insert(0, HostPlatform);
+                        }
                         foreach (var Plat in Platforms)
                         {
                             if (!Platform.Platforms[HostPlatform].CanHostPlatform(Plat))
@@ -4110,23 +4120,23 @@ public class GUBP : BuildCommand
                                     }
                                     if (bAutomatedTesting)
                                     {
-                                        if (HostPlatform == UnrealTargetPlatform.Mac) continue; //temp hack till mac automated testing works
-                                        var ThisMonoGameTestNodes = new List<string>();
-                                        foreach (var Test in GameTests)
-                                        {
-                                            var TestName = Test.Key + "_" + Plat.ToString();
-                                            ThisMonoGameTestNodes.Add(AddNode(new UATTestNode(this, HostPlatform, NonCodeProject, TestName, Test.Value, CookedAgentSharingGroup, false, RequiredPlatforms)));
-                                        }
-                                        if (ThisMonoGameTestNodes.Count > 0)
-                                        {
-                                            GameTestNodes.Add(AddNode(new GameAggregateNode(this, HostPlatform, NonCodeProject, "CookedTests_" + Plat.ToString() + "_" + Kind.ToString() + HostPlatformNode.StaticGetHostPlatformSuffix(HostPlatform), ThisMonoGameTestNodes, 0.0f)));
-                                        }
+                                    if (HostPlatform == UnrealTargetPlatform.Mac) continue; //temp hack till mac automated testing works
+                                    var ThisMonoGameTestNodes = new List<string>();
+                                    foreach (var Test in GameTests)
+                                    {
+                                        var TestName = Test.Key + "_" + Plat.ToString();
+                                        ThisMonoGameTestNodes.Add(AddNode(new UATTestNode(this, HostPlatform, NonCodeProject, TestName, Test.Value, CookedAgentSharingGroup, false, RequiredPlatforms)));
+                                    }
+                                    if (ThisMonoGameTestNodes.Count > 0)
+                                    {
+                                        GameTestNodes.Add(AddNode(new GameAggregateNode(this, HostPlatform, NonCodeProject, "CookedTests_" + Plat.ToString() + "_" + Kind.ToString() + HostPlatformNode.StaticGetHostPlatformSuffix(HostPlatform), ThisMonoGameTestNodes, 0.0f)));
                                     }
                                 }
                             }
                         }
                     }
                 }
+                }
 #if false
                 //for now, non-code projects don't do client or server.
                 foreach (var ServerPlatform in ServerPlatforms)
@@ -4244,20 +4254,20 @@ public class GUBP : BuildCommand
                                 }
                                 if (bAutomatedTesting)
                                 {
-                                    if (HostPlatform == UnrealTargetPlatform.Mac) continue; //temp hack till mac automated testing works
-                                    var GameTests = Target.Rules.GUBP_GetGameTests_MonolithicOnly(HostPlatform, GetAltHostPlatform(HostPlatform), Plat);
-                                    var RequiredPlatforms = new List<UnrealTargetPlatform> { Plat };
-                                    var ThisMonoGameTestNodes = new List<string>();
+if (HostPlatform == UnrealTargetPlatform.Mac) continue; //temp hack till mac automated testing works
+                                var GameTests = Target.Rules.GUBP_GetGameTests_MonolithicOnly(HostPlatform, GetAltHostPlatform(HostPlatform), Plat);
+                                var RequiredPlatforms = new List<UnrealTargetPlatform> { Plat };
+                                var ThisMonoGameTestNodes = new List<string>();
 
-                                    foreach (var Test in GameTests)
-                                    {
-                                        var TestNodeName = Test.Key + "_" + Plat.ToString();
-                                        ThisMonoGameTestNodes.Add(AddNode(new UATTestNode(this, HostPlatform, CodeProj, TestNodeName, Test.Value, CookedAgentSharingGroup, false, RequiredPlatforms)));
-                                    }
-                                    if (ThisMonoGameTestNodes.Count > 0)
-                                    {
-                                        GameTestNodes.Add(AddNode(new GameAggregateNode(this, HostPlatform, CodeProj, "CookedTests_" + Plat.ToString() + "_" + Kind.ToString() + HostPlatformNode.StaticGetHostPlatformSuffix(HostPlatform), ThisMonoGameTestNodes, 0.0f)));
-                                    }
+                                foreach (var Test in GameTests)
+                                {
+                                    var TestNodeName = Test.Key + "_" + Plat.ToString();
+                                    ThisMonoGameTestNodes.Add(AddNode(new UATTestNode(this, HostPlatform, CodeProj, TestNodeName, Test.Value, CookedAgentSharingGroup, false, RequiredPlatforms)));
+                                }
+                                if (ThisMonoGameTestNodes.Count > 0)
+                                {
+                                    GameTestNodes.Add(AddNode(new GameAggregateNode(this, HostPlatform, CodeProj, "CookedTests_" + Plat.ToString() + "_" + Kind.ToString() + HostPlatformNode.StaticGetHostPlatformSuffix(HostPlatform), ThisMonoGameTestNodes, 0.0f)));
+                                }
 
                                 }
                             }
@@ -4266,31 +4276,31 @@ public class GUBP : BuildCommand
                 }
                 if (bAutomatedTesting)
                 {
-                    foreach (var ServerPlatform in ServerPlatforms)
+                foreach (var ServerPlatform in ServerPlatforms)
+                {
+                    foreach (var GamePlatform in GamePlatforms)
                     {
-                        foreach (var GamePlatform in GamePlatforms)
+if (HostPlatform == UnrealTargetPlatform.Mac) continue; //temp hack till mac automated testing works
+                        var Target = CodeProj.Properties.Targets[TargetRules.TargetType.Game];
+                        var ClientServerTests = Target.Rules.GUBP_GetClientServerTests_MonolithicOnly(HostPlatform, GetAltHostPlatform(HostPlatform), ServerPlatform, GamePlatform);
+                        var RequiredPlatforms = new List<UnrealTargetPlatform> { ServerPlatform };
+                        if (ServerPlatform != GamePlatform)
                         {
-                            if (HostPlatform == UnrealTargetPlatform.Mac) continue; //temp hack till mac automated testing works
-                            var Target = CodeProj.Properties.Targets[TargetRules.TargetType.Game];
-                            var ClientServerTests = Target.Rules.GUBP_GetClientServerTests_MonolithicOnly(HostPlatform, GetAltHostPlatform(HostPlatform), ServerPlatform, GamePlatform);
-                            var RequiredPlatforms = new List<UnrealTargetPlatform> { ServerPlatform };
-                            if (ServerPlatform != GamePlatform)
-                            {
-                                RequiredPlatforms.Add(GamePlatform);
-                            }
-                            foreach (var Test in ClientServerTests)
-                            {
-                                var TestNodeName = Test.Key + "_" + GamePlatform.ToString() + "_" + ServerPlatform.ToString();
-                                GameTestNodes.Add(AddNode(new UATTestNode(this, HostPlatform, CodeProj, TestNodeName, Test.Value, CookedAgentSharingGroup, false, RequiredPlatforms)));
-                            }
+                            RequiredPlatforms.Add(GamePlatform);
+                        }
+                        foreach (var Test in ClientServerTests)
+                        {
+                            var TestNodeName = Test.Key + "_" + GamePlatform.ToString() + "_" + ServerPlatform.ToString();
+                            GameTestNodes.Add(AddNode(new UATTestNode(this, HostPlatform, CodeProj, TestNodeName, Test.Value, CookedAgentSharingGroup, false, RequiredPlatforms)));
                         }
                     }
-                    if (GameTestNodes.Count > 0)
-                    {
-                        AddNode(new GameAggregateNode(this, HostPlatform, CodeProj, "AllCookedTests", GameTestNodes));
-                    }
+                }
+                if (GameTestNodes.Count > 0)
+                {
+                    AddNode(new GameAggregateNode(this, HostPlatform, CodeProj, "AllCookedTests", GameTestNodes));
                 }
             }
+            }
         }
 
         int NumSharedAllHosts = 0;
@@ -4882,7 +4892,7 @@ public class GUBP : BuildCommand
                             EMails = ForceAllEmailsTo;
                         }
                         else
-                        {                            
+                        {
                             EMails = GetEMailListForNode(this, NodeToDo);
                         }
                         string AddEmails = ParseParamValue("AddEmails");
@@ -4965,19 +4975,19 @@ public class GUBP : BuildCommand
                     var UncompletedEcDeps = new List<string>();
 
                     {
-                        var EcDeps = GetECDependencies(NodeToDo);
-                        foreach (var Dep in EcDeps)
+                    var EcDeps = GetECDependencies(NodeToDo);
+                    foreach (var Dep in EcDeps)
+                    {
+                        if (GUBPNodes[Dep].RunInEC() && !NodeIsAlreadyComplete(Dep, LocalOnly) && OrdereredToDo.Contains(Dep)) // if something is already finished, we don't put it into EC
                         {
-                            if (GUBPNodes[Dep].RunInEC() && !NodeIsAlreadyComplete(Dep, LocalOnly) && OrdereredToDo.Contains(Dep)) // if something is already finished, we don't put it into EC
+                            if (OrdereredToDo.IndexOf(Dep) > OrdereredToDo.IndexOf(NodeToDo))
                             {
-                                if (OrdereredToDo.IndexOf(Dep) > OrdereredToDo.IndexOf(NodeToDo))
-                                {
-                                    throw new AutomationException("Topological sort error, node {0} has a dependency of {1} which sorted after it.", NodeToDo, Dep);
-                                }
-                                UncompletedEcDeps.Add(Dep);
+                                throw new AutomationException("Topological sort error, node {0} has a dependency of {1} which sorted after it.", NodeToDo, Dep);
                             }
+                            UncompletedEcDeps.Add(Dep);
                         }
                     }
+                    }
 
                     string MyAgentGroup = GUBPNodes[NodeToDo].AgentSharingGroup;
                     bool bDoNestedJobstep = false;
@@ -4992,10 +5002,10 @@ public class GUBP : BuildCommand
 
                         PreConditionUncompletedEcDeps = new List<string>();
 
-                        var MyChain = AgentGroupChains[MyAgentGroup];
-                        int MyIndex = MyChain.IndexOf(NodeToDo);
-                        if (MyIndex > 0)
-                        {
+                                var MyChain = AgentGroupChains[MyAgentGroup];
+                                int MyIndex = MyChain.IndexOf(NodeToDo);
+                                if (MyIndex  > 0)
+                                {
                             PreConditionUncompletedEcDeps.Add(MyChain[MyIndex - 1]);
                         }
                         else
@@ -5014,48 +5024,48 @@ public class GUBP : BuildCommand
                                             throw new AutomationException("Topological sort error, node {0} has a dependency of {1} which sorted after it.", Chain, Dep);
                                         }
                                         if (!MyChain.Contains(Dep) && !PreConditionUncompletedEcDeps.Contains(Dep))
-                                        {
+                                    {
                                             PreConditionUncompletedEcDeps.Add(Dep);
                                         }
                                     }
+                                    }
+                                }
                                 }
                             }
-                        }
-                    }
                     if (PreConditionUncompletedEcDeps.Count > 0)
                     {
                         PreCondition = "\"$[/javascript if(";
                         // these run "parallel", but we add preconditions to serialize them
-                        int Index = 0;
+                            int Index = 0;
                         foreach (var Dep in PreConditionUncompletedEcDeps)
-                        {
+                            {
                             PreCondition = PreCondition + "getProperty('" + GetJobStep(ParentPath, Dep) + "/status') == 'completed'";
-                            Index++;
+                                Index++;
                             if (Index != PreConditionUncompletedEcDeps.Count)
-                            {
-                                PreCondition = PreCondition + " && ";
+                                {
+                                    PreCondition = PreCondition + " && ";
+                                }
                             }
+                            PreCondition = PreCondition + ") true;]\"";
                         }
-                        PreCondition = PreCondition + ") true;]\"";
-                    }
 
                     if (UncompletedEcDeps.Count > 0)
-                    {
-                        RunCondition = "\"$[/javascript if(";
-                        int Index = 0;
-                        foreach (var Dep in UncompletedEcDeps)
                         {
+                            RunCondition = "\"$[/javascript if(";
+                            int Index = 0;
+                            foreach (var Dep in UncompletedEcDeps)
+                            {
                             RunCondition = RunCondition + "('$[" + GetJobStep(ParentPath, Dep) + "/outcome]' == 'success' || ";
                             RunCondition = RunCondition + "'$[" + GetJobStep(ParentPath, Dep) + "/outcome]' == 'warning')";
 
-                            Index++;
-                            if (Index != UncompletedEcDeps.Count)
-                            {
-                                RunCondition = RunCondition + " && ";
+                                Index++;
+                                if (Index != UncompletedEcDeps.Count)
+                                {
+                                    RunCondition = RunCondition + " && ";
+                                }
                             }
+                            RunCondition = RunCondition + ") true; else false;]\"";
                         }
-                        RunCondition = RunCondition + ") true; else false;]\"";
-                    }
 
                     if (bDoNestedJobstep)
                     {
@@ -5069,7 +5079,7 @@ public class GUBP : BuildCommand
                                     NestArgs = NestArgs + " --precondition " + PreCondition;
                                 }
                                 RunECTool(NestArgs);
-                            }
+                    }
                             {
                                 string NestArgs = String.Format("createJobStep --parentPath {0}/jobSteps[{1}] --jobStepName {2}_GetPool --subprocedure GUBP{3}_AgentShare_GetPool --parallel 1 --actualParameter AgentSharingGroup={4} --actualParameter NodeName={5}",
                                     ParentPath, MyAgentGroup, MyAgentGroup, GUBPNodes[NodeToDo].ECProcedureInfix(), MyAgentGroup, NodeToDo);
@@ -5217,35 +5227,35 @@ public class GUBP : BuildCommand
             {
                 if (ParseParamValue("OnlyNode", "") == "")
                 {
-                    throw new AutomationException("Node {0} was not processed yet?", NodeToDo);
-                }
+                throw new AutomationException("Node {0} was not processed yet?", NodeToDo);
+            }
             }
             else
             {
-                foreach (var Dep in GUBPNodes[NodeToDo].FullNamesOfDependencies)
+            foreach (var Dep in GUBPNodes[NodeToDo].FullNamesOfDependencies)
+            {
+                if (GUBPNodes[Dep].BuildProducts == null)
                 {
-                    if (GUBPNodes[Dep].BuildProducts == null)
-                    {
                         if (ParseParamValue("OnlyNode", "") == "")
                         {
-                            throw new AutomationException("Node {0} was not processed yet2? Processing {1}", Dep, NodeToDo);
-                        }
+                    throw new AutomationException("Node {0} was not processed yet2? Processing {1}", Dep, NodeToDo);
+                }
                     }
                     else
                     {
-                        foreach (var Prod in GUBPNodes[Dep].BuildProducts)
-                        {
-                            GUBPNodes[NodeToDo].AddDependentBuildProduct(Prod);
-                        }
-                        if (GUBPNodes[Dep].AllDependencyBuildProducts == null)
-                        {
-                            throw new AutomationException("Node {0} was not processed yet3?  Processing {1}", Dep, NodeToDo);
-                        }
-                        foreach (var Prod in GUBPNodes[Dep].AllDependencyBuildProducts)
-                        {
-                            GUBPNodes[NodeToDo].AddDependentBuildProduct(Prod);
-                        }
-                    }
+                foreach (var Prod in GUBPNodes[Dep].BuildProducts)
+                {
+                    GUBPNodes[NodeToDo].AddDependentBuildProduct(Prod);
+                }
+                if (GUBPNodes[Dep].AllDependencyBuildProducts == null)
+                {
+                    throw new AutomationException("Node {0} was not processed yet3?  Processing {1}", Dep, NodeToDo);
+                }
+                foreach (var Prod in GUBPNodes[Dep].AllDependencyBuildProducts)
+                {
+                    GUBPNodes[NodeToDo].AddDependentBuildProduct(Prod);
+                } 
+            }
                 }
             }
             string NodeStoreName = StoreName + "-" + GUBPNodes[NodeToDo].GetFullName();
