#summary Modified the Unreal Enging Launcher build failure e-mail list
#change Modified the e-mail list to include lswift and myself
#change Added support for properly resoving [phosphor] and [pitbull] e-mail tags

diff --git a/Engine/Source/Programs/AutomationTool/Scripts/GUBP.Automation.cs b/Engine/Source/Programs/AutomationTool/Scripts/GUBP.Automation.cs
--- a/Engine/Source/Programs/AutomationTool/Scripts/GUBP.Automation.cs
+++ b/Engine/Source/Programs/AutomationTool/Scripts/GUBP.Automation.cs
@@ -4900,7 +4900,13 @@ if (HostPlatform == UnrealTargetPlatform.Mac) continue; //temp hack till mac aut
                         {
                             EMails = GUBPNode.MergeSpaceStrings(AddEmails, EMails);
                         }
-                        EMails = EMails.Trim().Replace("[epic]", "@epicgames.com");
+
+                        {
+                            EMails = EMails.Trim().Replace("[epic]", "@epicgames.com");
+                            EMails = EMails.Replace("[phosphor]", "@phosphorgames.com");
+                            EMails = EMails.Replace("[pitbull]", "@pitbullstudio.co.uk");
+                        }
+                        
                         ECProps.Add("FailEmails/" + NodeToDo + "=" + EMails);
                     }
                     if (GUBPNodes[NodeToDo].SendSuccessEmail())
