#ttp 331732 - UAT: Make sure *.exe.config build products also get +w flag when being added to P4
#proj UE4
#branch UE4
#summary Done.

diff --git a/Engine/Source/Programs/AutomationTool/UE4Build.cs b/Engine/Source/Programs/AutomationTool/UE4Build.cs
--- a/Engine/Source/Programs/AutomationTool/UE4Build.cs
+++ b/Engine/Source/Programs/AutomationTool/UE4Build.cs
@@ -1305,7 +1305,7 @@ namespace AutomationTool
                 }
                 else
                 {
-                    if (FileStats.Type == P4FileType.Binary && (FileStats.Attributes & P4FileAttributes.Writeable) != P4FileAttributes.Writeable)
+					if (IsBuildProduct(File, FileStats) && (FileStats.Attributes & P4FileAttributes.Writeable) != P4FileAttributes.Writeable)
                     {
                         P4.ChangeFileType(File, P4FileAttributes.Writeable);
                     }
@@ -1314,6 +1314,21 @@ namespace AutomationTool
 			}
 		}
 
+		/// <summary>
+		/// Determines if this file is a build product.
+		/// </summary>
+		/// <param name="File">File path</param>
+		/// <param name="FileStats">P4 file stats.</param>
+		/// <returns>True if this is a Windows build product. False otherwise.</returns>
+		private static bool IsBuildProduct(string File, P4FileStat FileStats)
+		{
+			if(FileStats.Type == P4FileType.Binary)
+			{
+				return true;
+			}
+
+			return FileStats.Type == P4FileType.Text && File.EndsWith(".exe.config", StringComparison.InvariantCultureIgnoreCase);
+		}
 
 		/// <summary>
 		/// Add UBT files to build products
