build script, git pull requests -> shelves

diff --git a/Engine/Source/Programs/AutomationTool/P4Utils.cs b/Engine/Source/Programs/AutomationTool/P4Utils.cs
--- a/Engine/Source/Programs/AutomationTool/P4Utils.cs
+++ b/Engine/Source/Programs/AutomationTool/P4Utils.cs
@@ -370,6 +370,10 @@ namespace AutomationTool
 		/// List of global options for this connection (client/user)
 		/// </summary>
 		private string GlobalOptions;
+        /// <summary>
+        /// List of global options for this connection (client/user)
+        /// </summary>
+        private string GlobalOptionsWithoutClient;
 		/// <summary>
 		/// Path where this connection's log is to go to
 		/// </summary>
@@ -388,6 +392,7 @@ namespace AutomationTool
 			var ClientOpts = String.IsNullOrEmpty(Client) ? "" : ("-c" + Client + " ");
 			var ServerOpts = String.IsNullOrEmpty(ServerAndPort) ? "" : ("-p" + ServerAndPort + " ");			
 			GlobalOptions = UserOpts + ClientOpts + ServerOpts;
+            GlobalOptionsWithoutClient = UserOpts + ServerOpts;
 
 			if (P4LogPath == null)
 			{
@@ -414,10 +419,10 @@ namespace AutomationTool
 		/// <param name="Input">Stdin</param>
 		/// <param name="AllowSpew">true for spew</param>
 		/// <returns>Exit code</returns>
-		public ProcessResult P4(string CommandLine, string Input = null, bool AllowSpew = true)
+        public ProcessResult P4(string CommandLine, string Input = null, bool AllowSpew = true, bool WithClient = true)
 		{
 			CheckP4Enabled();
-			return CommandUtils.Run(HostPlatform.Current.P4Exe, GlobalOptions + CommandLine, Input, AllowSpew ? CommandUtils.ERunOptions.AllowSpew : CommandUtils.ERunOptions.NoLoggingOfRunCommand);
+            return CommandUtils.Run(HostPlatform.Current.P4Exe, (WithClient ? GlobalOptions : GlobalOptionsWithoutClient) + CommandLine, Input, AllowSpew ? CommandUtils.ERunOptions.AllowSpew : CommandUtils.ERunOptions.NoLoggingOfRunCommand);
 		}
 
 		/// <summary>
@@ -428,12 +433,12 @@ namespace AutomationTool
 		/// <param name="Input">Stdin input.</param>
 		/// <param name="AllowSpew">Whether the command should spew.</param>
 		/// <returns>True if succeeded, otherwise false.</returns>
-		public bool P4Output(out string Output, string CommandLine, string Input = null, bool AllowSpew = true)
+        public bool P4Output(out string Output, string CommandLine, string Input = null, bool AllowSpew = true, bool WithClient = true)
 		{
 			CheckP4Enabled();
 			Output = "";
 
-			var Result = P4(CommandLine, Input, AllowSpew);
+            var Result = P4(CommandLine, Input, AllowSpew, WithClient);
 
 			Output = Result.Output;
 			return Result == 0;
@@ -445,11 +450,11 @@ namespace AutomationTool
 		/// <param name="CommandLine">Commandline to pass to p4.</param>
 		/// <param name="Input">Stdin input.</param>
 		/// <param name="AllowSpew">Whether the command is allowed to spew.</param>
-		public void LogP4(string CommandLine, string Input = null, bool AllowSpew = true)
+        public void LogP4(string CommandLine, string Input = null, bool AllowSpew = true, bool WithClient = true)
 		{
 			CheckP4Enabled();
 			string Output;
-			if (!LogP4Output(out Output, CommandLine, Input, AllowSpew))
+            if (!LogP4Output(out Output, CommandLine, Input, AllowSpew, WithClient))
 			{
 				throw new P4Exception("p4.exe {0} failed.", CommandLine);
 			}
@@ -463,7 +468,7 @@ namespace AutomationTool
 		/// <param name="Input">Stdin input.</param>
 		/// <param name="AllowSpew">Whether the command should spew.</param>
 		/// <returns>True if succeeded, otherwise false.</returns>
-		public bool LogP4Output(out string Output, string CommandLine, string Input = null, bool AllowSpew = true)
+        public bool LogP4Output(out string Output, string CommandLine, string Input = null, bool AllowSpew = true, bool WithClient = true)
 		{
 			CheckP4Enabled();
 			Output = "";
@@ -474,7 +479,7 @@ namespace AutomationTool
 				return false;
 			}
 
-			var Result = P4(CommandLine, Input, AllowSpew);
+            var Result = P4(CommandLine, Input, AllowSpew, WithClient);
 
 			CommandUtils.WriteToFile(LogPath, CommandLine + "\n");
 			CommandUtils.WriteToFile(LogPath, Result.Output);
@@ -628,6 +633,17 @@ namespace AutomationTool
 			LogP4("unshelve " + String.Format("-s {0} ", FromCL) + String.Format("-c {0} ", ToCL) + CommandLine);
 		}
 
+        /// <summary>
+        /// Invokes p4 unshelve command.
+        /// </summary>
+        /// <param name="FromCL">Changelist to unshelve.</param>
+        /// <param name="ToCL">Changelist where the checked out files should be added.</param>
+        /// <param name="CommandLine">Commandline for the command.</param>
+        public void Shelve(int FromCL, string CommandLine = "")
+        {
+            CheckP4Enabled();
+            LogP4("shelve " + String.Format("-r -c {0} ", FromCL) + CommandLine);
+        }
 		/// <summary>
 		/// Invokes p4 edit command.
 		/// </summary>
@@ -738,7 +754,7 @@ namespace AutomationTool
 		/// </summary>
 		/// <param name="CL">Changelist to revert</param>
 		/// <param name="CommandLine">Commandline for the command.</param>
-		public void Revert(int CL, string CommandLine)
+		public void Revert(int CL, string CommandLine = "")
 		{
 			CheckP4Enabled();
 			LogP4("revert " + String.Format("-c {0} ", CL) + CommandLine);
@@ -1542,8 +1558,9 @@ namespace AutomationTool
 		{
 			CheckP4Enabled();
 			CommandUtils.Log("Checking if client {0} exists", ClientName);
-			var P4Result = P4(String.Format("-c {0} where //...", ClientName), AllowSpew: false);
-			return P4Result.Output.IndexOf("unknown - use 'client' command", StringComparison.InvariantCultureIgnoreCase) < 0;
+
+            var P4Result = P4(String.Format("-c {0} where //...", ClientName), AllowSpew: false, WithClient: false);
+            return P4Result.Output.IndexOf("unknown - use 'client' command", StringComparison.InvariantCultureIgnoreCase) < 0 && P4Result.Output.IndexOf("doesn't exist", StringComparison.InvariantCultureIgnoreCase) < 0;
 		}
 
 		/// <summary>
@@ -1583,7 +1600,7 @@ namespace AutomationTool
 		public P4ClientInfo GetClientInfoInternal(string ClientName)
 		{
 			P4ClientInfo Info = new P4ClientInfo();
-			var P4Result = P4(String.Format("client -o {0}", ClientName), AllowSpew: false);
+            var P4Result = P4(String.Format("client -o {0}", ClientName), AllowSpew: false, WithClient: false);
 			if (P4Result == 0)
 			{
 				var Tags = ParseTaggedP4Output(P4Result.Output);
@@ -1659,7 +1676,7 @@ namespace AutomationTool
 			var ClientList = new List<P4ClientInfo>();
 
 			// Get all clients for this user
-			var P4Result = P4(String.Format("clients -u {0}", UserName), AllowSpew: false);
+            var P4Result = P4(String.Format("clients -u {0}", UserName), AllowSpew: false, WithClient: false);
 			if (P4Result != 0)
 			{
 				throw new AutomationException("p4 clients -u {0} failed.", UserName);
@@ -1700,7 +1717,7 @@ namespace AutomationTool
         public void DeleteClient(string Name, bool Force = false)
         {
             CheckP4Enabled();
-            LogP4(String.Format("client -d {0} {1}", (Force ? "-f" : ""), Name));
+            LogP4(String.Format("client -d {0} {1}", (Force ? "-f" : ""), Name), WithClient: false);
         }
 
         /// <summary>
@@ -1723,7 +1740,7 @@ namespace AutomationTool
                 SpecInput += "\t" + Mapping.Key + " //" + ClientSpec.Name + Mapping.Value + Environment.NewLine;
             }
 			CommandUtils.Log(SpecInput);
-            LogP4("client -i", SpecInput);
+            LogP4("client -i", SpecInput, WithClient: false);
             return ClientSpec;
         }
 
diff --git a/Engine/Source/Programs/AutomationTool/Scripts/Tests.Automation.cs b/Engine/Source/Programs/AutomationTool/Scripts/Tests.Automation.cs
--- a/Engine/Source/Programs/AutomationTool/Scripts/Tests.Automation.cs
+++ b/Engine/Source/Programs/AutomationTool/Scripts/Tests.Automation.cs
@@ -30,6 +30,182 @@ class TestP4_Info : BuildCommand
 	}
 }
 
+[Help("GitPullRequest")]
+[Help("Dir=", "Directory of the Git repo.")]
+[Help("PR=", "PR # to shelve, can use a range PR=5-25")]
+[RequireP4]
+class GitPullRequest : BuildCommand
+{
+    string FindExeFromPath(string ExeName, string ExpectedPathSubstring = null)
+    {
+        if (File.Exists(ExeName))
+        {
+            return Path.GetFullPath(ExeName);
+        }
+
+        foreach (string BasePath in Environment.GetEnvironmentVariable("PATH").Split(';'))
+        {
+            var FullPath = Path.Combine(BasePath, ExeName);
+            if (ExpectedPathSubstring == null || FullPath.IndexOf(ExpectedPathSubstring, StringComparison.InvariantCultureIgnoreCase) != -1)
+            {
+                if (File.Exists(FullPath))
+                {
+                    return FullPath;
+                }
+            }
+        }
+
+        return null;
+    }
+    string RunGit(string GitCommandLine)
+    {
+        string GitExePath = FindExeFromPath("Git.exe", "PortableGit");
+        if (GitExePath == null)
+        {
+            throw new AutomationException("Unable to find Git.exe in the system path under a 'PortableGit' subdirectory.  Make sure the GitHub client is installed, and you are running this script from within a GitHub Command Shell.  We want to make sure we're using the correct version of Git, in case multiple versions are on the computer, which is why we check for a PortableGit folder that the GitHub Shell uses.");
+        }
+
+        Log("Running {0} {1}", GitExePath, GitCommandLine);
+
+        ProcessResult Result = Run(App: GitExePath, CommandLine: GitCommandLine, Options: (ERunOptions.NoLoggingOfRunCommand | ERunOptions.AllowSpew | ERunOptions.AppMustExist));
+        if (Result > 0 || Result < 0)
+        {
+            throw new AutomationException(String.Format("Command failed (Result:{2}): {0} {1}", GitExePath, GitCommandLine, Result.ExitCode));
+        }
+
+        // Return the results (sans leading and trailing whitespace)
+        return Result.Output.Trim();
+    }
+
+    void ExecuteInner(string Dir, int PR)
+    {
+        string PRNum = PR.ToString();
+        RunGit("reset --hard");
+        RunGit("checkout master");
+        var Refs = RunGit("show-ref");
+        if (!Refs.Contains("refs/remotes/origin/pr/" + PRNum))
+        {
+            throw new AutomationException("This is not among the refs: refs/remotes/origin/pr/{0}", PRNum);
+        }
+        RunGit(String.Format("fetch origin refs/pull/{0}/head:pr/{1}", PRNum, PRNum));
+        RunGit(String.Format("checkout pr/{0} --", PRNum));
+        var Base = RunGit(String.Format("log --author=TimSweeney --author=UnrealBot -1 pr/{0} --", PRNum));
+        string LookFor = "Engine source (";
+        if (!Base.Contains(LookFor))
+        {
+            throw new AutomationException("Unrecognized commit.");
+        }
+        Base = Base.Substring(Base.IndexOf(LookFor) + LookFor.Length);
+        string Depot = null;
+        if (Base.StartsWith("4."))
+        {
+            Depot = "//depot/UE4-Releases/4." + Base.Substring(2, 1);
+        }
+        else if (Base.StartsWith("Main"))
+        {
+            Depot = "//depot/UE4";
+        }
+        else
+        {
+            throw new AutomationException("Unrecognized branch.");
+        }
+        Log("Depot {0}", Depot);
+        if (!Base.Contains(")"))
+        {
+            throw new AutomationException("Unrecognized commit2.");
+        }
+        Base = Base.Substring(0, Base.IndexOf(")"));
+        if (!Base.Contains(" "))
+        {
+            throw new AutomationException("Unrecognized commit3.");
+        }
+        Base = Base.Substring(Base.LastIndexOf(" "));
+        Log("CL String {0}", Base);
+        int CL = int.Parse(Base);
+        Log("CL int {0}", CL);
+        if (CL < 2000000)
+        {
+            throw new AutomationException("Unrecognized commit3.");
+        }
+
+        P4ClientInfo NewClient = P4.GetClientInfo(P4Env.Client);
+
+        foreach (var p in NewClient.View)
+        {
+            Log("{0} = {1}", p.Key, p.Value);
+        }
+
+        string TestClient = P4Env.User + "_" + Environment.MachineName + "_PullRequests_" + Depot.Replace("/", "_");
+        NewClient.Owner = P4Env.User;
+        NewClient.Host = Environment.MachineName;
+        NewClient.RootPath = Dir;
+        NewClient.Name = TestClient;
+        NewClient.View = new List<KeyValuePair<string, string>>();
+        NewClient.View.Add(new KeyValuePair<string, string>(Depot + "/...", "/..."));
+        if (!P4.DoesClientExist(TestClient))
+        {
+            P4.CreateClient(NewClient);
+        }
+
+        var P4Sub = new P4Connection(P4Env.User, TestClient, P4Env.P4Port);
+
+        P4Sub.Sync(String.Format("-f -k -q {0}/...@{1}", Depot, CL));
+
+        var Change = P4Sub.CreateChange(null, String.Format("PullRequest number {0}", PRNum));
+        P4Sub.ReconcileNoDeletes(Change, CombinePaths(Dir, "Engine", "..."));
+        P4Sub.Shelve(Change);
+        P4Sub.Revert(Change, "-k //...");
+    }
+
+    public override void ExecuteBuild()
+    {
+        var Dir = ParseParamValue("Dir", @"D:\GitHub\UnrealEngine");
+        if (!FileExists_NoExceptions(CombinePaths(Dir, ".gitignore")))
+        {
+            throw new AutomationException("Dir {0} does not exist or does not like like a repo", Dir);
+        }
+
+        var PRNum = ParseParamValue("PR");
+        if (String.IsNullOrEmpty(PRNum))
+        {
+            throw new AutomationException("Must Provide PR arg, the number of the PR, or a range.");
+        }
+        int PRMin;
+        int PRMax;
+
+        if (PRNum.Contains("-"))
+        {
+            var Nums = PRNum.Split("-".ToCharArray());
+            PRMin = int.Parse(Nums[0]);
+            PRMax = int.Parse(Nums[1]);
+        }
+        else
+        {
+            PRMin = int.Parse(PRNum);
+            PRMax = PRMin;
+        }
+        var Failures = new List<int>();
+        PushDir(Dir);
+        for (int PR = PRMin; PR <= PRMax; PR++)
+        {
+            try
+            {
+                ExecuteInner(Dir, PR);
+            }
+            catch(Exception)
+            {
+                Failures.Add(PR);
+            }
+        }
+        PopDir();
+
+        foreach (var Failed in Failures)
+        {
+            Log("FAILED! PR {0}", Failed);
+        }
+    }
+}
+
 [Help("Throws an automation exception.")]
 class TestFail : BuildCommand
 {
