#ttp LOCALIZATION: EDITOR: Packaging log displays Unicode characters as ?

#change The engine now accepts a -UTF8Output command line argument which allows the output stream to be in the correct mode.
#change UAT accepts the same UTF8Output parameter to allow the connecting streams to be in the correct mode too (though it's currently only passed during the UnrealPak invocation).
#add FPlatformMisc::SetUTF8Output added to allow the opting-in of a UTF-8 output stream in Windows.

diff --git a/Engine/Source/Editor/MainFrame/Private/Frame/MainFrameActions.cpp b/Engine/Source/Editor/MainFrame/Private/Frame/MainFrameActions.cpp
--- a/Engine/Source/Editor/MainFrame/Private/Frame/MainFrameActions.cpp
+++ b/Engine/Source/Editor/MainFrame/Private/Frame/MainFrameActions.cpp
@@ -598,7 +598,7 @@ void FMainFrameActionCallbacks::PackageProject( const FString InPlatformName, co
 	}
 
 	FString ProjectPath = FPaths::IsProjectFilePathSet() ? FPaths::ConvertRelativePathToFull(FPaths::GetProjectFilePath()) : FPaths::RootDir() / FApp::GetGameName() / FApp::GetGameName() + TEXT(".uproject");
-	FString CommandLine = FString::Printf(TEXT("BuildCookRun %s%s -nop4 -project=\"%s\" -cook -allmaps -stage -archive -archivedirectory=\"%s\" -package -%s -clientconfig=%s -ue4exe=%s %s"),
+	FString CommandLine = FString::Printf(TEXT("BuildCookRun %s%s -nop4 -project=\"%s\" -cook -allmaps -stage -archive -archivedirectory=\"%s\" -package -%s -clientconfig=%s -ue4exe=%s %s -utf8output"),
 		FRocketSupport::IsRocket() ? TEXT( "-rocket -nocompile" ) : TEXT( "-nocompileeditor" ),
 		FApp::IsEngineInstalled() ? TEXT(" -installed") : TEXT(""),
 		*ProjectPath,
diff --git a/Engine/Source/Programs/AutomationTool/Automation.cs b/Engine/Source/Programs/AutomationTool/Automation.cs
--- a/Engine/Source/Programs/AutomationTool/Automation.cs
+++ b/Engine/Source/Programs/AutomationTool/Automation.cs
@@ -110,6 +110,7 @@ namespace AutomationTool
 		public static CommandLineArg Rocket = new CommandLineArg("-Rocket");
 		public static CommandLineArg NoKill = new CommandLineArg("-NoKill");
 		public static CommandLineArg Installed = new CommandLineArg("-Installed");
+		public static CommandLineArg UTF8Output = new CommandLineArg("-UTF8Output");
 
 		/// <summary>
 		/// Force initialize static members by calling this.
diff --git a/Engine/Source/Programs/AutomationTool/Log.cs b/Engine/Source/Programs/AutomationTool/Log.cs
--- a/Engine/Source/Programs/AutomationTool/Log.cs
+++ b/Engine/Source/Programs/AutomationTool/Log.cs
@@ -171,6 +171,15 @@ namespace AutomationTool
 	/// </summary>
 	class AutomationConsoleTraceListener : TraceListener
 	{
+		public AutomationConsoleTraceListener()
+		{
+			// We use the command line directly here, because GlobalCommandLine isn't initialized early enough
+			if (SharedUtils.ParseCommandLine().Any(Arg => Arg.ToLower() == "-utf8output"))
+			{
+				Console.OutputEncoding = new System.Text.UTF8Encoding(false, false);
+			}
+		}
+
 		/// <summary>
 		/// Writes a formatted line to the console.
 		/// </summary>
diff --git a/Engine/Source/Programs/AutomationTool/ProcessUtils.cs b/Engine/Source/Programs/AutomationTool/ProcessUtils.cs
--- a/Engine/Source/Programs/AutomationTool/ProcessUtils.cs
+++ b/Engine/Source/Programs/AutomationTool/ProcessUtils.cs
@@ -528,6 +528,7 @@ namespace AutomationTool
 			NoWaitForExit = 1 << 2,
 			NoStdOutRedirect = 1 << 3,
             NoLoggingOfRunCommand = 1 << 4,
+            UTF8Output = 1 << 5,
 
 			Default = AllowSpew | AppMustExist,
 		}
@@ -578,6 +579,10 @@ namespace AutomationTool
 			}
 			Proc.StartInfo.RedirectStandardInput = Input != null;
 			Proc.StartInfo.CreateNoWindow = true;
+			if ((Options & ERunOptions.UTF8Output) == ERunOptions.UTF8Output)
+			{
+				Proc.StartInfo.StandardOutputEncoding = new System.Text.UTF8Encoding(false, false);
+			}
 			Proc.Start();
 
 			if (bRedirectStdOut)
@@ -641,9 +646,11 @@ namespace AutomationTool
 		/// <param name="App">Executable to run</param>
 		/// <param name="CommandLine">Commandline to pass on to the executable</param>
 		/// <param name="LogName">Name of the logfile ( if null, executable name is used )</param>
-		public static void RunAndLog(CommandEnvironment Env, string App, string CommandLine, string LogName = null, int MaxSuccessCode = 0)
+		/// <param name="Input">Optional Input for the program (will be provided as stdin)</param>
+		/// <param name="Options">Defines the options how to run. See ERunOptions.</param>
+		public static void RunAndLog(CommandEnvironment Env, string App, string CommandLine, string LogName = null, int MaxSuccessCode = 0, string Input = null, ERunOptions Options = ERunOptions.Default)
 		{
-            RunAndLog(App, CommandLine, GetRunAndLogLogName(Env, App, LogName), MaxSuccessCode);
+            RunAndLog(App, CommandLine, GetRunAndLogLogName(Env, App, LogName), MaxSuccessCode, Input, Options);
 		}
         /// <summary>
         /// Runs external program and writes the output to a logfile.
@@ -651,12 +658,13 @@ namespace AutomationTool
         /// <param name="App">Executable to run</param>
         /// <param name="CommandLine">Commandline to pass on to the executable</param>
         /// <param name="Logfile">Full path to the logfile, where the application output should be written to.</param>
-        /// <returns>Whether the program executed successfully or not.</returns>
-        public static void RunAndLog(string App, string CommandLine, string Logfile = null, int MaxSuccessCode = 0)
+		/// <param name="Input">Optional Input for the program (will be provided as stdin)</param>
+		/// <param name="Options">Defines the options how to run. See ERunOptions.</param>
+        public static void RunAndLog(string App, string CommandLine, string Logfile = null, int MaxSuccessCode = 0, string Input = null, ERunOptions Options = ERunOptions.Default)
         {
             Log("Running {0} {1}", App, CommandLine);
 
-            ProcessResult Result = Run(App, CommandLine);
+            ProcessResult Result = Run(App, CommandLine, Input, Options);
             if (Result.Output.Length > 0 && Logfile != null)
             {
                 WriteToFile(Logfile, Result.Output);
diff --git a/Engine/Source/Programs/AutomationTool/Scripts/CopyBuildToStagingDirectory.Automation.cs b/Engine/Source/Programs/AutomationTool/Scripts/CopyBuildToStagingDirectory.Automation.cs
--- a/Engine/Source/Programs/AutomationTool/Scripts/CopyBuildToStagingDirectory.Automation.cs
+++ b/Engine/Source/Programs/AutomationTool/Scripts/CopyBuildToStagingDirectory.Automation.cs
@@ -71,7 +71,11 @@ public partial class Project : CommandUtils
 			CmdLine += " -installed";
 		}
 		CmdLine += " -order=" + CommandUtils.MakePathSafeToUseWithCommandLine(PakOrderFileLocation);
-		RunAndLog(CmdEnv, UnrealPakExe, CmdLine);
+		if (GlobalCommandLine.UTF8Output)
+		{
+			CmdLine += " -UTF8Output";
+		}
+		RunAndLog(CmdEnv, UnrealPakExe, CmdLine, Options: ERunOptions.Default | ERunOptions.UTF8Output);
 		Log("UnrealPak Done *******");
 	}
 
diff --git a/Engine/Source/Runtime/Core/Private/GenericPlatform/GenericPlatformMisc.cpp b/Engine/Source/Runtime/Core/Private/GenericPlatform/GenericPlatformMisc.cpp
--- a/Engine/Source/Runtime/Core/Private/GenericPlatform/GenericPlatformMisc.cpp
+++ b/Engine/Source/Runtime/Core/Private/GenericPlatform/GenericPlatformMisc.cpp
@@ -222,6 +222,11 @@ void FGenericPlatformMisc::LowLevelOutputDebugStringf(const TCHAR *Fmt, ... )
 	);
 }
 
+void FGenericPlatformMisc::SetUTF8Output()
+{
+	// assume that UTF-8 is possible by default anyway
+}
+
 void FGenericPlatformMisc::LocalPrint( const TCHAR* Str )
 {
 #if PLATFORM_USE_LS_SPEC_FOR_WIDECHAR
diff --git a/Engine/Source/Runtime/Core/Private/Windows/WindowsPlatformMisc.cpp b/Engine/Source/Runtime/Core/Private/Windows/WindowsPlatformMisc.cpp
--- a/Engine/Source/Runtime/Core/Private/Windows/WindowsPlatformMisc.cpp
+++ b/Engine/Source/Runtime/Core/Private/Windows/WindowsPlatformMisc.cpp
@@ -36,7 +36,8 @@
 	#pragma comment(lib, "psapi.lib")
 #endif
 
-
+#include <fcntl.h>
+#include <io.h>
 
 /** 
  * Whether support for integrating into the firewall is there
@@ -714,6 +715,10 @@ uint32 FWindowsPlatformMisc::GetKeyMap( uint16* KeyCodes, FString* KeyNames, uin
 	return NumMappings;
 }
 
+void FWindowsPlatformMisc::SetUTF8Output()
+{
+	_setmode(_fileno(stdout), _O_U8TEXT);
+}
 
 void FWindowsPlatformMisc::LocalPrint( const TCHAR *Message )
 {
diff --git a/Engine/Source/Runtime/Core/Public/GenericPlatform/GenericPlatformMisc.h b/Engine/Source/Runtime/Core/Public/GenericPlatform/GenericPlatformMisc.h
--- a/Engine/Source/Runtime/Core/Public/GenericPlatform/GenericPlatformMisc.h
+++ b/Engine/Source/Runtime/Core/Public/GenericPlatform/GenericPlatformMisc.h
@@ -419,6 +419,9 @@ public:
 	static void LowLevelOutputDebugString(const TCHAR *Message);
 	static void VARARGS LowLevelOutputDebugStringf(const TCHAR *Format, ... );
 
+	/** Sets the default output to UTF8 */
+	static void SetUTF8Output();
+
 	/** Prints string to the default output */
 	static void LocalPrint( const TCHAR* Str );
 
diff --git a/Engine/Source/Runtime/Core/Public/Windows/WindowsPlatformMisc.h b/Engine/Source/Runtime/Core/Public/Windows/WindowsPlatformMisc.h
--- a/Engine/Source/Runtime/Core/Public/Windows/WindowsPlatformMisc.h
+++ b/Engine/Source/Runtime/Core/Public/Windows/WindowsPlatformMisc.h
@@ -38,6 +38,7 @@ struct CORE_API FWindowsPlatformMisc : public FGenericPlatformMisc
 	static void PumpMessages(bool bFromMainLoop);
 	static uint32 GetKeyMap( uint16* KeyCodes, FString* KeyNames, uint32 MaxMappings );
 	static uint32 GetCharKeyMap(uint16* KeyCodes, FString* KeyNames, uint32 MaxMappings);
+	static void SetUTF8Output();
 	static void LocalPrint(const TCHAR *Message);
 	static void RequestExit(bool Force);
 	static const TCHAR* GetSystemErrorMessage(TCHAR* OutBuffer, int32 BufferCount, int32 Error);
diff --git a/Engine/Source/Runtime/Launch/Private/LaunchEngineLoop.cpp b/Engine/Source/Runtime/Launch/Private/LaunchEngineLoop.cpp
--- a/Engine/Source/Runtime/Launch/Private/LaunchEngineLoop.cpp
+++ b/Engine/Source/Runtime/Launch/Private/LaunchEngineLoop.cpp
@@ -528,6 +528,11 @@ int32 FEngineLoop::PreInit( const TCHAR* CmdLine )
 {
 	DECLARE_SCOPE_CYCLE_COUNTER(TEXT("Engine Pre-Initialized"), STAT_PreInit, STATGROUP_LoadTime);
 
+	if (FParse::Param(CmdLine, TEXT("UTF8Output")))
+	{
+		FPlatformMisc::SetUTF8Output();
+	}
+
 	// Switch into executable's directory.
 	FPlatformProcess::SetCurrentWorkingDirectoryToBaseDir();
 
