build script, fix parsing of commit comment

diff --git a/Engine/Source/Programs/AutomationTool/Scripts/Tests.Automation.cs b/Engine/Source/Programs/AutomationTool/Scripts/Tests.Automation.cs
--- a/Engine/Source/Programs/AutomationTool/Scripts/Tests.Automation.cs
+++ b/Engine/Source/Programs/AutomationTool/Scripts/Tests.Automation.cs
@@ -90,7 +90,7 @@ class GitPullRequest : BuildCommand
         RunGit(String.Format("fetch origin refs/pull/{0}/head:pr/{1}", PRNum, PRNum));
         RunGit(String.Format("checkout pr/{0} --", PRNum));
         var Base = RunGit(String.Format("log --author=TimSweeney --author=UnrealBot -1 pr/{0} --", PRNum));
-        string LookFor = "Engine source (";
+        string LookFor = "(";
         if (!Base.Contains(LookFor))
         {
             throw new AutomationException("Unrecognized commit.");
@@ -184,7 +184,7 @@ class GitPullRequest : BuildCommand
             PRMin = int.Parse(PRNum);
             PRMax = PRMin;
         }
-        var Failures = new List<int>();
+        var Failures = new List<string>();
         PushDir(Dir);
         for (int PR = PRMin; PR <= PRMax; PR++)
         {
@@ -192,16 +192,17 @@ class GitPullRequest : BuildCommand
             {
                 ExecuteInner(Dir, PR);
             }
-            catch(Exception)
+            catch(Exception Ex)
             {
-                Failures.Add(PR);
+                Log(" Exception was {0}", LogUtils.FormatException(Ex));
+                Failures.Add(String.Format("PR {0} Failed with {1}", PR, LogUtils.FormatException(Ex)));
             }
         }
         PopDir();
 
         foreach (var Failed in Failures)
         {
-            Log("FAILED! PR {0}", Failed);
+            Log("{0}", Failed);
         }
     }
 }
