build script, increase wiat time between submit retries.

diff --git a/Engine/Source/Programs/AutomationTool/P4Utils.cs b/Engine/Source/Programs/AutomationTool/P4Utils.cs
--- a/Engine/Source/Programs/AutomationTool/P4Utils.cs
+++ b/Engine/Source/Programs/AutomationTool/P4Utils.cs
@@ -890,7 +890,7 @@ namespace AutomationTool
 					{
 						CommandUtils.Log(TraceEventType.Information, "Change {0} failed to submit for reasons we do not recognize.\n{1}\nWaiting and retrying.", CL, CmdOutput);
 					}
-					System.Threading.Thread.Sleep(15000);
+					System.Threading.Thread.Sleep(30000);
 				}
 				else
 				{
