build script, add causers to email

diff --git a/Engine/Source/Programs/AutomationTool/Scripts/GUBP.Automation.cs b/Engine/Source/Programs/AutomationTool/Scripts/GUBP.Automation.cs
--- a/Engine/Source/Programs/AutomationTool/Scripts/GUBP.Automation.cs
+++ b/Engine/Source/Programs/AutomationTool/Scripts/GUBP.Automation.cs
@@ -4877,6 +4877,38 @@
                 if (GUBPNodes[NodeToDo].RunInEC() && !NodeIsAlreadyComplete(NodeToDo, LocalOnly)) // if something is already finished, we don't put it into EC  
                 {
                     string EMails = "";
+                    string FailCauserEMails = "";
+                    if (GUBPNodesHistory.ContainsKey(NodeToDo))
                     {
+                        var History = GUBPNodesHistory[NodeToDo];
+
+                        ECProps.Add(string.Format("LastGreen/{0}={1}", NodeToDo, History.LastSucceeded));
+                        ECProps.Add(string.Format("RedsSince/{0}={1}", NodeToDo, History.FailedString));
+                        ECProps.Add(string.Format("InProgress/{0}={1}", NodeToDo, History.InProgressString));
+
+
+                        if (History.LastSucceeded > 0 && History.LastSucceeded < P4Env.Changelist)
+                        {
+                            var ChangeRecords = GetChanges(History.LastSucceeded, P4Env.Changelist, History.LastSucceeded);
+                            int NumChanges = 0;
+                            foreach (var Record in ChangeRecords)
+                            {
+                                FailCauserEMails = GUBPNode.MergeSpaceStrings(FailCauserEMails, Record.UserEmail);
+                                if (++NumChanges > 50)
+                                {
+                                    FailCauserEMails = "";
+                                    break;
+                                }
+                            }
+                        }
+                        ECProps.Add(string.Format("FailCausers/{0}={1}", NodeToDo, FailCauserEMails));
+                    }
+                    else
+                    {
+                        ECProps.Add(string.Format("LastGreen/{0}=0", NodeToDo));
+                        ECProps.Add(string.Format("RedsSince/{0}=", NodeToDo));
+                        ECProps.Add(string.Format("InProgress/{0}=", NodeToDo));
+                    }
+                    {
                         string EmailOnly = ParseParamValue("EmailOnly");
                         if (bFake)
@@ -4894,6 +4926,10 @@
                         else
                         {
                             EMails = GetEMailListForNode(this, NodeToDo);
+                            if (ParseParam("CIS"))
+                            {
+                                EMails = GUBPNode.MergeSpaceStrings(FailCauserEMails, EMails);
+                            }
                         }
                         string AddEmails = ParseParamValue("AddEmails");
                         if (!String.IsNullOrEmpty(AddEmails))
@@ -4920,40 +4956,6 @@
                     ECProps.Add(string.Format("AgentRequirementString/{0}={1}", NodeToDo, GUBPNodes[NodeToDo].ECAgentString()));
                     ECProps.Add(string.Format("RequiredMemory/{0}={1}", NodeToDo, GUBPNodes[NodeToDo].AgentMemoryRequirement(this)));
                     ECProps.Add(string.Format("Timeouts/{0}={1}", NodeToDo, GUBPNodes[NodeToDo].TimeoutInMinutes()));
-                    if (GUBPNodesHistory.ContainsKey(NodeToDo))
-                    {
-                        var History = GUBPNodesHistory[NodeToDo];
 
-                        ECProps.Add(string.Format("LastGreen/{0}={1}", NodeToDo, History.LastSucceeded));
-                        ECProps.Add(string.Format("RedsSince/{0}={1}", NodeToDo, History.FailedString));
-                        ECProps.Add(string.Format("InProgress/{0}={1}", NodeToDo, History.InProgressString));
-
-                        string FailCauserEMails = "";
-
-                        if (History.LastSucceeded > 0 && History.LastSucceeded < P4Env.Changelist)
-                        {
-                            var ChangeRecords = GetChanges(History.LastSucceeded, P4Env.Changelist, History.LastSucceeded);
-                            int NumChanges = 0;
-                            foreach (var Record in ChangeRecords)
-                            {
-                                FailCauserEMails = GUBPNode.MergeSpaceStrings(FailCauserEMails, Record.UserEmail);
-                                if (++NumChanges > 50)
-                                {
-                                    FailCauserEMails = "";
-                                    break;
-                                }
-                            }
-                        }
-                        ECProps.Add(string.Format("FailCausers/{0}={1}", NodeToDo, FailCauserEMails));
-                    }
-                    else
-                    {
-                        ECProps.Add(string.Format("LastGreen/{0}=0", NodeToDo));
-                        ECProps.Add(string.Format("RedsSince/{0}=", NodeToDo));
-                        ECProps.Add(string.Format("InProgress/{0}=", NodeToDo));
-                    }
-
-
-
                     bool Sticky = GUBPNodes[NodeToDo].IsSticky();
                     bool DoParallel = !Sticky;
