Updated GitPullRequest script
- Now automatically creates GitRepo for you
     - The default folder is "/Engine/Intermediate/PullRequestGitRepo/" (you can override with the '-Dir=<folder>' option)
- Added optional '-Clean' argument to wipe any existing temp Git repo

diff --git a/Engine/Source/Programs/AutomationTool/Scripts/Tests.Automation.cs b/Engine/Source/Programs/AutomationTool/Scripts/Tests.Automation.cs
--- a/Engine/Source/Programs/AutomationTool/Scripts/Tests.Automation.cs
+++ b/Engine/Source/Programs/AutomationTool/Scripts/Tests.Automation.cs
@@ -36,6 +36,9 @@ class TestP4_Info : BuildCommand
 [RequireP4]
 class GitPullRequest : BuildCommand
 {
+	/// URL to our UnrealEngine repository on GitHub
+	static readonly string GitRepositoryURL = "https://github.com/EpicGames/UnrealEngine.git";
+
     string FindExeFromPath(string ExeName, string ExpectedPathSubstring = null)
     {
         if (File.Exists(ExeName))
@@ -80,8 +83,11 @@ class GitPullRequest : BuildCommand
     void ExecuteInner(string Dir, int PR)
     {
         string PRNum = PR.ToString();
+
+		// Discard any old changes we may have committed
         RunGit("reset --hard");
-        RunGit("checkout master");
+
+		// show-ref is just a double check that the PR exists
         var Refs = RunGit("show-ref");
         if (!Refs.Contains("refs/remotes/origin/pr/" + PRNum))
         {
@@ -89,7 +95,10 @@ class GitPullRequest : BuildCommand
         }
         RunGit(String.Format("fetch origin refs/pull/{0}/head:pr/{1}", PRNum, PRNum));
         RunGit(String.Format("checkout pr/{0} --", PRNum));
-        var Base = RunGit(String.Format("log --author=TimSweeney --author=UnrealBot -1 pr/{0} --", PRNum));
+
+		// after the fetch we do git log --Author... to figure out the P4 branch and CL
+		// the -1 limits it to the first one with the right author
+		var Base = RunGit(String.Format("log --author=TimSweeney --author=UnrealBot -1 pr/{0} --", PRNum));
         string LookFor = "(";
         if (!Base.Contains(LookFor))
         {
@@ -159,10 +168,11 @@ class GitPullRequest : BuildCommand
 
     public override void ExecuteBuild()
     {
-        var Dir = ParseParamValue("Dir", @"D:\GitHub\UnrealEngine");
-        if (!FileExists_NoExceptions(CombinePaths(Dir, ".gitignore")))
+        var Dir = ParseParamValue("Dir");
+        if (String.IsNullOrEmpty(Dir))
         {
-            throw new AutomationException("Dir {0} does not exist or does not like like a repo", Dir);
+			// No Git repo directory was specified, so we'll choose a directory automatically
+			Dir = Path.GetFullPath( Path.Combine( BuildConfiguration.RelativeEnginePath, "Intermediate", "PullRequestGitRepo" ) );
         }
 
         var PRNum = ParseParamValue("PR");
@@ -185,7 +195,64 @@ class GitPullRequest : BuildCommand
             PRMax = PRMin;
         }
         var Failures = new List<string>();
-        PushDir(Dir);
+
+
+		// Setup Git repo
+		{
+			if( ParseParam( "Clean" ) )
+			{ 
+				Console.WriteLine( "Cleaning temporary Git repository folder... " );
+
+				// Wipe the Git repo directory
+				if( !InternalUtils.SafeDeleteDirectory( Dir ) )
+				{
+					throw new AutomationException("Unable to clean out temporary Git repo directory: " + Dir );
+				}
+			}
+
+			// Change directory to the Git repository
+			bool bRepoAlreadyExists = InternalUtils.SafeDirectoryExists( Dir );
+			if( !bRepoAlreadyExists )
+			{ 
+				InternalUtils.SafeCreateDirectory( Dir );
+			}
+			PushDir( Dir );
+
+			if( !bRepoAlreadyExists )
+			{ 
+				// Don't init Git if we didn't clean, because the old repo is probably still there.
+
+				// Create the Git repository
+				RunGit( "init" );
+			}
+
+			// Make sure that creating the repository worked OK
+			RunGit( "status" );
+
+			// Check to see if we already have a remote origin setup
+			{ 
+				string Result = RunGit( "remote -v" );
+				if( Result == "origin" )
+				{
+					// OK, we already have an origin but no remote is associated with it.  We'll do that now.
+					RunGit( "remote set-url origin " + GitRepositoryURL );
+				}
+				else if( Result.IndexOf( GitRepositoryURL, StringComparison.InvariantCultureIgnoreCase ) != -1 )
+				{
+					// Origin is already setup!  Nothing to do.
+				}
+				else
+				{
+					// No origin is set, so let's add it!
+					RunGit( "remote add origin " + GitRepositoryURL );
+				}
+			}
+
+			// Fetch all of the remote branches/tags into our local index.  This is needed so that we can figure out
+			// which branches exist already.
+			RunGit( "fetch" );
+		}
+
         for (int PR = PRMin; PR <= PRMax; PR++)
         {
             try
