build script, success nodes don't have causers

diff --git a/Engine/Source/Programs/AutomationTool/Scripts/GUBP.Automation.cs b/Engine/Source/Programs/AutomationTool/Scripts/GUBP.Automation.cs
--- a/Engine/Source/Programs/AutomationTool/Scripts/GUBP.Automation.cs
+++ b/Engine/Source/Programs/AutomationTool/Scripts/GUBP.Automation.cs
@@ -4926,7 +4926,7 @@ if (HostPlatform == UnrealTargetPlatform.Mac) continue; //temp hack till mac aut
                         else
                         {
                             EMails = GetEMailListForNode(this, NodeToDo);
-                            if (ParseParam("CIS"))
+                            if (ParseParam("CIS") && !GUBPNodes[NodeToDo].SendSuccessEmail())
                             {
                                 EMails = GUBPNode.MergeSpaceStrings(FailCauserEMails, EMails);
                             }
