GitPullRequest: Fixed bad base directory; renamed changelist

diff --git a/Engine/Source/Programs/AutomationTool/Scripts/Tests.Automation.cs b/Engine/Source/Programs/AutomationTool/Scripts/Tests.Automation.cs
--- a/Engine/Source/Programs/AutomationTool/Scripts/Tests.Automation.cs
+++ b/Engine/Source/Programs/AutomationTool/Scripts/Tests.Automation.cs
@@ -160,7 +160,7 @@ class GitPullRequest : BuildCommand
 
         P4Sub.Sync(String.Format("-f -k -q {0}/...@{1}", Depot, CL));
 
-        var Change = P4Sub.CreateChange(null, String.Format("PullRequest number {0}", PRNum));
+        var Change = P4Sub.CreateChange(null, String.Format("GitHub pull request #{0}", PRNum));
         P4Sub.ReconcileNoDeletes(Change, CombinePaths(Dir, "Engine", "..."));
         P4Sub.Shelve(Change);
         P4Sub.Revert(Change, "-k //...");
@@ -172,7 +172,7 @@ class GitPullRequest : BuildCommand
         if (String.IsNullOrEmpty(Dir))
         {
 			// No Git repo directory was specified, so we'll choose a directory automatically
-			Dir = Path.GetFullPath( Path.Combine( BuildConfiguration.RelativeEnginePath, "Intermediate", "PullRequestGitRepo" ) );
+			Dir = Path.GetFullPath( Path.Combine( CmdEnv.LocalRoot, "Engine", "Intermediate", "PullRequestGitRepo" ) );
         }
 
         var PRNum = ParseParamValue("PR");
