support for codeless formal builds

diff --git a/Engine/Source/Programs/AutomationTool/Scripts/GUBP.Automation.cs b/Engine/Source/Programs/AutomationTool/Scripts/GUBP.Automation.cs
--- a/Engine/Source/Programs/AutomationTool/Scripts/GUBP.Automation.cs
+++ b/Engine/Source/Programs/AutomationTool/Scripts/GUBP.Automation.cs
@@ -2300,6 +2300,10 @@ public class GUBP : BuildCommand
         {
             return base.Priority() + 20.0f;
         }
+        public override int CISFrequencyQuantumShift(GUBP bp)
+        {
+            return base.CISFrequencyQuantumShift(bp) + 3;
+        }
 
         public override void DoBuild(GUBP bp)
         {
@@ -3957,6 +3961,7 @@ public class GUBP : BuildCommand
                 var AgentSharingGroup = "Shared_EditorTests" + HostPlatformNode.StaticGetHostPlatformSuffix(HostPlatform);
 
                 Dictionary<string, List<UnrealTargetPlatform>> NonCodeProjectNames;
+                var NonCodeFormalBuilds = new Dictionary<string, List<KeyValuePair<UnrealTargetPlatform, UnrealTargetConfiguration>>>();
                 {
                     var Target = Branch.BaseEngineProject.Properties.Targets[TargetRules.TargetType.Editor];
                     NonCodeProjectNames = Target.Rules.GUBP_NonCodeProjects_BaseEditorTypeOnly(HostPlatform);
@@ -3980,6 +3985,26 @@ public class GUBP : BuildCommand
                     {
                         throw new AutomationException("We assume that if we have shared promotable, the base engine is in it.");
                     }
+
+                    var TempNonCodeFormalBuilds = Target.Rules.GUBP_NonCodeFormalBuilds_BaseEditorTypeOnly();
+                    var HostMonos = GetMonolithicPlatformsForUProject(UnrealTargetPlatform.Mac, Branch.BaseEngineProject, true);
+
+                    foreach (var Codeless in TempNonCodeFormalBuilds)
+                    {
+                        var PlatList = Codeless.Value;
+                        var NewPlatList = new List<KeyValuePair<UnrealTargetPlatform, UnrealTargetConfiguration>>();
+                        foreach (var PlatPair in PlatList)
+                        {
+                            if (HostMonos.Contains(PlatPair.Key))
+                            {
+                                NewPlatList.Add(PlatPair);
+                            }
+                        }
+                        if (NewPlatList.Count > 0)
+                        {
+                            NonCodeFormalBuilds.Add(Codeless.Key, NewPlatList);
+                        }
+                    }
                 }
                 if (bAutomatedTesting)
                 {
@@ -4117,26 +4142,51 @@ public class GUBP : BuildCommand
                                     if (!GUBPNodes.ContainsKey(GamePlatformCookedAndCompiledNode.StaticGetFullName(HostPlatform, NonCodeProject, Plat)))
                                     {
                                         AddNode(new GamePlatformCookedAndCompiledNode(this, HostPlatform, NonCodeProject, Plat, false));
+                                        if (NonCodeFormalBuilds.ContainsKey(NonCodeProject.GameName))
+                                        {
+                                            var PlatList = NonCodeFormalBuilds[NonCodeProject.GameName];
+                                            foreach (var PlatPair in PlatList)
+                                            {
+                                                if (PlatPair.Key == Plat)
+                                                {
+                                                    AddNode(new FormalBuildNode(this, NonCodeProject, HostPlatform, new List<UnrealTargetPlatform>() { Plat }, new List<UnrealTargetConfiguration>() {PlatPair.Value}));
+                                                    // we don't want this delayed
+                                                    // this would normally wait for the testing phase, we just want to build it right away
+                                                    RemovePseudodependencyFromNode(
+                                                        CookNode.StaticGetFullName(HostPlatform, NonCodeProject, CookedPlatform),
+                                                        WaitForTestShared.StaticGetFullName());
+                                                    RemovePseudodependencyFromNode(
+                                                        CookNode.StaticGetFullName(HostPlatform, NonCodeProject, CookedPlatform),
+                                                        CookNode.StaticGetFullName(HostPlatform, Branch.BaseEngineProject, CookedPlatform));
+
+                                                    string BuildAgentSharingGroup = NonCodeProject.GameName + "_MakeFormalBuild_" + Plat.ToString();
+                                                    GUBPNodes[CookNode.StaticGetFullName(HostPlatform, NonCodeProject, CookedPlatform)].AgentSharingGroup = BuildAgentSharingGroup;
+                                                    GUBPNodes[GamePlatformCookedAndCompiledNode.StaticGetFullName(HostPlatform, NonCodeProject, Plat)].AgentSharingGroup = BuildAgentSharingGroup;
+                                                    GUBPNodes[FormalBuildNode.StaticGetFullName(NonCodeProject, HostPlatform, new List<UnrealTargetPlatform>() { Plat })].AgentSharingGroup = BuildAgentSharingGroup;
+
+                                                }
+                                            }
+                                        }
                                     }
                                     if (bAutomatedTesting)
                                     {
-                                    if (HostPlatform == UnrealTargetPlatform.Mac) continue; //temp hack till mac automated testing works
-                                    var ThisMonoGameTestNodes = new List<string>();
-                                    foreach (var Test in GameTests)
-                                    {
-                                        var TestName = Test.Key + "_" + Plat.ToString();
-                                        ThisMonoGameTestNodes.Add(AddNode(new UATTestNode(this, HostPlatform, NonCodeProject, TestName, Test.Value, CookedAgentSharingGroup, false, RequiredPlatforms)));
-                                    }
-                                    if (ThisMonoGameTestNodes.Count > 0)
-                                    {
-                                        GameTestNodes.Add(AddNode(new GameAggregateNode(this, HostPlatform, NonCodeProject, "CookedTests_" + Plat.ToString() + "_" + Kind.ToString() + HostPlatformNode.StaticGetHostPlatformSuffix(HostPlatform), ThisMonoGameTestNodes, 0.0f)));
+                                        if (HostPlatform == UnrealTargetPlatform.Mac) continue; //temp hack till mac automated testing works
+                                        var ThisMonoGameTestNodes = new List<string>();
+                                        foreach (var Test in GameTests)
+                                        {
+                                            var TestName = Test.Key + "_" + Plat.ToString();
+                                            ThisMonoGameTestNodes.Add(AddNode(new UATTestNode(this, HostPlatform, NonCodeProject, TestName, Test.Value, CookedAgentSharingGroup, false, RequiredPlatforms)));
+                                        }
+                                        if (ThisMonoGameTestNodes.Count > 0)
+                                        {
+                                            GameTestNodes.Add(AddNode(new GameAggregateNode(this, HostPlatform, NonCodeProject, "CookedTests_" + Plat.ToString() + "_" + Kind.ToString() + HostPlatformNode.StaticGetHostPlatformSuffix(HostPlatform), ThisMonoGameTestNodes, 0.0f)));
+                                        }
                                     }
                                 }
                             }
                         }
                     }
                 }
-                }
 #if false
                 //for now, non-code projects don't do client or server.
                 foreach (var ServerPlatform in ServerPlatforms)
diff --git a/Engine/Source/Programs/UnrealBuildTool/System/RulesCompiler.cs b/Engine/Source/Programs/UnrealBuildTool/System/RulesCompiler.cs
--- a/Engine/Source/Programs/UnrealBuildTool/System/RulesCompiler.cs
+++ b/Engine/Source/Programs/UnrealBuildTool/System/RulesCompiler.cs
@@ -648,6 +648,14 @@ namespace UnrealBuildTool
             return new Dictionary<string, List<UnrealTargetPlatform>>();
         }
         /// <summary>
+        /// Return a list of the non-code projects to make formal builds for
+        /// </summary>
+        /// <returns>a list of the non-code projects to build cook and test</returns>
+        public virtual Dictionary<string, List<KeyValuePair<UnrealTargetPlatform, UnrealTargetConfiguration>>> GUBP_NonCodeFormalBuilds_BaseEditorTypeOnly()
+        {
+            return new Dictionary<string, List<KeyValuePair<UnrealTargetPlatform, UnrealTargetConfiguration>>>();
+        }
+        /// <summary>
         /// Return a list of "test name", "UAT command" pairs for testing the editor
         /// </summary>
         public virtual Dictionary<string, string> GUBP_GetEditorTests_EditorTypeOnly(UnrealTargetPlatform HostPlatform)
diff --git a/Engine/Source/UE4Editor.Target.cs b/Engine/Source/UE4Editor.Target.cs
--- a/Engine/Source/UE4Editor.Target.cs
+++ b/Engine/Source/UE4Editor.Target.cs
@@ -103,6 +103,12 @@ public class UE4EditorTarget : TargetRules
 
         return NonCodeProjectNames;
     }
+    public override Dictionary<string, List<KeyValuePair<UnrealTargetPlatform, UnrealTargetConfiguration>>> GUBP_NonCodeFormalBuilds_BaseEditorTypeOnly()
+    {
+        var NonCodeProjectNames = new Dictionary<string, List<KeyValuePair<UnrealTargetPlatform, UnrealTargetConfiguration>>>();
+        NonCodeProjectNames.Add("TappyChicken", new List<KeyValuePair<UnrealTargetPlatform, UnrealTargetConfiguration>>{new KeyValuePair<UnrealTargetPlatform, UnrealTargetConfiguration>(UnrealTargetPlatform.IOS, UnrealTargetConfiguration.Shipping)});
+        return NonCodeProjectNames;
+    }
     public override string GUBP_GetPromotionEMails_EditorTypeOnly(string Branch)
     {
         return "Engine-QA[epic] Engine-QA-Team[epic]";
