#ue4
#html5
* fix for putting in the proper js filename in to the generate web page for testing

diff --git a/Engine/Source/Programs/AutomationTool/HTML5/HTML5Platform.Automation.cs b/Engine/Source/Programs/AutomationTool/HTML5/HTML5Platform.Automation.cs
--- a/Engine/Source/Programs/AutomationTool/HTML5/HTML5Platform.Automation.cs
+++ b/Engine/Source/Programs/AutomationTool/HTML5/HTML5Platform.Automation.cs
@@ -80,7 +80,7 @@ public class HTML5Platform : Platform
 				{
                     if (IsContentOnly)
                         InGameName = "UE4Game";
-					LineStr = LineStr.Replace("%CONFIG%", (InGameConfiguration != "Development" ? (InGameName + "-" + InGameConfiguration) : InGameName));
+					LineStr = LineStr.Replace("%CONFIG%", (InGameConfiguration != "Development" ? (InGameName + "-HTML5-" + InGameConfiguration) : InGameName));
 				}
 
 				if (LineStr.Contains("%UE4CMDLINE%"))
