#ue4
#ios
* fix for not looking for the correctly generated project name

diff --git a/Engine/Source/Programs/AutomationTool/IOS/IOSPlatform.Automation.cs b/Engine/Source/Programs/AutomationTool/IOS/IOSPlatform.Automation.cs
--- a/Engine/Source/Programs/AutomationTool/IOS/IOSPlatform.Automation.cs
+++ b/Engine/Source/Programs/AutomationTool/IOS/IOSPlatform.Automation.cs
@@ -225,7 +225,8 @@ public class IOSPlatform : Platform
 				bWasGenerated = true;
 				Directory.SetCurrentDirectory (CWD);
 
-				if (!Directory.Exists (XcodeProj))
+				XcodeProj = RawProjectPath.Replace(".uproject", ".xcodeproj"); 
+				if (!Directory.Exists(XcodeProj))
 				{
 					// something very bad happened
 					throw new AutomationException("iOS couldn't find the appropriate Xcode Project");
