#UE4 - UAT: More logging when when exiting

diff --git a/Engine/Source/Programs/AutomationTool/Program.cs b/Engine/Source/Programs/AutomationTool/Program.cs
--- a/Engine/Source/Programs/AutomationTool/Program.cs
+++ b/Engine/Source/Programs/AutomationTool/Program.cs
@@ -60,6 +60,7 @@ namespace AutomationTool
 			CommandUtils.ClearDirStack();
 			Environment.ExitCode = Result;
 
+			Log.WriteLine(TraceEventType.Information, "AutomationTool exiting with ExitCode={0}", Result);
 			return Result;
 		}
 
@@ -77,8 +78,9 @@ namespace AutomationTool
 		static void Domain_ProcessExit(object sender, EventArgs e)
 		{
 			// Kill all spawned processes.
+			Console.WriteLine("Domain_ProcessExit");
 			if (ShouldKillProcesses)
-			{
+			{			
 				ProcessManager.KillAll();
 			}
 			// Close logging
diff --git a/Engine/Source/Programs/AutomationToolLauncher/Launcher.cs b/Engine/Source/Programs/AutomationToolLauncher/Launcher.cs
--- a/Engine/Source/Programs/AutomationToolLauncher/Launcher.cs
+++ b/Engine/Source/Programs/AutomationToolLauncher/Launcher.cs
@@ -35,7 +35,7 @@ namespace AutomationToolLauncher
 			{
 				Console.WriteLine(Ex.Message);
 			}
-
+			Console.WriteLine("AutomationToolLauncher exiting with ExitCode={0}", ExitCode);
 			return ExitCode;
 		}
 	}
