build script, long descriptions for changes

diff --git a/Engine/Source/Programs/AutomationTool/P4Utils.cs b/Engine/Source/Programs/AutomationTool/P4Utils.cs
--- a/Engine/Source/Programs/AutomationTool/P4Utils.cs
+++ b/Engine/Source/Programs/AutomationTool/P4Utils.cs
@@ -557,7 +557,7 @@ namespace AutomationTool
             return Result;
         }
         static Dictionary<string, List<ChangeRecord>> ChangesCache = new Dictionary<string, List<ChangeRecord>>();
-        public bool Changes(out List<ChangeRecord> ChangeRecords, string CommandLine, bool AllowSpew = true, bool UseCaching = false)
+        public bool Changes(out List<ChangeRecord> ChangeRecords, string CommandLine, bool AllowSpew = true, bool UseCaching = false, bool LongComment = false)
         {
             if (UseCaching && ChangesCache.ContainsKey(CommandLine))
             {
@@ -570,11 +570,46 @@ namespace AutomationTool
             {
                 // Change 1999345 on 2014/02/16 by buildmachine@BuildFarm_BUILD-23_buildmachine_++depot+UE4 'GUBP Node Shadow_LabelPromotabl'
                 string Output;
-                if (!LogP4Output(out Output, "changes " + CommandLine, null, AllowSpew))
+                if (!LogP4Output(out Output, "changes " + (LongComment ? "-L " : "") + CommandLine, null, AllowSpew))
                 {
                     return false;
                 }
-                var Lines = Output.Split(new string[] { Environment.NewLine }, StringSplitOptions.RemoveEmptyEntries);
+                var TempLines = Output.Split(new string[] { Environment.NewLine }, StringSplitOptions.RemoveEmptyEntries);
+                var Lines = new List<string>();
+                if (LongComment)
+                {
+                    string Start = null;
+                    string Comment = "";
+                    foreach (var Line in TempLines)
+                    {
+                        if (Line.StartsWith("\t"))
+                        {
+                            Comment += Line;
+                        }
+                        else
+                        {
+                            if (Start != null)
+                            {
+                                Start = Start + " '" + Comment.Replace("\t", "").Replace("\n", "  ").Replace("\r", "") + "'";
+                                Lines.Add(Start);
+                                Comment = "";
+                            }
+                            Start = Line;
+                        }
+                    }
+                    if (Start != null)
+                    {
+                        Start = Start + " '" + Comment.Replace("\t", "").Replace("\n", "  ").Replace("\r", "") + "'";
+                        Lines.Add(Start);
+                    }
+                }
+                else
+                {
+                    foreach (var Line in TempLines)
+                    {
+                        Lines.Add(Line);
+                    }
+                }
                 foreach (var Line in Lines)
                 {
                     ChangeRecord Change = new ChangeRecord();
diff --git a/Engine/Source/Programs/AutomationTool/Scripts/GUBP.Automation.cs b/Engine/Source/Programs/AutomationTool/Scripts/GUBP.Automation.cs
--- a/Engine/Source/Programs/AutomationTool/Scripts/GUBP.Automation.cs
+++ b/Engine/Source/Programs/AutomationTool/Scripts/GUBP.Automation.cs
@@ -2968,7 +2968,7 @@ public class GUBP : BuildCommand
                     CombinePaths(PathSeparator.Slash, P4Env.BuildRootP4, "*", "Build", "..."), LastOutputForChanges + 1, TopCL
                     );
                 List<P4Connection.ChangeRecord> ChangeRecords;
-				if (P4.Changes(out ChangeRecords, Cmd, false, true))
+				if (P4.Changes(out ChangeRecords, Cmd, false, true, LongComment: true))
                 {
                     foreach (var Record in ChangeRecords)
                     {
@@ -3008,7 +3008,10 @@ public class GUBP : BuildCommand
         int LastGreen = History.LastSucceeded;
         if (bShowAllChanges)
         {
-            LastGreen = 0;
+            if (History.AllStarted.Count > 0)
+            {
+                LastGreen = History.AllStarted[0];
+            }
         }
         foreach (var cl in History.AllStarted)
         {
diff --git a/Engine/Source/Programs/AutomationTool/TempStorage.cs b/Engine/Source/Programs/AutomationTool/TempStorage.cs
--- a/Engine/Source/Programs/AutomationTool/TempStorage.cs
+++ b/Engine/Source/Programs/AutomationTool/TempStorage.cs
@@ -874,7 +874,6 @@ namespace AutomationTool
                             {
                                 throw new AutomationException("Dir {0} didn't have any string to fit the star in the wildcard {1}", ThisSubDir, WildCard);
                             }
-                            Log("  Star Replacement {0}", StarReplacement);
                             if (!Files.Contains(StarReplacement))
                             {
                                 Files.Add(StarReplacement);
