#UE4 - UAT Improvements:

- Project rules will now be stored in EngineSavedFolder when running non-installed UAT, otherwise in Temp/UAT/EscapedRoot folder
- Added CommandUtils.EscapePath utility function
- Moved DeleteDirectoryContents function from LocalBuildEnvironment to CommandUtils
- Added code to clean rules folder up on startup.

diff --git a/Engine/Source/Programs/AutomationTool/Automation.cs b/Engine/Source/Programs/AutomationTool/Automation.cs
--- a/Engine/Source/Programs/AutomationTool/Automation.cs
+++ b/Engine/Source/Programs/AutomationTool/Automation.cs
@@ -295,6 +295,9 @@ AutomationTool.exe [-verbose] [-compileonly] [-p4] Command0 [-Arg0 -Arg1 -Arg2 
 			// Fill in the project info
 			UnrealBuildTool.UProjectInfo.FillProjectInfo();
 
+			// Clean rules folders up
+			ProjectUtils.CleanupFolders();
+
 			// Compile scripts.
 			Log.TraceInformation("Compiling scripts.");
 			ScriptCompiler Compiler = new ScriptCompiler();
diff --git a/Engine/Source/Programs/AutomationTool/CommandUtils.cs b/Engine/Source/Programs/AutomationTool/CommandUtils.cs
--- a/Engine/Source/Programs/AutomationTool/CommandUtils.cs
+++ b/Engine/Source/Programs/AutomationTool/CommandUtils.cs
@@ -464,8 +464,11 @@ namespace AutomationTool
                 }
                 catch (Exception Ex)
                 {
-                    Log(TraceEventType.Warning, "Failed to delete directory, exception '{0}'", NormalizedDirectory);
-                    Log(TraceEventType.Warning, Ex);
+					if (!bQuiet)
+					{
+						Log(TraceEventType.Warning, "Failed to delete directory, exception '{0}'", NormalizedDirectory);
+						Log(TraceEventType.Warning, Ex);
+					}
                     Result = false;
                 }
             }
@@ -484,6 +487,29 @@ namespace AutomationTool
 			return DeleteDirectory_NoExceptions(false, Directories);
 		}
 
+
+		/// <summary>
+		/// Attempts to delete a directory, if that fails deletes all files and folder from the specified directory.
+		/// This works around the issue when the user has a file open in a notepad from that directory. Somehow deleting the file works but
+		/// deleting the directory with the file that's open, doesn't.
+		/// </summary>
+		/// <param name="DirectoryName"></param>
+		public static void DeleteDirectoryContents(string DirectoryName)
+		{
+			Log("DeleteDirectoryContents({0})", DirectoryName);
+			const bool bQuiet = true;
+			var Files = CommandUtils.FindFiles_NoExceptions(bQuiet, "*", false, DirectoryName);
+			foreach (var Filename in Files)
+			{
+				CommandUtils.DeleteFile_NoExceptions(Filename);
+			}
+			var Directories = CommandUtils.FindDirectories_NoExceptions(bQuiet, "*", false, DirectoryName);
+			foreach (var SubDirectoryName in Directories)
+			{
+				CommandUtils.DeleteDirectory_NoExceptions(bQuiet, SubDirectoryName);
+			}
+		}
+
 		/// <summary>
 		/// Checks if a directory(or directories) exists.
 		/// </summary>
@@ -1587,6 +1613,11 @@ namespace AutomationTool
 
 		#region Other
 
+		public static string EscapePath(string InPath)
+		{
+			return InPath.Replace(":", "").Replace("/", "+").Replace("\\", "+").Replace(" ", "+");
+		}
+
 		/// <summary>
 		/// Checks if collection is either null or empty.
 		/// </summary>
diff --git a/Engine/Source/Programs/AutomationTool/LocalBuildEnvironment.cs b/Engine/Source/Programs/AutomationTool/LocalBuildEnvironment.cs
--- a/Engine/Source/Programs/AutomationTool/LocalBuildEnvironment.cs
+++ b/Engine/Source/Programs/AutomationTool/LocalBuildEnvironment.cs
@@ -87,8 +87,7 @@ namespace AutomationTool
 			{
 				LocalRootEscaped = LocalRootEscaped.Substring(0, LocalRootEscaped.Length - 1);
 			}
-			LocalRootEscaped = LocalRootEscaped.Replace("/", "+");
-			LocalRootEscaped = LocalRootEscaped.Replace(" ", "+");
+			LocalRootEscaped = CommandUtils.EscapePath(LocalRootEscaped);
 			return CommandUtils.CombinePaths(PathSeparator.Slash, HostPlatform.Current.LocalBuildsLogFolder, LocalRootEscaped);
 		}
 
@@ -102,7 +101,7 @@ namespace AutomationTool
 			{
 				string RootSavedPath;
 				if (!GlobalCommandLine.Installed)
-				{
+				{					
 					RootSavedPath = CommandUtils.CombinePaths(PathSeparator.Slash, CommandUtils.GetEnvVar(EnvVarNames.EngineSavedFolder), "Logs");
 				}
 				else
@@ -125,26 +124,6 @@ namespace AutomationTool
 			}
 		}
 
-		/// <summary>
-		/// Deletes all files and folder from the specified directory.
-		/// </summary>
-		/// <param name="DirectoryName"></param>
-		private static void DeleteDirectoryContents(string DirectoryName)
-		{
-			Log.TraceInformation("DeleteDirectoryContents({0})", DirectoryName);
-			const bool bQuiet = true;
-			var Files = CommandUtils.FindFiles_NoExceptions(bQuiet, "*", false, DirectoryName);
-			foreach (var Filename in Files)
-			{
-				CommandUtils.DeleteFile_NoExceptions(Filename);
-			}
-			var Directories = CommandUtils.FindDirectories_NoExceptions(bQuiet, "*", false, DirectoryName);
-			foreach (var SubDirectoryName in Directories)
-			{
-				CommandUtils.DeleteDirectory_NoExceptions(true, SubDirectoryName);
-			}
-		}
-
 		private bool TryCreateLogFolder(string PossibleLogFolder)
 		{
 			bool Result = true;
@@ -152,7 +131,8 @@ namespace AutomationTool
 			{
 				if (Directory.Exists(PossibleLogFolder))
 				{
-					DeleteDirectoryContents(PossibleLogFolder);
+					CommandUtils.DeleteDirectoryContents(PossibleLogFolder);
+
 					// Since the directory existed and might have been empty, we need to make sure
 					// we can actually write to it, so create and delete a temporary file.
 					var RandomFilename = Path.GetRandomFileName();
diff --git a/Engine/Source/Programs/AutomationTool/LocalP4Environment.cs b/Engine/Source/Programs/AutomationTool/LocalP4Environment.cs
--- a/Engine/Source/Programs/AutomationTool/LocalP4Environment.cs
+++ b/Engine/Source/Programs/AutomationTool/LocalP4Environment.cs
@@ -102,8 +102,7 @@ namespace AutomationTool
 			{
 				throw new AutomationException("Build root is empty");
 			}
-			BuildRoot = BuildRoot.Replace(":", "");
-			BuildRoot = BuildRoot.Replace("/", "+");
+			BuildRoot = CommandUtils.EscapePath(BuildRoot);
 			CommandUtils.ConditionallySetEnvVar(EnvVarNames.BuildRootEscaped, BuildRoot);
 		}
 
diff --git a/Engine/Source/Programs/AutomationTool/ProjectUtils.cs b/Engine/Source/Programs/AutomationTool/ProjectUtils.cs
--- a/Engine/Source/Programs/AutomationTool/ProjectUtils.cs
+++ b/Engine/Source/Programs/AutomationTool/ProjectUtils.cs
@@ -203,6 +203,23 @@ namespace AutomationTool
 			return BinPath;
 		}
 
+		/// <summary>
+		/// Gets the location where all rules assemblies should go
+		/// </summary>
+		private static string GetRulesAssemblyFolder()
+		{
+			string RulesFolder;
+			if (GlobalCommandLine.Installed.IsSet)
+			{
+				RulesFolder = CommandUtils.CombinePaths(Path.GetTempPath(), "UAT", CommandUtils.EscapePath(CommandUtils.CmdEnv.LocalRoot), "Rules"); 
+			}
+			else
+			{
+				RulesFolder = CommandUtils.CombinePaths(CommandUtils.CmdEnv.EngineSavedFolder, "Rules");
+			}
+			return RulesFolder;
+		}
+
 		/// <summary>
 		/// Finds all targets for the project.
 		/// </summary>
@@ -215,11 +232,12 @@ namespace AutomationTool
 			string FullProjectPath = null;
 
 			var GameFolders = new List<string>();
+			var RulesFolder = GetRulesAssemblyFolder();
 			if (!String.IsNullOrEmpty(Properties.RawProjectPath))
 			{
 				CommandUtils.Log("Looking for targets for project {0}", Properties.RawProjectPath);
 
-				TargetsDllFilename = CommandUtils.CombinePaths(Path.GetTempPath(), String.Format("UATRules{0}.dll", Properties.RawProjectPath.GetHashCode()));
+				TargetsDllFilename = CommandUtils.CombinePaths(RulesFolder, String.Format("UATRules{0}.dll", Properties.RawProjectPath.GetHashCode()));
 
 				FullProjectPath = CommandUtils.GetDirectoryName(Properties.RawProjectPath);
 				GameFolders.Add(FullProjectPath);
@@ -227,7 +245,7 @@ namespace AutomationTool
 			}
 			else
 			{
-				TargetsDllFilename = CommandUtils.CombinePaths(Path.GetTempPath(), String.Format("UATRules{0}.dll", "_BaseEngine_"));
+				TargetsDllFilename = CommandUtils.CombinePaths(RulesFolder, String.Format("UATRules{0}.dll", "_BaseEngine_"));
 			}
 			RulesCompiler.SetAssemblyNameAndGameFolders(TargetsDllFilename, GameFolders);
 
@@ -383,6 +401,19 @@ namespace AutomationTool
 			}
 			return Name;
 		}
+
+		/// <summary>
+		/// Performs initial cleanup of target rules folder
+		/// </summary>
+		public static void CleanupFolders()
+		{
+			CommandUtils.Log("Cleaning up project rules folder");
+			var RulesFolder = GetRulesAssemblyFolder();
+			if (CommandUtils.DirectoryExists(RulesFolder))
+			{
+				CommandUtils.DeleteDirectoryContents(RulesFolder);
+			}
+		}
 	}
 
     public class BranchInfo
diff --git a/Engine/Source/Programs/AutomationTool/Scripts/RunProjectCommand.Automation.cs b/Engine/Source/Programs/AutomationTool/Scripts/RunProjectCommand.Automation.cs
--- a/Engine/Source/Programs/AutomationTool/Scripts/RunProjectCommand.Automation.cs
+++ b/Engine/Source/Programs/AutomationTool/Scripts/RunProjectCommand.Automation.cs
@@ -60,7 +60,7 @@ public partial class Project : CommandUtils
 	{
 		return GlobalCommandLine.Installed ? 
 			CmdEnv.LogFolder : 
-			CombinePaths(Path.GetTempPath(), CmdEnv.LocalRoot.Replace("/", "+").Replace("\\", "+").Replace(":", "+"), "Logs");
+			CombinePaths(Path.GetTempPath(), CommandUtils.EscapePath(CmdEnv.LocalRoot), "Logs");
 	}
 
 	/// <summary>
