build script, make sure version files make real build products for a local rocket build

diff --git a/Engine/Source/Programs/AutomationTool/Scripts/GUBP.Automation.cs b/Engine/Source/Programs/AutomationTool/Scripts/GUBP.Automation.cs
--- a/Engine/Source/Programs/AutomationTool/Scripts/GUBP.Automation.cs
+++ b/Engine/Source/Programs/AutomationTool/Scripts/GUBP.Automation.cs
@@ -353,16 +353,8 @@ public class GUBP : BuildCommand
 
         public override void DoBuild(GUBP bp)
         {
-            if (CommandUtils.P4Enabled && CommandUtils.AllowSubmit)
-            {
-                // only update version files on build machines
-                var UE4Build = new UE4Build(bp);
-                BuildProducts = UE4Build.UpdateVersionFiles();
-            }
-            else
-            {
-                base.DoBuild(bp);
-            }
+            var UE4Build = new UE4Build(bp);
+            BuildProducts = UE4Build.UpdateVersionFiles(ActuallyUpdateVersionFiles: CommandUtils.P4Enabled && CommandUtils.AllowSubmit);
         }
         public override bool IsSticky()
         {
diff --git a/Engine/Source/Programs/AutomationTool/UE4Build.cs b/Engine/Source/Programs/AutomationTool/UE4Build.cs
--- a/Engine/Source/Programs/AutomationTool/UE4Build.cs
+++ b/Engine/Source/Programs/AutomationTool/UE4Build.cs
@@ -415,13 +415,13 @@ namespace AutomationTool
 		/// <summary>
 		/// Updates the engine version files
 		/// </summary>
-		public List<string> UpdateVersionFiles()
+		public List<string> UpdateVersionFiles(bool ActuallyUpdateVersionFiles = true)
 		{
 			var Result = new List<String>();
 			string Branch = P4Enabled ? P4Env.BuildRootEscaped : "";
 			{
 				string VerFile = CombinePaths(CmdEnv.LocalRoot, "Engine", "Build", "build.properties");
-				if (CommandUtils.P4Enabled)
+                if (CommandUtils.P4Enabled && ActuallyUpdateVersionFiles)
 				{
 					Log("Updating {0} with:", VerFile);
 					Log("  TimestampForBVT={0}", CmdEnv.TimestampAsString);
@@ -449,7 +449,7 @@ namespace AutomationTool
 
 			{
 				string VerFile = CombinePaths(CmdEnv.LocalRoot, "Engine", "Source", "Runtime", "Core", "Private", "UObject", "ObjectVersion.cpp");
-				if (CommandUtils.P4Enabled)
+                if (CommandUtils.P4Enabled && ActuallyUpdateVersionFiles)
 				{
 					Log("Updating {0} with:", VerFile);
 					Log(" #define	ENGINE_VERSION  {0}", P4Env.ChangelistString);
@@ -468,7 +468,7 @@ namespace AutomationTool
             string EngineVersionFile = CombinePaths(CmdEnv.LocalRoot, "Engine", "Source", "Runtime", "Launch", "Resources", "Version.h");
             {
                 string VerFile = EngineVersionFile;
-				if (CommandUtils.P4Enabled)
+                if (CommandUtils.P4Enabled && ActuallyUpdateVersionFiles)
 				{
 					Log("Updating {0} with:", VerFile);
 					Log(" #define	ENGINE_VERSION  {0}", P4Env.ChangelistString);
@@ -497,7 +497,7 @@ namespace AutomationTool
                 // Use Version.h data to update MetaData.cs so the assemblies match the engine version.
                 string VerFile = CombinePaths(CmdEnv.LocalRoot, "Engine", "Source", "Programs", "DotNETCommon", "MetaData.cs");
 
-                if (CommandUtils.P4Enabled)
+                if (CommandUtils.P4Enabled && ActuallyUpdateVersionFiles)
                 {
                     // Get the MAJOR/MINOR/PATCH from the Engine Version file, as it is authoritative. The rest we get from the P4Env.
                     string NewInformationalVersion = FEngineVersionSupport.FromVersionFile(EngineVersionFile).ToString();
