build script, add success emails

diff --git a/Engine/Source/Programs/AutomationTool/Scripts/GUBP.Automation.cs b/Engine/Source/Programs/AutomationTool/Scripts/GUBP.Automation.cs
--- a/Engine/Source/Programs/AutomationTool/Scripts/GUBP.Automation.cs
+++ b/Engine/Source/Programs/AutomationTool/Scripts/GUBP.Automation.cs
@@ -35,6 +35,7 @@ public class GUBP : BuildCommand
     class NodeHistory
     {
         public int LastSucceeded = 0;
+        public int LastFailed = 0;
         public List<int> InProgress = new List<int>();
         public string InProgressString = "";
         public List<int> Failed = new List<int>();
@@ -3167,6 +3168,7 @@ public class GUBP : BuildCommand
                 else
                 {
                     Log("         Last Success: {0}", History.LastSucceeded);
+                    Log("         Last Fail   : {0}", History.LastFailed);
                     Log("          Fails Since: {0}", History.FailedString);
                     Log("     InProgress Since: {0}", History.InProgressString);
                 }
@@ -4767,9 +4769,14 @@ if (HostPlatform == UnrealTargetPlatform.Mac) continue; //temp hack till mac aut
                     History.AllSucceeded = ConvertCLToIntList(FindTempStorageManifests(CmdEnv, NodeStoreWildCard + SucceededTempStorageSuffix, false, true, GameNameIfAny));
                     History.AllFailed = ConvertCLToIntList(FindTempStorageManifests(CmdEnv, NodeStoreWildCard + FailedTempStorageSuffix, false, true, GameNameIfAny));
 
+                    if (History.AllFailed.Count > 0)
+                    {
+                        History.LastFailed = History.AllFailed[History.AllFailed.Count - 1];
+                    }
                     if (History.AllSucceeded.Count > 0)
                     {
                         History.LastSucceeded = History.AllSucceeded[History.AllSucceeded.Count - 1];
+    
                         foreach (var Failed in History.AllFailed)
                         {
                             if (Failed > History.LastSucceeded)
@@ -4962,6 +4969,7 @@ if (HostPlatform == UnrealTargetPlatform.Mac) continue; //temp hack till mac aut
                 {
                     string EMails = "";
                     string FailCauserEMails = "";
+                    bool SendSuccessForGreenAfterRed = false;
                     if (GUBPNodesHistory.ContainsKey(NodeToDo))
                     {
                         var History = GUBPNodesHistory[NodeToDo];
@@ -4986,6 +4994,11 @@ if (HostPlatform == UnrealTargetPlatform.Mac) continue; //temp hack till mac aut
                             }
                         }
                         ECProps.Add(string.Format("FailCausers/{0}={1}", NodeToDo, FailCauserEMails));
+
+                        if (History.LastSucceeded > 0 && History.LastSucceeded < P4Env.Changelist && History.LastFailed > History.LastSucceeded && History.LastFailed < P4Env.Changelist)
+                        {
+                            SendSuccessForGreenAfterRed = ParseParam("CIS");
+                        }
                     }
                     else
                     {
@@ -5029,7 +5042,7 @@ if (HostPlatform == UnrealTargetPlatform.Mac) continue; //temp hack till mac aut
                         
                         ECProps.Add("FailEmails/" + NodeToDo + "=" + EMails);
                     }
-                    if (GUBPNodes[NodeToDo].SendSuccessEmail())
+                    if (GUBPNodes[NodeToDo].SendSuccessEmail() || SendSuccessForGreenAfterRed)
                     {
                         ECProps.Add("SendSuccessEmail/" + NodeToDo + "=1");
                     }
