build script, remove GUBP_V3 stuff and rearrange mac vs PC options

diff --git a/Engine/Source/Programs/AutomationTool/Scripts/GUBP.Automation.cs b/Engine/Source/Programs/AutomationTool/Scripts/GUBP.Automation.cs
--- a/Engine/Source/Programs/AutomationTool/Scripts/GUBP.Automation.cs
+++ b/Engine/Source/Programs/AutomationTool/Scripts/GUBP.Automation.cs
@@ -507,7 +507,7 @@ public class GUBP : BuildCommand
         public override int AgentMemoryRequirement(GUBP bp)
         {
             if (bp.ParseParam("Launcher") || bp.TimeIndex != 0)
-        {
+            {
                 return base.AgentMemoryRequirement(bp);
             }
             return 32;
@@ -532,10 +532,10 @@ public class GUBP : BuildCommand
         {
             if (HostPlatform != UnrealTargetPlatform.Mac)
             {
-                UE4Build.AddUATFilesToBuildProducts();
-                UE4Build.AddUBTFilesToBuildProducts();
-            }
-        }
+            	UE4Build.AddUATFilesToBuildProducts();
+            	UE4Build.AddUBTFilesToBuildProducts();
+        	}
+    	}
     }
 
     public class RootEditorNode : CompileNode
@@ -657,7 +657,7 @@ public class GUBP : BuildCommand
         {
             if (!GUBP.bBuildRocket) // more errors and more performance by just starting before the root editor is done
             {
-            AddPseudodependency(RootEditorNode.StaticGetFullName(HostPlatform));
+                AddPseudodependency(RootEditorNode.StaticGetFullName(HostPlatform));
             }
             AgentSharingGroup = "Tools" + StaticGetHostPlatformSuffix(HostPlatform);
         }
@@ -782,7 +782,7 @@ public class GUBP : BuildCommand
         {
             if (!GUBP.bBuildRocket) // more errors and more performance by just starting before the root editor is done
             {
-            AddPseudodependency(RootEditorNode.StaticGetFullName(HostPlatform));
+                AddPseudodependency(RootEditorNode.StaticGetFullName(HostPlatform));
             }
             AgentSharingGroup = "Tools" + StaticGetHostPlatformSuffix(HostPlatform);
         }
@@ -1037,8 +1037,8 @@ public class GUBP : BuildCommand
                 AddPseudodependency(EditorGameNode.StaticGetFullName(InHostPlatform, GameProj));
                 if (bp.HasNode(GamePlatformMonolithicsNode.StaticGetFullName(HostPlatform, bp.Branch.BaseEngineProject, TargetPlatform)))
                 {
-                	AddPseudodependency(GamePlatformMonolithicsNode.StaticGetFullName(HostPlatform, bp.Branch.BaseEngineProject, TargetPlatform));
-            	}
+                AddPseudodependency(GamePlatformMonolithicsNode.StaticGetFullName(HostPlatform, bp.Branch.BaseEngineProject, TargetPlatform));
+            }
             }
             else
             {
@@ -1514,7 +1514,7 @@ public class GUBP : BuildCommand
             {
                 return base.ECProcedure(); // after this user hits the trigger, we want to run this as an ordinary node
             }
-            if (GUBP.ECProject != "GUBP_V3" || TriggerRequiresRecursiveWorkflow())
+            if (TriggerRequiresRecursiveWorkflow())
             {
                 return String.Format("GUBP_UAT_Trigger"); //here we run a recursive workflow to wait for the trigger
             }
@@ -1523,7 +1523,7 @@ public class GUBP : BuildCommand
         public override string ECProcedureParams()
         {
             var Result = base.ECProcedureParams();
-            if (GUBP.ECProject == "GUBP_V3" && !bTriggerWasTriggered)
+            if (!bTriggerWasTriggered)
             {
                 Result += String.Format(" --actualParameter TriggerState={0} --actualParameter ActionText=\"{1}\" --actualParameter DescText=\"{2}\"", GetTriggerStateName(), GetTriggerActionText(), GetTriggerDescText());
             }
@@ -1677,7 +1677,7 @@ public class GUBP : BuildCommand
             {
                 //shared promotable has a shorter name
                 CompleteLabelPrefix = "Promotable";
-            }
+            }                   
             if (GUBP.bBuildRocket)
             {
                 CompleteLabelPrefix = "Rocket-" + CompleteLabelPrefix;
@@ -3556,7 +3556,7 @@ public class GUBP : BuildCommand
 
     string GetJobStepPath(string Dep)
     {
-        if (GUBP.ECProject == "GUBP_V3" && GUBPNodes[Dep].AgentSharingGroup != "")
+        if (GUBPNodes[Dep].AgentSharingGroup != "")
         {
             return "jobSteps[" + GUBPNodes[Dep].AgentSharingGroup + "]/jobSteps[" + Dep + "]";
         }
@@ -3570,8 +3570,8 @@ public class GUBP : BuildCommand
 
     [Help("Runs one, several or all of the GUBP nodes")]
     [Help(typeof(UE4Build))]
-    [Help("WithMac", "Toggle to set the host platform to Mac+Win64, default is Win64")]
-    [Help("OnlyMac", "Toggle to set the host platform to Mac, default is Win64")]
+    [Help("NoMac", "Toggle to exclude the Mac host platform, default is Win64+Mac")]
+    [Help("NoPC", "Toggle to exclude the PC host platform, default is Win64+Mac")]
     [Help("CleanLocal", "delete the local temp storage before we start")]
     [Help("Store=", "Sets the name of the temp storage block, normally, this is built for you.")]
     [Help("StoreSuffix=", "Tacked onto a store name constructed from CL, branch, etc")]
@@ -3612,16 +3612,15 @@ public class GUBP : BuildCommand
         ECProject = ParseParamValue("ECProject");
         if (ECProject == null)
         {
-            ECProject = "GUBP_V2";
+            ECProject = "";
         }
 
         HostPlatforms = new List<UnrealTargetPlatform>();
-        if (!ParseParam("OnlyMac"))
+        if (!ParseParam("NoPC"))
         {
             HostPlatforms.Add(UnrealTargetPlatform.Win64);
         }
-
-        if (ParseParam("WithMac") || HostPlatforms.Count == 0 || ECProject == "GUBP_V3")
+        if (!ParseParam("NoMac"))
         {
             HostPlatforms.Add(UnrealTargetPlatform.Mac);
         }
@@ -4034,7 +4033,7 @@ public class GUBP : BuildCommand
                             if (HostMonos.Contains(PlatPair.Key))
                             {
                                 NewPlatList.Add(PlatPair);
-                            }
+                }
                         }
                         if (NewPlatList.Count > 0)
                         {
@@ -4202,27 +4201,27 @@ public class GUBP : BuildCommand
 
                                                 }
                                             }
-                                        }
+                                    }
                                     }
                                     if (bAutomatedTesting)
                                     {
-                                        if (HostPlatform == UnrealTargetPlatform.Mac) continue; //temp hack till mac automated testing works
-                                        var ThisMonoGameTestNodes = new List<string>();
-                                        foreach (var Test in GameTests)
-                                        {
-                                            var TestName = Test.Key + "_" + Plat.ToString();
-                                            ThisMonoGameTestNodes.Add(AddNode(new UATTestNode(this, HostPlatform, NonCodeProject, TestName, Test.Value, CookedAgentSharingGroup, false, RequiredPlatforms)));
-                                        }
-                                        if (ThisMonoGameTestNodes.Count > 0)
-                                        {
-                                            GameTestNodes.Add(AddNode(new GameAggregateNode(this, HostPlatform, NonCodeProject, "CookedTests_" + Plat.ToString() + "_" + Kind.ToString() + HostPlatformNode.StaticGetHostPlatformSuffix(HostPlatform), ThisMonoGameTestNodes, 0.0f)));
-                                        }
+                                    if (HostPlatform == UnrealTargetPlatform.Mac) continue; //temp hack till mac automated testing works
+                                    var ThisMonoGameTestNodes = new List<string>();
+                                    foreach (var Test in GameTests)
+                                    {
+                                        var TestName = Test.Key + "_" + Plat.ToString();
+                                        ThisMonoGameTestNodes.Add(AddNode(new UATTestNode(this, HostPlatform, NonCodeProject, TestName, Test.Value, CookedAgentSharingGroup, false, RequiredPlatforms)));
+                                    }
+                                    if (ThisMonoGameTestNodes.Count > 0)
+                                    {
+                                        GameTestNodes.Add(AddNode(new GameAggregateNode(this, HostPlatform, NonCodeProject, "CookedTests_" + Plat.ToString() + "_" + Kind.ToString() + HostPlatformNode.StaticGetHostPlatformSuffix(HostPlatform), ThisMonoGameTestNodes, 0.0f)));
                                     }
                                 }
                             }
                         }
                     }
                 }
+                }
 #if false
                 //for now, non-code projects don't do client or server.
                 foreach (var ServerPlatform in ServerPlatforms)
@@ -5021,7 +5020,7 @@ if (HostPlatform == UnrealTargetPlatform.Mac) continue; //temp hack till mac aut
                             EMails = ForceAllEmailsTo;
                         }
                         else
-                        {
+                        {                            
                             EMails = GetEMailListForNode(this, NodeToDo);
                             if (ParseParam("CIS") && !GUBPNodes[NodeToDo].SendSuccessEmail())
                             {
@@ -5035,7 +5034,7 @@ if (HostPlatform == UnrealTargetPlatform.Mac) continue; //temp hack till mac aut
                         }
 
                         {
-                            EMails = EMails.Trim().Replace("[epic]", "@epicgames.com");
+                        EMails = EMails.Trim().Replace("[epic]", "@epicgames.com");
                             EMails = EMails.Replace("[phosphor]", "@phosphorgames.com");
                             EMails = EMails.Replace("[pitbull]", "@pitbullstudio.co.uk");
                         }
@@ -5070,7 +5069,7 @@ if (HostPlatform == UnrealTargetPlatform.Mac) continue; //temp hack till mac aut
                         Args = Args + " " + ProcedureParams;
                     }
 
-                    if ((Procedure == "GUBP_UAT_Trigger" || Procedure == "GUBP_Hardcoded_Trigger") && !String.IsNullOrEmpty(EMails) && GUBP.ECProject == "GUBP_V3")
+                    if ((Procedure == "GUBP_UAT_Trigger" || Procedure == "GUBP_Hardcoded_Trigger") && !String.IsNullOrEmpty(EMails))
                     {
                         Args = Args + " --actualParameter EmailsForTrigger=\"" + EMails + "\"";
                     }
@@ -5080,19 +5079,19 @@ if (HostPlatform == UnrealTargetPlatform.Mac) continue; //temp hack till mac aut
                     var UncompletedEcDeps = new List<string>();
 
                     {
-                    var EcDeps = GetECDependencies(NodeToDo);
-                    foreach (var Dep in EcDeps)
-                    {
-                        if (GUBPNodes[Dep].RunInEC() && !NodeIsAlreadyComplete(Dep, LocalOnly) && OrdereredToDo.Contains(Dep)) // if something is already finished, we don't put it into EC
+                        var EcDeps = GetECDependencies(NodeToDo);
+                        foreach (var Dep in EcDeps)
                         {
-                            if (OrdereredToDo.IndexOf(Dep) > OrdereredToDo.IndexOf(NodeToDo))
+                            if (GUBPNodes[Dep].RunInEC() && !NodeIsAlreadyComplete(Dep, LocalOnly) && OrdereredToDo.Contains(Dep)) // if something is already finished, we don't put it into EC
                             {
-                                throw new AutomationException("Topological sort error, node {0} has a dependency of {1} which sorted after it.", NodeToDo, Dep);
+                                if (OrdereredToDo.IndexOf(Dep) > OrdereredToDo.IndexOf(NodeToDo))
+                                {
+                                    throw new AutomationException("Topological sort error, node {0} has a dependency of {1} which sorted after it.", NodeToDo, Dep);
+                                }
+                                UncompletedEcDeps.Add(Dep);
                             }
-                            UncompletedEcDeps.Add(Dep);
                         }
                     }
-                    }
 
                     string MyAgentGroup = GUBPNodes[NodeToDo].AgentSharingGroup;
                     bool bDoNestedJobstep = false;
@@ -5102,15 +5101,15 @@ if (HostPlatform == UnrealTargetPlatform.Mac) continue; //temp hack till mac aut
                     string NodeParentPath = ParentPath;
                     if (MyAgentGroup != "")
                     {
-                        bDoNestedJobstep = GUBP.ECProject == "GUBP_V3";
+                        bDoNestedJobstep = true;
                         NodeParentPath = ParentPath + "/jobSteps[" + MyAgentGroup + "]";
 
                         PreConditionUncompletedEcDeps = new List<string>();
 
-                                var MyChain = AgentGroupChains[MyAgentGroup];
-                                int MyIndex = MyChain.IndexOf(NodeToDo);
+                        var MyChain = AgentGroupChains[MyAgentGroup];
+                        int MyIndex = MyChain.IndexOf(NodeToDo);
                                 if (MyIndex  > 0)
-                                {
+                        {
                             PreConditionUncompletedEcDeps.Add(MyChain[MyIndex - 1]);
                         }
                         else
@@ -5129,48 +5128,48 @@ if (HostPlatform == UnrealTargetPlatform.Mac) continue; //temp hack till mac aut
                                             throw new AutomationException("Topological sort error, node {0} has a dependency of {1} which sorted after it.", Chain, Dep);
                                         }
                                         if (!MyChain.Contains(Dep) && !PreConditionUncompletedEcDeps.Contains(Dep))
-                                    {
+                                        {
                                             PreConditionUncompletedEcDeps.Add(Dep);
                                         }
                                     }
-                                    }
-                                }
                                 }
                             }
+                        }
+                    }
                     if (PreConditionUncompletedEcDeps.Count > 0)
                     {
                         PreCondition = "\"$[/javascript if(";
                         // these run "parallel", but we add preconditions to serialize them
-                            int Index = 0;
+                        int Index = 0;
                         foreach (var Dep in PreConditionUncompletedEcDeps)
-                            {
+                        {
                             PreCondition = PreCondition + "getProperty('" + GetJobStep(ParentPath, Dep) + "/status') == 'completed'";
-                                Index++;
+                            Index++;
                             if (Index != PreConditionUncompletedEcDeps.Count)
-                                {
-                                    PreCondition = PreCondition + " && ";
-                                }
+                            {
+                                PreCondition = PreCondition + " && ";
                             }
-                            PreCondition = PreCondition + ") true;]\"";
                         }
+                        PreCondition = PreCondition + ") true;]\"";
+                    }
 
                     if (UncompletedEcDeps.Count > 0)
+                    {
+                        RunCondition = "\"$[/javascript if(";
+                        int Index = 0;
+                        foreach (var Dep in UncompletedEcDeps)
                         {
-                            RunCondition = "\"$[/javascript if(";
-                            int Index = 0;
-                            foreach (var Dep in UncompletedEcDeps)
-                            {
                             RunCondition = RunCondition + "('$[" + GetJobStep(ParentPath, Dep) + "/outcome]' == 'success' || ";
                             RunCondition = RunCondition + "'$[" + GetJobStep(ParentPath, Dep) + "/outcome]' == 'warning')";
 
-                                Index++;
-                                if (Index != UncompletedEcDeps.Count)
-                                {
-                                    RunCondition = RunCondition + " && ";
-                                }
+                            Index++;
+                            if (Index != UncompletedEcDeps.Count)
+                            {
+                                RunCondition = RunCondition + " && ";
                             }
-                            RunCondition = RunCondition + ") true; else false;]\"";
                         }
+                        RunCondition = RunCondition + ") true; else false;]\"";
+                    }
 
                     if (bDoNestedJobstep)
                     {
@@ -5184,7 +5183,7 @@ if (HostPlatform == UnrealTargetPlatform.Mac) continue; //temp hack till mac aut
                                     NestArgs = NestArgs + " --precondition " + PreCondition;
                                 }
                                 RunECTool(NestArgs);
-                    }
+                            }
                             {
                                 string NestArgs = String.Format("createJobStep --parentPath {0}/jobSteps[{1}] --jobStepName {2}_GetPool --subprocedure GUBP{3}_AgentShare_GetPool --parallel 1 --actualParameter AgentSharingGroup={4} --actualParameter NodeName={5}",
                                     ParentPath, MyAgentGroup, MyAgentGroup, GUBPNodes[NodeToDo].ECProcedureInfix(), MyAgentGroup, NodeToDo);
@@ -5332,35 +5331,35 @@ if (HostPlatform == UnrealTargetPlatform.Mac) continue; //temp hack till mac aut
             {
                 if (ParseParamValue("OnlyNode", "") == "")
                 {
-                throw new AutomationException("Node {0} was not processed yet?", NodeToDo);
-            }
+                    throw new AutomationException("Node {0} was not processed yet?", NodeToDo);
+                }
             }
             else
             {
-            foreach (var Dep in GUBPNodes[NodeToDo].FullNamesOfDependencies)
-            {
-                if (GUBPNodes[Dep].BuildProducts == null)
+                foreach (var Dep in GUBPNodes[NodeToDo].FullNamesOfDependencies)
                 {
+                    if (GUBPNodes[Dep].BuildProducts == null)
+                    {
                         if (ParseParamValue("OnlyNode", "") == "")
                         {
-                    throw new AutomationException("Node {0} was not processed yet2? Processing {1}", Dep, NodeToDo);
-                }
+                            throw new AutomationException("Node {0} was not processed yet2? Processing {1}", Dep, NodeToDo);
+                        }
                     }
                     else
                     {
-                foreach (var Prod in GUBPNodes[Dep].BuildProducts)
-                {
-                    GUBPNodes[NodeToDo].AddDependentBuildProduct(Prod);
-                }
-                if (GUBPNodes[Dep].AllDependencyBuildProducts == null)
-                {
-                    throw new AutomationException("Node {0} was not processed yet3?  Processing {1}", Dep, NodeToDo);
-                }
-                foreach (var Prod in GUBPNodes[Dep].AllDependencyBuildProducts)
-                {
-                    GUBPNodes[NodeToDo].AddDependentBuildProduct(Prod);
-                } 
-            }
+                        foreach (var Prod in GUBPNodes[Dep].BuildProducts)
+                        {
+                            GUBPNodes[NodeToDo].AddDependentBuildProduct(Prod);
+                        }
+                        if (GUBPNodes[Dep].AllDependencyBuildProducts == null)
+                        {
+                            throw new AutomationException("Node {0} was not processed yet3?  Processing {1}", Dep, NodeToDo);
+                        }
+                        foreach (var Prod in GUBPNodes[Dep].AllDependencyBuildProducts)
+                        {
+                            GUBPNodes[NodeToDo].AddDependentBuildProduct(Prod);
+                        }
+                    }
                 }
             }
             string NodeStoreName = StoreName + "-" + GUBPNodes[NodeToDo].GetFullName();
