build script, remove trigger for testing, 5h20m

diff --git a/Engine/Source/Programs/AutomationTool/Scripts/GUBP.Automation.cs b/Engine/Source/Programs/AutomationTool/Scripts/GUBP.Automation.cs
--- a/Engine/Source/Programs/AutomationTool/Scripts/GUBP.Automation.cs
+++ b/Engine/Source/Programs/AutomationTool/Scripts/GUBP.Automation.cs
@@ -1800,34 +1800,22 @@ public class GUBP : BuildCommand
         }
     }
 
-    public class WaitForTestShared : GUBP.WaitForUserInput
+    public class WaitForTestShared : AggregateNode
     {
         public WaitForTestShared(GUBP bp)
         {
-            // delay this until the root editor succeeds
-            AddPseudodependency(GUBP.SharedLabelPromotableNode.StaticGetFullName(false));
-        }
-        public override string GetTriggerDescText()
-        {
-            return "The shared promotable is ready to begin automated testing.";
-        }
-        public override string GetTriggerActionText()
-        {
-            return "Begin automated testing on the shared promotable.";
         }
         public static string StaticGetFullName()
         {
-            return "Shared_WaitForTesting";
+            return "Shared_TestingAggregate";
         }
         public override string GetFullName()
         {
             return StaticGetFullName();
         }
-        public override string FailureEMails(GUBP bp, string Branch)
+        public override int CISFrequencyQuantumShift(GUBP bp)
         {
-            return MergeSpaceStrings(base.FailureEMails(bp, Branch),
-               bp.Branch.BaseEngineProject.Properties.Targets[TargetRules.TargetType.Editor].Rules.GUBP_GetPromotionEMails_EditorTypeOnly(Branch));
-
+            return base.CISFrequencyQuantumShift(bp) + 4;
         }
     }
 
