L10N: Launcher: Automation: Updated the Launcher localization. Added/modified the Launcher configuration so it outputs word count and conflicts reports. Added a new automation script to build Launcher localization.

diff --git a/Engine/Source/Programs/AutomationTool/Scripts/AutomationScripts.Automation.csproj b/Engine/Source/Programs/AutomationTool/Scripts/AutomationScripts.Automation.csproj
--- a/Engine/Source/Programs/AutomationTool/Scripts/AutomationScripts.Automation.csproj
+++ b/Engine/Source/Programs/AutomationTool/Scripts/AutomationScripts.Automation.csproj
@@ -68,6 +68,7 @@
     <Compile Include="GenerateDSYM.Automation.cs" />
     <Compile Include="GUBP.Automation.cs" />
     <Compile Include="IOSBuildCookPackage.Automation.cs" />
+    <Compile Include="LauncherLocalization.Automation.cs" />
     <Compile Include="Localisation.Automation.cs" />
     <Compile Include="MegaXGE.Automation.cs" />
     <Compile Include="PackageCommand.Automation.cs" />
diff --git a/Engine/Source/Programs/AutomationTool/Scripts/LauncherLocalization.Automation.cs b/Engine/Source/Programs/AutomationTool/Scripts/LauncherLocalization.Automation.cs
new file mode 100644
--- /dev/null
+++ b/Engine/Source/Programs/AutomationTool/Scripts/LauncherLocalization.Automation.cs
@@ -0,0 +1,69 @@
+﻿// Copyright 1998-2014 Epic Games, Inc. All Rights Reserved.
+using System;
+using System.Collections.Generic;
+using System.Text;
+using System.IO;
+using AutomationTool;
+using UnrealBuildTool;
+
+class LauncherLocalization : BuildCommand
+{
+    public override void ExecuteBuild()
+    {
+        var EditorExe = CombinePaths(CmdEnv.LocalRoot, @"Engine/Binaries/Win64/UE4Editor-Cmd.exe");
+
+        if (P4Enabled)
+        {
+            Log("Sync necessary content to head revision");
+			P4.Sync(P4Env.BuildRootP4 + "/Engine/Config/...");
+			P4.Sync(P4Env.BuildRootP4 + "/Engine/Content/...");
+			P4.Sync(P4Env.BuildRootP4 + "/Engine/Source/...");
+
+            P4.Sync(P4Env.BuildRootP4 + "/Engine/Programs/NoRedist/UnrealEngineLauncher/Config/...");
+            P4.Sync(P4Env.BuildRootP4 + "/Engine/Programs/NoRedist/UnrealEngineLauncher/Content/...");
+            //P4.Sync(P4Env.BuildRootP4 + "/Engine/Source/..."); <- takes care of syncing Launcher source already
+
+            Log("Localize from label {0}", P4Env.LabelToSync);            
+        }
+
+        // Setup editor arguments for SCC.
+        string EditorArguments = String.Empty;
+        if (P4Enabled)
+        {
+			EditorArguments = String.Format("-SCCProvider={0} -P4Port={1} -P4User={2} -P4Client={3} -P4Passwd={4}", "Perforce", P4Env.P4Port, P4Env.User, P4Env.Client, P4.GetAuthenticationToken());
+        }
+        else
+        {
+            EditorArguments = String.Format("-SCCProvider={0}", "None");
+        }
+
+        // Setup commandlet arguments for SCC.
+        string CommandletSCCArguments = String.Empty;
+        if (P4Enabled) { CommandletSCCArguments += (string.IsNullOrEmpty(CommandletSCCArguments) ? "" : " ") + "-EnableSCC"; }
+        if (!AllowSubmit) { CommandletSCCArguments += (string.IsNullOrEmpty(CommandletSCCArguments) ? "" : " ") + "-DisableSCCSubmit"; }
+
+        // Setup commandlet arguments with configurations.
+        var CommandletArgumentSets = new string[] 
+				{
+                    String.Format("-config={0}", @"../../../Engine/Programs/NoRedist/UnrealEngineLauncher/Config/Localization/App.ini") + (string.IsNullOrEmpty(CommandletSCCArguments) ? "" : " " + CommandletSCCArguments),
+                    String.Format("-config={0}", @"../../../Engine/Programs/NoRedist/UnrealEngineLauncher/Config/Localization/WordCount.ini"),
+				};
+
+        // Execute commandlet for each set of arguments.
+        foreach (var CommandletArguments in CommandletArgumentSets)
+        {
+            Log("Localization for {0} {1}", EditorArguments, CommandletArguments);
+
+            Log("Running UE4Editor to generate Localization data");
+
+            string Arguments = String.Format("-run=GatherText {0} {1}", EditorArguments, CommandletArguments);
+            var RunResult = Run(EditorExe, Arguments);
+
+            if (RunResult.ExitCode != 0)
+            {
+                throw new AutomationException("Error while executing localization commandlet '{0}'", Arguments);
+            }
+        }
+    }
+}
+
