build script, nonunity is not under a trigger, makebuild independent of the promotable, cooks not under a trigger, got rid of hardcoded games that use custom workflows

diff --git a/Engine/Source/Programs/AutomationTool/Scripts/GUBP.Automation.cs b/Engine/Source/Programs/AutomationTool/Scripts/GUBP.Automation.cs
--- a/Engine/Source/Programs/AutomationTool/Scripts/GUBP.Automation.cs
+++ b/Engine/Source/Programs/AutomationTool/Scripts/GUBP.Automation.cs
@@ -1566,11 +1566,14 @@ public class GUBP : BuildCommand
     public class WaitForGamePromotionUserInput : WaitForPromotionUserInput
     {
         BranchInfo.BranchUProject GameProj;
+        bool bCustomWorkflow;
 
         public WaitForGamePromotionUserInput(GUBP bp, BranchInfo.BranchUProject InGameProj, bool bInLabelPromoted)
             : base(InGameProj.GameName, bInLabelPromoted)
         {
             GameProj = InGameProj;
+            var Options = InGameProj.Options(UnrealTargetPlatform.Win64);
+            bCustomWorkflow = Options.bCustomWorkflowForPromotion;       
         }
         public static string StaticGetFullName(BranchInfo.BranchUProject InGameProj, bool bInLabelPromoted)
         {
@@ -1598,7 +1601,7 @@ public class GUBP : BuildCommand
         }
         public override string GetTriggerStateName()
         {
-            if (GameProj.GameName.Equals("OrionGame", StringComparison.InvariantCultureIgnoreCase))
+            if (bCustomWorkflow)
             {
                 return GetFullName();
             }
@@ -1606,7 +1609,7 @@ public class GUBP : BuildCommand
         }
         public override bool TriggerRequiresRecursiveWorkflow()
         {
-            if (GameProj.GameName.Equals("OrionGame", StringComparison.InvariantCultureIgnoreCase))
+            if (bCustomWorkflow)
             {
                 return !bLabelPromoted; // the promotable starts the hardcoded chain
             }
@@ -1853,11 +1856,7 @@ public class GUBP : BuildCommand
                 // add an arc to prevent cooks from running until promotable is labeled
                 if (Options.bIsPromotable)
                 {
-                    if (Options.bSeparateGamePromotion)
-                    {
-                        AddPseudodependency(GameLabelPromotableNode.StaticGetFullName(GameProj, false));
-                    }
-                    else
+                    if (!Options.bSeparateGamePromotion)
                     {
                         bIsShared = true;
                     }
@@ -1906,6 +1905,11 @@ public class GUBP : BuildCommand
         {
             return GameProj.GameName;
         }
+        public override int CISFrequencyQuantumShift(GUBP bp)
+        {
+            return base.CISFrequencyQuantumShift(bp) + 4 + (bIsMassive ? 1 : 0) 
+                - (HostPlatform == UnrealTargetPlatform.Mac ? 1 : 0); // we don't want the mac penalty here
+        }
         public override float Priority()
         {
             return 10.0f;
@@ -2451,7 +2455,6 @@ public class GUBP : BuildCommand
         public NonUnityTestNode(UnrealTargetPlatform InHostPlatform)
             : base(InHostPlatform)
         {
-            AddPseudodependency(SharedLabelPromotableNode.StaticGetFullName(false));
         }
         public static string StaticGetFullName(UnrealTargetPlatform InHostPlatform)
         {
@@ -2461,7 +2464,10 @@ public class GUBP : BuildCommand
         {
             return StaticGetFullName(HostPlatform);
         }
-
+        public override int CISFrequencyQuantumShift(GUBP bp)
+        {
+            return base.CISFrequencyQuantumShift(bp) + 3;
+        }
         public override void DoTest(GUBP bp)
         {
             var Build = new UE4Build(bp);
diff --git a/Engine/Source/Programs/UnrealBuildTool/System/RulesCompiler.cs b/Engine/Source/Programs/UnrealBuildTool/System/RulesCompiler.cs
--- a/Engine/Source/Programs/UnrealBuildTool/System/RulesCompiler.cs
+++ b/Engine/Source/Programs/UnrealBuildTool/System/RulesCompiler.cs
@@ -620,6 +620,7 @@ namespace UnrealBuildTool
             public bool bSeparateGamePromotion = false;
             public bool bTestWithShared = false;
             public bool bIsMassive = false;
+            public bool bCustomWorkflowForPromotion = false;
         }
         public virtual GUBPProjectOptions GUBP_IncludeProjectInPromotedBuild_EditorTypeOnly(UnrealTargetPlatform HostPlatform)
         {
