build script, recheck history after each node to make sure we have up to date information

diff --git a/Engine/Source/Programs/AutomationTool/Scripts/GUBP.Automation.cs b/Engine/Source/Programs/AutomationTool/Scripts/GUBP.Automation.cs
--- a/Engine/Source/Programs/AutomationTool/Scripts/GUBP.Automation.cs
+++ b/Engine/Source/Programs/AutomationTool/Scripts/GUBP.Automation.cs
@@ -3561,6 +3561,169 @@ public class GUBP : BuildCommand
         return ParentPath + "/" + GetJobStepPath(Dep);
     }
 
+    void UpdateNodeHistory(string Node, string CLString)
+    {
+        if (GUBPNodes[Node].RunInEC() && !GUBPNodes[Node].TriggerNode() && CLString != "")
+        {
+            string GameNameIfAny = GUBPNodes[Node].GameNameIfAnyForTempStorage();
+            string NodeStoreWildCard = StoreName.Replace(CLString, "*") + "-" + GUBPNodes[Node].GetFullName();
+            var History = new NodeHistory();
+
+            History.AllStarted = ConvertCLToIntList(FindTempStorageManifests(CmdEnv, NodeStoreWildCard + StartedTempStorageSuffix, false, true, GameNameIfAny));
+            History.AllSucceeded = ConvertCLToIntList(FindTempStorageManifests(CmdEnv, NodeStoreWildCard + SucceededTempStorageSuffix, false, true, GameNameIfAny));
+            History.AllFailed = ConvertCLToIntList(FindTempStorageManifests(CmdEnv, NodeStoreWildCard + FailedTempStorageSuffix, false, true, GameNameIfAny));
+
+            if (History.AllFailed.Count > 0)
+            {
+                History.LastFailed = History.AllFailed[History.AllFailed.Count - 1];
+            }
+            if (History.AllSucceeded.Count > 0)
+            {
+                History.LastSucceeded = History.AllSucceeded[History.AllSucceeded.Count - 1];
+
+                foreach (var Failed in History.AllFailed)
+                {
+                    if (Failed > History.LastSucceeded)
+                    {
+                        History.Failed.Add(Failed);
+                        History.FailedString = GUBPNode.MergeSpaceStrings(History.FailedString, String.Format("{0}", Failed));
+                    }
+                }
+                foreach (var Started in History.AllStarted)
+                {
+                    if (Started > History.LastSucceeded && !History.Failed.Contains(Started))
+                    {
+                        History.InProgress.Add(Started);
+                        History.InProgressString = GUBPNode.MergeSpaceStrings(History.InProgressString, String.Format("{0}", Started));
+                    }
+                }
+                if (GUBPNodesHistory.ContainsKey(Node))
+                {
+                    GUBPNodesHistory.Remove(Node);
+                }
+                GUBPNodesHistory.Add(Node, History);
+            }
+        }
+    }
+
+    List<string> GetECPropsForNode(string NodeToDo, out string EMails, bool OnlyLateUpdates = false)
+    {
+        var ECProps = new List<string>();
+        EMails = "";
+        string FailCauserEMails = "";
+        bool SendSuccessForGreenAfterRed = false;
+        if (GUBPNodesHistory.ContainsKey(NodeToDo))
+        {
+            var History = GUBPNodesHistory[NodeToDo];
+
+            ECProps.Add(string.Format("LastGreen/{0}={1}", NodeToDo, History.LastSucceeded));
+            ECProps.Add(string.Format("RedsSince/{0}={1}", NodeToDo, History.FailedString));
+            ECProps.Add(string.Format("InProgress/{0}={1}", NodeToDo, History.InProgressString));
+
+
+            if (History.LastSucceeded > 0 && History.LastSucceeded < P4Env.Changelist)
+            {
+                var ChangeRecords = GetChanges(History.LastSucceeded, P4Env.Changelist, History.LastSucceeded);
+                int NumChanges = 0;
+                foreach (var Record in ChangeRecords)
+                {
+                    FailCauserEMails = GUBPNode.MergeSpaceStrings(FailCauserEMails, Record.UserEmail);
+                    if (++NumChanges > 50)
+                    {
+                        FailCauserEMails = "";
+                        break;
+                    }
+                }
+            }
+            ECProps.Add(string.Format("FailCausers/{0}={1}", NodeToDo, FailCauserEMails));
+
+            if (History.LastSucceeded > 0 && History.LastSucceeded < P4Env.Changelist && History.LastFailed > History.LastSucceeded && History.LastFailed < P4Env.Changelist)
+            {
+                SendSuccessForGreenAfterRed = ParseParam("CIS");
+            }
+        }
+        else
+        {
+            ECProps.Add(string.Format("LastGreen/{0}=0", NodeToDo));
+            ECProps.Add(string.Format("RedsSince/{0}=", NodeToDo));
+            ECProps.Add(string.Format("InProgress/{0}=", NodeToDo));
+        }
+        {
+            string EmailOnly = ParseParamValue("EmailOnly");
+            if (bFake)
+            {
+                EMails = "kellan.carr[epic] gil.gribb[epic]";
+            }
+            else if (!String.IsNullOrEmpty(EmailOnly))
+            {
+                EMails = EmailOnly;
+            }
+            else if (!String.IsNullOrEmpty(ForceAllEmailsTo))
+            {
+                EMails = ForceAllEmailsTo;
+            }
+            else
+            {
+                EMails = GetEMailListForNode(this, NodeToDo);
+                if (ParseParam("CIS") && !GUBPNodes[NodeToDo].SendSuccessEmail())
+                {
+                    EMails = GUBPNode.MergeSpaceStrings(FailCauserEMails, EMails);
+                }
+            }
+            string AddEmails = ParseParamValue("AddEmails");
+            if (!String.IsNullOrEmpty(AddEmails))
+            {
+                EMails = GUBPNode.MergeSpaceStrings(AddEmails, EMails);
+            }
+
+            {
+                EMails = EMails.Trim().Replace("[epic]", "@epicgames.com");
+                EMails = EMails.Replace("[phosphor]", "@phosphorgames.com");
+                EMails = EMails.Replace("[pitbull]", "@pitbullstudio.co.uk");
+            }
+
+            ECProps.Add("FailEmails/" + NodeToDo + "=" + EMails);
+        }
+        if (GUBPNodes[NodeToDo].SendSuccessEmail() || SendSuccessForGreenAfterRed)
+        {
+            ECProps.Add("SendSuccessEmail/" + NodeToDo + "=1");
+        }
+        else
+        {
+            ECProps.Add("SendSuccessEmail/" + NodeToDo + "=0");
+        }
+        if (!OnlyLateUpdates)
+        {
+            ECProps.Add(string.Format("AgentRequirementString/{0}={1}", NodeToDo, GUBPNodes[NodeToDo].ECAgentString()));
+            ECProps.Add(string.Format("RequiredMemory/{0}={1}", NodeToDo, GUBPNodes[NodeToDo].AgentMemoryRequirement(this)));
+            ECProps.Add(string.Format("Timeouts/{0}={1}", NodeToDo, GUBPNodes[NodeToDo].TimeoutInMinutes()));
+            ECProps.Add(string.Format("JobStepPath/{0}={1}", NodeToDo, GetJobStepPath(NodeToDo)));
+        }
+
+        return ECProps;
+    }
+
+    void UpdateHistoryAndEC(string NodeToDo, string CLString)
+    {
+        try
+        {
+            Log("Updating node props for node {0}", NodeToDo);
+            UpdateNodeHistory(NodeToDo, CLString);
+            string EMails;
+            var Props = GetECPropsForNode(NodeToDo, out EMails, true);
+            foreach (var Prop in Props)
+            {
+                var Parts = Prop.Split("=".ToCharArray());
+                RunECTool(String.Format("setProperty \"/myWorkflow/{0}\" \"{1}\"", Parts[0], Parts[1]));
+            }
+        }
+        catch (Exception Ex)
+        {
+            Log(System.Diagnostics.TraceEventType.Warning, "Failed to UpdateHistoryAndEC.");
+            Log(System.Diagnostics.TraceEventType.Warning, LogUtils.FormatException(Ex));
+        }
+    }
+
 
     [Help("Runs one, several or all of the GUBP nodes")]
     [Help(typeof(UE4Build))]
@@ -4739,55 +4902,26 @@ if (HostPlatform == UnrealTargetPlatform.Mac) continue; //temp hack till mac aut
             }
         }
 
-        Log("******* Caching completion");
         GUBPNodesCompleted = new Dictionary<string, bool>();
         GUBPNodesControllingTrigger = new Dictionary<string, string>();
         GUBPNodesControllingTriggerDotName = new Dictionary<string, string>();
         GUBPNodesHistory = new Dictionary<string, NodeHistory>();
 
+        Log("******* Caching completion");
         foreach (var Node in NodesToDo)
         {
             Log("** {0}", Node);
-            bool bComplete = NodeIsAlreadyComplete(Node, LocalOnly); // cache these now to avoid spam later
+            NodeIsAlreadyComplete(Node, LocalOnly); // cache these now to avoid spam later
             GetControllingTriggerDotName(Node);
-            if (!bComplete && CLString != "" && StoreName.Contains(CLString))
+        }
+        if (CLString != "" && StoreName.Contains(CLString) && !ParseParam("NoHistory"))
+        {
+            Log("******* Updating history");
+            foreach (var Node in NodesToDo)
             {
-                if (GUBPNodes[Node].RunInEC() && !GUBPNodes[Node].TriggerNode())
+                if (!NodeIsAlreadyComplete(Node, LocalOnly))
                 {
-                    string GameNameIfAny = GUBPNodes[Node].GameNameIfAnyForTempStorage();
-                    string NodeStoreWildCard = StoreName.Replace(CLString, "*") + "-" + GUBPNodes[Node].GetFullName();
-                    var History = new NodeHistory();
-
-                    History.AllStarted = ConvertCLToIntList(FindTempStorageManifests(CmdEnv, NodeStoreWildCard + StartedTempStorageSuffix, false, true, GameNameIfAny));
-                    History.AllSucceeded = ConvertCLToIntList(FindTempStorageManifests(CmdEnv, NodeStoreWildCard + SucceededTempStorageSuffix, false, true, GameNameIfAny));
-                    History.AllFailed = ConvertCLToIntList(FindTempStorageManifests(CmdEnv, NodeStoreWildCard + FailedTempStorageSuffix, false, true, GameNameIfAny));
-
-                    if (History.AllFailed.Count > 0)
-                    {
-                        History.LastFailed = History.AllFailed[History.AllFailed.Count - 1];
-                    }
-                    if (History.AllSucceeded.Count > 0)
-                    {
-                        History.LastSucceeded = History.AllSucceeded[History.AllSucceeded.Count - 1];
-    
-                        foreach (var Failed in History.AllFailed)
-                        {
-                            if (Failed > History.LastSucceeded)
-                            {
-                                History.Failed.Add(Failed);
-                                History.FailedString = GUBPNode.MergeSpaceStrings(History.FailedString, String.Format("{0}", Failed));
-                            }
-                        }
-                        foreach (var Started in History.AllStarted)
-                        {
-                            if (Started > History.LastSucceeded && !History.Failed.Contains(Started))
-                            {
-                                History.InProgress.Add(Started);
-                                History.InProgressString = GUBPNode.MergeSpaceStrings(History.InProgressString, String.Format("{0}", Started));
-                            }
-                        }
-                        GUBPNodesHistory.Add(Node, History);
-                    }
+                    UpdateNodeHistory(Node, CLString);
                 }
             }
         }
@@ -4960,92 +5094,9 @@ if (HostPlatform == UnrealTargetPlatform.Mac) continue; //temp hack till mac aut
             {
                 if (GUBPNodes[NodeToDo].RunInEC() && !NodeIsAlreadyComplete(NodeToDo, LocalOnly)) // if something is already finished, we don't put it into EC  
                 {
-                    string EMails = "";
-                    string FailCauserEMails = "";
-                    bool SendSuccessForGreenAfterRed = false;
-                    if (GUBPNodesHistory.ContainsKey(NodeToDo))
-                    {
-                        var History = GUBPNodesHistory[NodeToDo];
-
-                        ECProps.Add(string.Format("LastGreen/{0}={1}", NodeToDo, History.LastSucceeded));
-                        ECProps.Add(string.Format("RedsSince/{0}={1}", NodeToDo, History.FailedString));
-                        ECProps.Add(string.Format("InProgress/{0}={1}", NodeToDo, History.InProgressString));
-
-
-                        if (History.LastSucceeded > 0 && History.LastSucceeded < P4Env.Changelist)
-                        {
-                            var ChangeRecords = GetChanges(History.LastSucceeded, P4Env.Changelist, History.LastSucceeded);
-                            int NumChanges = 0;
-                            foreach (var Record in ChangeRecords)
-                            {
-                                FailCauserEMails = GUBPNode.MergeSpaceStrings(FailCauserEMails, Record.UserEmail);
-                                if (++NumChanges > 50)
-                                {
-                                    FailCauserEMails = "";
-                                    break;
-                                }
-                            }
-                        }
-                        ECProps.Add(string.Format("FailCausers/{0}={1}", NodeToDo, FailCauserEMails));
-
-                        if (History.LastSucceeded > 0 && History.LastSucceeded < P4Env.Changelist && History.LastFailed > History.LastSucceeded && History.LastFailed < P4Env.Changelist)
-                        {
-                            SendSuccessForGreenAfterRed = ParseParam("CIS");
-                        }
-                    }
-                    else
-                    {
-                        ECProps.Add(string.Format("LastGreen/{0}=0", NodeToDo));
-                        ECProps.Add(string.Format("RedsSince/{0}=", NodeToDo));
-                        ECProps.Add(string.Format("InProgress/{0}=", NodeToDo));
-                    }
-                    {
-                        string EmailOnly = ParseParamValue("EmailOnly");
-                        if (bFake)
-                        {
-                            EMails = "kellan.carr[epic] gil.gribb[epic]";
-                        }
-                        else if (!String.IsNullOrEmpty(EmailOnly))
-                        {
-                            EMails = EmailOnly;
-                        }
-                        else if (!String.IsNullOrEmpty(ForceAllEmailsTo))
-                        {
-                            EMails = ForceAllEmailsTo;
-                        }
-                        else
-                        {                            
-                            EMails = GetEMailListForNode(this, NodeToDo);
-                            if (ParseParam("CIS") && !GUBPNodes[NodeToDo].SendSuccessEmail())
-                            {
-                                EMails = GUBPNode.MergeSpaceStrings(FailCauserEMails, EMails);
-                            }
-                        }
-                        string AddEmails = ParseParamValue("AddEmails");
-                        if (!String.IsNullOrEmpty(AddEmails))
-                        {
-                            EMails = GUBPNode.MergeSpaceStrings(AddEmails, EMails);
-                        }
-
-                        {
-                        EMails = EMails.Trim().Replace("[epic]", "@epicgames.com");
-                            EMails = EMails.Replace("[phosphor]", "@phosphorgames.com");
-                            EMails = EMails.Replace("[pitbull]", "@pitbullstudio.co.uk");
-                        }
-                        
-                        ECProps.Add("FailEmails/" + NodeToDo + "=" + EMails);
-                    }
-                    if (GUBPNodes[NodeToDo].SendSuccessEmail() || SendSuccessForGreenAfterRed)
-                    {
-                        ECProps.Add("SendSuccessEmail/" + NodeToDo + "=1");
-                    }
-                    else
-                    {
-                        ECProps.Add("SendSuccessEmail/" + NodeToDo + "=0");
-                    }
-                    ECProps.Add(string.Format("AgentRequirementString/{0}={1}", NodeToDo, GUBPNodes[NodeToDo].ECAgentString()));
-                    ECProps.Add(string.Format("RequiredMemory/{0}={1}", NodeToDo, GUBPNodes[NodeToDo].AgentMemoryRequirement(this)));
-                    ECProps.Add(string.Format("Timeouts/{0}={1}", NodeToDo, GUBPNodes[NodeToDo].TimeoutInMinutes()));
+                    string EMails;
+                    var NodeProps = GetECPropsForNode(NodeToDo, out EMails);
+                    ECProps.AddRange(NodeProps);
 
                     bool Sticky = GUBPNodes[NodeToDo].IsSticky();
                     bool DoParallel = !Sticky;
@@ -5209,7 +5260,6 @@ if (HostPlatform == UnrealTargetPlatform.Mac) continue; //temp hack till mac aut
                         Args = Args.Replace(String.Format("--parentPath {0} ", ParentPath), String.Format("--parentPath {0} ", NodeParentPath));
                         Args = Args.Replace("UAT_Node_Parallel_AgentShare", "UAT_Node_Parallel_AgentShare3");
                     }
-					ECProps.Add(string.Format("JobStepPath/{0}={1}", NodeToDo, GetJobStepPath(NodeToDo)));
 
                     if (!String.IsNullOrEmpty(PreCondition))
                     {
@@ -5421,6 +5471,7 @@ if (HostPlatform == UnrealTargetPlatform.Mac) continue; //temp hack till mac aut
                 {
                     if (SaveSuccessRecords)
                     {
+                        UpdateHistoryAndEC(NodeToDo, CLString);
                         SaveStatus(NodeToDo, FailedTempStorageSuffix, NodeStoreName, bSaveSharedTempStorage, GameNameIfAny);
                     }
 
@@ -5471,6 +5522,7 @@ if (HostPlatform == UnrealTargetPlatform.Mac) continue; //temp hack till mac aut
                 }
                 if (SaveSuccessRecords) 
                 {
+                    UpdateHistoryAndEC(NodeToDo, CLString);
                     SaveStatus(NodeToDo, SucceededTempStorageSuffix, NodeStoreName, bSaveSharedTempStorage, GameNameIfAny);
                 }
             }
