#ue4
#ios
* fix IPA name when archiving a content only project

diff --git a/Engine/Source/Programs/AutomationTool/IOS/IOSPlatform.Automation.cs b/Engine/Source/Programs/AutomationTool/IOS/IOSPlatform.Automation.cs
--- a/Engine/Source/Programs/AutomationTool/IOS/IOSPlatform.Automation.cs
+++ b/Engine/Source/Programs/AutomationTool/IOS/IOSPlatform.Automation.cs
@@ -488,9 +488,14 @@ public class IOSPlatform : Platform
 			ProjectGameExeFilename = CombinePaths (Path.GetDirectoryName(Params.RawProjectPath), "Binaries", "IOS", Path.GetFileName (Params.ProjectGameExeFilename));
 		}
 		var ProjectIPA = MakeIPAFileName( TargetConfiguration, ProjectGameExeFilename );
+		// rename the .ipa if not code based
 		if (!Params.IsCodeBasedProject)
 		{
-			ProjectIPA = ProjectIPA.Replace("UE4Game", Params.ShortProjectName);
+			ProjectIPA = Path.Combine(Path.GetDirectoryName(Params.RawProjectPath), "Binaries\\IOS", Params.ShortProjectName + ".ipa");
+			if (TargetConfiguration != UnrealTargetConfiguration.Development)
+			{
+				ProjectIPA = Path.Combine(Path.GetDirectoryName(Params.RawProjectPath), "Binaries\\IOS", Params.ShortProjectName + "-" + PlatformType.ToString() + "-" + TargetConfiguration.ToString() + ".ipa");
+			}
 		}
 
 		// verify the .ipa exists
