build script, reduce temp storage to 2 days

diff --git a/Engine/Source/Programs/AutomationTool/TempStorage.cs b/Engine/Source/Programs/AutomationTool/TempStorage.cs
--- a/Engine/Source/Programs/AutomationTool/TempStorage.cs
+++ b/Engine/Source/Programs/AutomationTool/TempStorage.cs
@@ -457,7 +457,7 @@ namespace AutomationTool
                 return;
             }
             TestedForClean.Add(Directory);
-            const int MaximumDaysToKeepTempStorage = 4;
+            const int MaximumDaysToKeepTempStorage = 2;
             DirectoryInfo DirInfo = new DirectoryInfo(Directory);
             var TopLevelDirs = DirInfo.GetDirectories();
             foreach (var TopLevelDir in TopLevelDirs)
