#ue4
#ios
* fix for bad path slashes on Mac

diff --git a/Engine/Source/Programs/AutomationTool/IOS/IOSPlatform.Automation.cs b/Engine/Source/Programs/AutomationTool/IOS/IOSPlatform.Automation.cs
--- a/Engine/Source/Programs/AutomationTool/IOS/IOSPlatform.Automation.cs
+++ b/Engine/Source/Programs/AutomationTool/IOS/IOSPlatform.Automation.cs
@@ -154,10 +154,10 @@ public class IOSPlatform : Platform
 			// rename the .ipa if not code based
 			if (!Params.IsCodeBasedProject)
 			{
-				ProjectIPA = Path.Combine(Path.GetDirectoryName(Params.RawProjectPath), "Binaries\\IOS", Params.ShortProjectName + ".ipa");
+				ProjectIPA = Path.Combine(Path.GetDirectoryName(Params.RawProjectPath), "Binaries", "IOS", Params.ShortProjectName + ".ipa");
 				if (TargetConfiguration != UnrealTargetConfiguration.Development)
 				{
-					ProjectIPA = Path.Combine(Path.GetDirectoryName(Params.RawProjectPath), "Binaries\\IOS", Params.ShortProjectName + "-" + PlatformType.ToString() + "-" + TargetConfiguration.ToString() + ".ipa");
+					ProjectIPA = Path.Combine(Path.GetDirectoryName(Params.RawProjectPath), "Binaries", "IOS", Params.ShortProjectName + "-" + PlatformType.ToString() + "-" + TargetConfiguration.ToString() + ".ipa");
 				}
 			}
 
@@ -491,10 +491,10 @@ public class IOSPlatform : Platform
 		// rename the .ipa if not code based
 		if (!Params.IsCodeBasedProject)
 		{
-			ProjectIPA = Path.Combine(Path.GetDirectoryName(Params.RawProjectPath), "Binaries\\IOS", Params.ShortProjectName + ".ipa");
+			ProjectIPA = Path.Combine(Path.GetDirectoryName(Params.RawProjectPath), "Binaries", "IOS", Params.ShortProjectName + ".ipa");
 			if (TargetConfiguration != UnrealTargetConfiguration.Development)
 			{
-				ProjectIPA = Path.Combine(Path.GetDirectoryName(Params.RawProjectPath), "Binaries\\IOS", Params.ShortProjectName + "-" + PlatformType.ToString() + "-" + TargetConfiguration.ToString() + ".ipa");
+				ProjectIPA = Path.Combine(Path.GetDirectoryName(Params.RawProjectPath), "Binaries", "IOS", Params.ShortProjectName + "-" + PlatformType.ToString() + "-" + TargetConfiguration.ToString() + ".ipa");
 			}
 		}
 
@@ -522,10 +522,10 @@ public class IOSPlatform : Platform
 			// rename the .ipa if not code based
 			if (!Params.IsCodeBasedProject)
 			{
-				ProjectIPA = Path.Combine(Path.GetDirectoryName(Params.RawProjectPath), "Binaries\\IOS", Params.ShortProjectName + ".ipa");
+				ProjectIPA = Path.Combine(Path.GetDirectoryName(Params.RawProjectPath), "Binaries", "IOS", Params.ShortProjectName + ".ipa");
 				if (TargetConfiguration != UnrealTargetConfiguration.Development)
 				{
-					ProjectIPA = Path.Combine(Path.GetDirectoryName(Params.RawProjectPath), "Binaries\\IOS", Params.ShortProjectName + "-" + PlatformType.ToString() + "-" + TargetConfiguration.ToString() + ".ipa");
+					ProjectIPA = Path.Combine(Path.GetDirectoryName(Params.RawProjectPath), "Binaries", "IOS", Params.ShortProjectName + "-" + PlatformType.ToString() + "-" + TargetConfiguration.ToString() + ".ipa");
 				}
 			}
 
