Fix for TTP #332331 CRITICAL: Mac: Launching from Editor or UFE fails

diff --git a/Engine/Source/Programs/AutomationTool/Mac/MacPlatform.Automation.cs b/Engine/Source/Programs/AutomationTool/Mac/MacPlatform.Automation.cs
--- a/Engine/Source/Programs/AutomationTool/Mac/MacPlatform.Automation.cs
+++ b/Engine/Source/Programs/AutomationTool/Mac/MacPlatform.Automation.cs
@@ -38,11 +38,11 @@ public class MacPlatform : Platform
 		{
 			if (Exe.StartsWith(CombinePaths(SC.RuntimeProjectRootDir, "Binaries", SC.PlatformDir)))
 			{
-				StageAppBundle(SC, CombinePaths(SC.ProjectRoot, "Binaries", SC.PlatformDir, Path.GetFileNameWithoutExtension(Exe) + ".app"), SC.ShortProjectName + ".app");
+				StageAppBundle(SC, CombinePaths(SC.ProjectRoot, "Binaries", SC.PlatformDir, Path.GetFileNameWithoutExtension(Exe) + ".app"), CombinePaths(SC.ShortProjectName, "Binaries", SC.PlatformDir, Path.GetFileNameWithoutExtension(Exe) + ".app"));
 			}
 			else if (Exe.StartsWith(CombinePaths(SC.RuntimeRootDir, "Engine/Binaries", SC.PlatformDir)))
 			{
-				StageAppBundle(SC, CombinePaths(SC.LocalRoot, "Engine/Binaries", SC.PlatformDir, Path.GetFileNameWithoutExtension(Exe) + ".app"), SC.ShortProjectName + ".app");
+				StageAppBundle(SC, CombinePaths(SC.LocalRoot, "Engine/Binaries", SC.PlatformDir, Path.GetFileNameWithoutExtension(Exe) + ".app"), CombinePaths("Engine/Binaries", SC.PlatformDir, Path.GetFileNameWithoutExtension(Exe) + ".app"));
 			}
 		}
 
@@ -58,7 +58,18 @@ public class MacPlatform : Platform
 
 	public override void ProcessArchivedProject(ProjectParams Params, DeploymentContext SC)
 	{
-		string TargetPath = CombinePaths(SC.ArchiveDirectory, SC.ShortProjectName + ".app", "Contents", "UE4");
+		string BundlePath = CombinePaths(SC.ArchiveDirectory, SC.ShortProjectName + ".app");
+		if (!Directory.Exists(BundlePath))
+		{
+			string SourceBundlePath = CombinePaths(SC.ArchiveDirectory, SC.ShortProjectName, "Binaries", "Mac", SC.ShortProjectName + ".app");
+			if (!Directory.Exists(BundlePath))
+			{
+				SourceBundlePath = CombinePaths(SC.ArchiveDirectory, "Engine", "Binaries", "Mac", "UE4Game.app");
+			}
+			Directory.Move(SourceBundlePath, BundlePath);
+		}
+
+		string TargetPath = CombinePaths(BundlePath, "Contents", "UE4");
 		if (DirectoryExists(TargetPath))
 		{
 			Directory.Delete(TargetPath, true);
