Deploy Steam .so and libopenal.so.1 for every project on Linux (TTP #331888 and possibly TTP #331978).

diff --git a/Engine/Source/Programs/AutomationTool/Linux/LinuxPlatform.Automation.cs b/Engine/Source/Programs/AutomationTool/Linux/LinuxPlatform.Automation.cs
--- a/Engine/Source/Programs/AutomationTool/Linux/LinuxPlatform.Automation.cs
+++ b/Engine/Source/Programs/AutomationTool/Linux/LinuxPlatform.Automation.cs
@@ -30,13 +30,18 @@ public abstract class BaseLinuxPlatform : Platform
             SC.StageFiles(StagedFileType.NonUFS, CombinePaths(SC.LocalRoot, "Engine/Binaries", SC.PlatformDir), "CrashReportClient", false);
         }
 
-        if (Params.bUsesSteam)
+        // assume that we always have to deploy Steam (FIXME: should be automatic)
         {
             string SteamVersion = "Steamv128";
             SC.StageFiles(StagedFileType.NonUFS, CombinePaths(SC.LocalRoot, "Engine/Binaries/ThirdParty/Steamworks/" + SteamVersion, SC.PlatformDir), "libsteam_api.so", false, null, CombinePaths("Engine/Binaries", SC.PlatformDir));
             SC.StageFiles(StagedFileType.NonUFS, CombinePaths(SC.LocalRoot, "Engine/Config"), "controller.vdf", false, null, CommandUtils.CombinePaths(SC.RelativeProjectRootForStage, "Saved/Config"));
         }
 
+        // assume that we always have to deploy OpenAL (FIXME: should be automatic)
+        {
+            SC.StageFiles(StagedFileType.NonUFS, CombinePaths(SC.LocalRoot, "Engine/Binaries/ThirdParty/OpenAL/", SC.PlatformDir), "libopenal.so.1", false, null, CombinePaths("Engine/Binaries", SC.PlatformDir));
+        }
+
         List<string> Exes = GetExecutableNames(SC);
 
         foreach (var Exe in Exes)
