Fix for TTP #332999 CRITICAL: Mac: Regression: PACKAGING: Fails for code based projects

diff --git a/Engine/Source/Programs/AutomationTool/Mac/MacPlatform.Automation.cs b/Engine/Source/Programs/AutomationTool/Mac/MacPlatform.Automation.cs
--- a/Engine/Source/Programs/AutomationTool/Mac/MacPlatform.Automation.cs
+++ b/Engine/Source/Programs/AutomationTool/Mac/MacPlatform.Automation.cs
@@ -62,7 +62,7 @@ public class MacPlatform : Platform
 		if (!Directory.Exists(BundlePath))
 		{
 			string SourceBundlePath = CombinePaths(SC.ArchiveDirectory, SC.ShortProjectName, "Binaries", "Mac", SC.ShortProjectName + ".app");
-			if (!Directory.Exists(BundlePath))
+			if (!Directory.Exists(SourceBundlePath))
 			{
 				SourceBundlePath = CombinePaths(SC.ArchiveDirectory, "Engine", "Binaries", "Mac", "UE4Game.app");
 			}
