One more Mac packaging fix

diff --git a/Engine/Source/Programs/AutomationTool/Mac/MacPlatform.Automation.cs b/Engine/Source/Programs/AutomationTool/Mac/MacPlatform.Automation.cs
--- a/Engine/Source/Programs/AutomationTool/Mac/MacPlatform.Automation.cs
+++ b/Engine/Source/Programs/AutomationTool/Mac/MacPlatform.Automation.cs
@@ -99,7 +99,7 @@ public class MacPlatform : Platform
 		// Update executable name and entry in Info.plist
 		string ExeName = SC.StageExecutables[0];
 		string UE4GamePath = CombinePaths(SC.ArchiveDirectory, SC.ShortProjectName + ".app", "Contents", "MacOS", ExeName);
-		if (File.Exists(UE4GamePath))
+		if (ExeName != SC.ShortProjectName && File.Exists(UE4GamePath))
 		{
 			string GameExePath = CombinePaths(SC.ArchiveDirectory, SC.ShortProjectName + ".app", "Contents", "MacOS", SC.ShortProjectName);
 			File.Delete(GameExePath);
