Improved logging in new Launcher promotion process to include more detail.  Also removed McpConfig spew.

diff --git a/Engine/Source/Programs/AutomationTool/MCPPublic.cs b/Engine/Source/Programs/AutomationTool/MCPPublic.cs
--- a/Engine/Source/Programs/AutomationTool/MCPPublic.cs
+++ b/Engine/Source/Programs/AutomationTool/MCPPublic.cs
@@ -548,7 +548,6 @@ namespace EpicGames.MCP.Config
                 throw new AutomationException("Unable to find requested McpConfig: {0}", ConfigName);
             }
             CommandUtils.Log("Found config for {0}", ConfigName);
-            LoadedConfig.SpewValues();
             return LoadedConfig;
         }
     }
