Moving Loc file staging out of cases other than Windows/Linux #codereview Daniel.Lamb, Justin.Sargent

diff --git a/Engine/Source/Programs/AutomationTool/Scripts/CopyBuildToStagingDirectory.Automation.cs b/Engine/Source/Programs/AutomationTool/Scripts/CopyBuildToStagingDirectory.Automation.cs
--- a/Engine/Source/Programs/AutomationTool/Scripts/CopyBuildToStagingDirectory.Automation.cs
+++ b/Engine/Source/Programs/AutomationTool/Scripts/CopyBuildToStagingDirectory.Automation.cs
@@ -181,7 +181,6 @@ public partial class Project : CommandUtils
 
 			}
             SC.StageFiles(StagedFileType.UFS, CombinePaths(SC.LocalRoot, "Engine/Content/Localization"), "*.dat", true, null, null, false, !Params.Pak);
-            SC.StageFiles(StagedFileType.NonUFS, CombinePaths(SC.LocalRoot, "Engine/Content/Localization"), "*.dat", true, null, null, false, !Params.Pak);
             SC.StageFiles(StagedFileType.UFS, CombinePaths(SC.LocalRoot, "Engine/Content/Localization/Engine"), "*.locres", true, null, null, false, !Params.Pak);
 			SC.StageFiles(StagedFileType.UFS, CombinePaths(SC.LocalRoot, "Engine/Plugins"), "*.uplugin", true, null, null, true, !Params.Pak);
 
@@ -208,6 +207,8 @@ public partial class Project : CommandUtils
 			// @todo Make CrashReportClient more portable so we don't have to do this
 			if (SC.bStageCrashReporter && !SC.DedicatedServer)
 			{
+                //If the .dat file needs to be staged as NonUFS for non-Windows/Linux hosts we need to change the casing as we do with the build properties file above.
+                SC.StageFiles(StagedFileType.NonUFS, CombinePaths(SC.LocalRoot, "Engine/Content/Localization"), "*.dat", true, null, null, false, !Params.Pak);
 				SC.StageFiles(StagedFileType.NonUFS, CombinePaths(SC.LocalRoot, "Engine/Content/Slate"));
 				SC.StageFiles(StagedFileType.NonUFS, CombinePaths(SC.LocalRoot, "Engine/Shaders/StandaloneRenderer"));
 			}
