- Fix problems deploying games to console with a space in the name [ttp 333160, 333164]

diff --git a/Engine/Source/Programs/AutomationTool/Scripts/RunProjectCommand.Automation.cs b/Engine/Source/Programs/AutomationTool/Scripts/RunProjectCommand.Automation.cs
--- a/Engine/Source/Programs/AutomationTool/Scripts/RunProjectCommand.Automation.cs
+++ b/Engine/Source/Programs/AutomationTool/Scripts/RunProjectCommand.Automation.cs
@@ -800,7 +800,7 @@ public partial class Project : CommandUtils
 		{
 			ClientCmdLine = "-run=Launch ";
 			ClientCmdLine += "-Device=" + Params.Device + " ";
-			ClientCmdLine += "-Exe=" + ClientApp + " ";
+			ClientCmdLine += "-Exe=\"" + ClientApp + "\" ";
 			ClientCmdLine += "-Targetplatform=" + Params.ClientTargetPlatforms[0].ToString() + " ";
 			ClientCmdLine += "-Params=\"" + TempCmdLine + "\"";
 			ClientApp = CombinePaths(CmdEnv.LocalRoot, "Engine/Binaries/Win64/UnrealFrontend.exe");
