Fix for TTP #333180 CRITICAL: Unable to open packaged project

diff --git a/Engine/Source/Programs/AutomationTool/Mac/MacPlatform.Automation.cs b/Engine/Source/Programs/AutomationTool/Mac/MacPlatform.Automation.cs
--- a/Engine/Source/Programs/AutomationTool/Mac/MacPlatform.Automation.cs
+++ b/Engine/Source/Programs/AutomationTool/Mac/MacPlatform.Automation.cs
@@ -48,6 +48,8 @@ public class MacPlatform : Platform
 
 		// Copy the splash screen, Mac specific
 		SC.StageFiles(StagedFileType.NonUFS, CombinePaths(SC.ProjectRoot, "Content/Splash"), "Splash.bmp", false, null, null, true);
+
+		SC.StageFiles(StagedFileType.NonUFS, CombinePaths(SC.LocalRoot, "Engine/Content/Localization"), "*.dat", true, null, null, false, !Params.Pak);
 	}
 
 	public override void Package(ProjectParams Params, DeploymentContext SC, int WorkingCL)
diff --git a/Engine/Source/Programs/AutomationTool/Win/WinPlatform.Automation.cs b/Engine/Source/Programs/AutomationTool/Win/WinPlatform.Automation.cs
--- a/Engine/Source/Programs/AutomationTool/Win/WinPlatform.Automation.cs
+++ b/Engine/Source/Programs/AutomationTool/Win/WinPlatform.Automation.cs
@@ -98,6 +98,8 @@ public abstract class BaseWinPlatform : Platform
         // Copy the splash screen, windows specific
         SC.StageFiles(StagedFileType.NonUFS, CombinePaths(SC.ProjectRoot, "Content/Splash"), "Splash.bmp", false, null, null, true);
 
+		SC.StageFiles(StagedFileType.NonUFS, CombinePaths(SC.LocalRoot, "Engine/Content/Localization"), "*.dat", true, null, null, false, !Params.Pak);
+
 		List<string> Exes = GetExecutableNames(SC);
 
         // the first exe is the "main" one, the rest are marked as debug files
