AutomationTool changes for source publishing
- Added better support for 'verbose' log output from P4 and Process commands (SpewIsVerbose arg and TraceEventType flags)
- Added new P4Connection.DescribeChangelists() that provides full changed file lists and edit types
- Added even better support for full changelist descriptions to P4Connection.Changes() function

diff --git a/Engine/Source/Programs/AutomationTool/P4Utils.cs b/Engine/Source/Programs/AutomationTool/P4Utils.cs
--- a/Engine/Source/Programs/AutomationTool/P4Utils.cs
+++ b/Engine/Source/Programs/AutomationTool/P4Utils.cs
@@ -419,10 +419,15 @@ namespace AutomationTool
 		/// <param name="Input">Stdin</param>
 		/// <param name="AllowSpew">true for spew</param>
 		/// <returns>Exit code</returns>
-        public ProcessResult P4(string CommandLine, string Input = null, bool AllowSpew = true, bool WithClient = true)
+        public ProcessResult P4(string CommandLine, string Input = null, bool AllowSpew = true, bool WithClient = true, bool SpewIsVerbose = false)
 		{
 			CheckP4Enabled();
-            return CommandUtils.Run(HostPlatform.Current.P4Exe, (WithClient ? GlobalOptions : GlobalOptionsWithoutClient) + CommandLine, Input, AllowSpew ? CommandUtils.ERunOptions.AllowSpew : CommandUtils.ERunOptions.NoLoggingOfRunCommand);
+			CommandUtils.ERunOptions RunOptions = AllowSpew ? CommandUtils.ERunOptions.AllowSpew : CommandUtils.ERunOptions.NoLoggingOfRunCommand;
+			if( SpewIsVerbose )
+			{
+				RunOptions |= CommandUtils.ERunOptions.SpewIsVerbose;
+			}
+            return CommandUtils.Run(HostPlatform.Current.P4Exe, (WithClient ? GlobalOptions : GlobalOptionsWithoutClient) + CommandLine, Input, Options:RunOptions);
 		}
 
 		/// <summary>
@@ -450,11 +455,11 @@ namespace AutomationTool
 		/// <param name="CommandLine">Commandline to pass to p4.</param>
 		/// <param name="Input">Stdin input.</param>
 		/// <param name="AllowSpew">Whether the command is allowed to spew.</param>
-        public void LogP4(string CommandLine, string Input = null, bool AllowSpew = true, bool WithClient = true)
+        public void LogP4(string CommandLine, string Input = null, bool AllowSpew = true, bool WithClient = true, bool SpewIsVerbose = false)
 		{
 			CheckP4Enabled();
 			string Output;
-            if (!LogP4Output(out Output, CommandLine, Input, AllowSpew, WithClient))
+            if (!LogP4Output(out Output, CommandLine, Input, AllowSpew, WithClient, SpewIsVerbose:SpewIsVerbose))
 			{
 				throw new P4Exception("p4.exe {0} failed.", CommandLine);
 			}
@@ -468,7 +473,7 @@ namespace AutomationTool
 		/// <param name="Input">Stdin input.</param>
 		/// <param name="AllowSpew">Whether the command should spew.</param>
 		/// <returns>True if succeeded, otherwise false.</returns>
-        public bool LogP4Output(out string Output, string CommandLine, string Input = null, bool AllowSpew = true, bool WithClient = true)
+        public bool LogP4Output(out string Output, string CommandLine, string Input = null, bool AllowSpew = true, bool WithClient = true, bool SpewIsVerbose = false)
 		{
 			CheckP4Enabled();
 			Output = "";
@@ -479,7 +484,7 @@ namespace AutomationTool
 				return false;
 			}
 
-            var Result = P4(CommandLine, Input, AllowSpew, WithClient);
+            var Result = P4(CommandLine, Input, AllowSpew, WithClient, SpewIsVerbose:SpewIsVerbose);
 
 			CommandUtils.WriteToFile(LogPath, CommandLine + "\n");
 			CommandUtils.WriteToFile(LogPath, Result.Output);
@@ -569,12 +574,17 @@ namespace AutomationTool
             try
             {
                 // Change 1999345 on 2014/02/16 by buildmachine@BuildFarm_BUILD-23_buildmachine_++depot+UE4 'GUBP Node Shadow_LabelPromotabl'
+
+				// If the user specified '-l' or '-L', the summary will appear on subsequent lines (no quotes) instead of the same line (surrounded by single quotes)
+				bool bSummaryIsOnSameLine = CommandLine.IndexOf( "-L", StringComparison.InvariantCultureIgnoreCase ) == -1;
+
                 string Output;
                 if (!LogP4Output(out Output, "changes " + (LongComment ? "-L " : "") + CommandLine, null, AllowSpew))
                 {
                     return false;
                 }
-                var TempLines = Output.Split(new string[] { Environment.NewLine }, StringSplitOptions.RemoveEmptyEntries);
+
+                var TempLines = Output.Split(new string[] { Environment.NewLine }, StringSplitOptions.None);
                 var Lines = new List<string>();
                 if (LongComment)
                 {
@@ -610,8 +620,16 @@ namespace AutomationTool
                         Lines.Add(Line);
                     }
                 }
-                foreach (var Line in Lines)
+                for(int LineIndex = 0; LineIndex < Lines.Count; ++LineIndex)
                 {
+					var Line = Lines[ LineIndex ];
+
+					// If we've hit a blank line, then we're done
+					if( String.IsNullOrEmpty( Line ) )
+					{
+						break;
+					}
+
                     ChangeRecord Change = new ChangeRecord();
                     string MatchChange = "Change ";
                     string MatchOn = " on "; 
@@ -621,9 +639,7 @@ namespace AutomationTool
                     int OnAt = Line.IndexOf(MatchOn);
                     int ByAt = Line.IndexOf(MatchBy);
                     int AtAt = Line.IndexOf("@");
-                    int TickAt = Line.IndexOf("'");
-                    int EndTick = Line.LastIndexOf("'");
-                    if (ChangeAt >= 0 && OnAt > ChangeAt && ByAt > OnAt && TickAt > ByAt && EndTick > TickAt)
+                    if (ChangeAt == 0 && OnAt > ChangeAt && ByAt > OnAt)
                     {
                         var ChangeString = Line.Substring(ChangeAt + MatchChange.Length, OnAt - ChangeAt - MatchChange.Length);
                         Change.CL = int.Parse(ChangeString);
@@ -632,10 +648,67 @@ namespace AutomationTool
                             throw new AutomationException("weird CL {0} in {1}", Change.CL, Line);
                         }
                         Change.User = Line.Substring(ByAt + MatchBy.Length, AtAt - ByAt - MatchBy.Length);
-                        Change.Summary = Line.Substring(TickAt + 1, EndTick - TickAt - 1);
+
+						if( bSummaryIsOnSameLine )
+						{ 
+							int TickAt = Line.IndexOf("'");
+							int EndTick = Line.LastIndexOf("'");
+							if( TickAt > ByAt && EndTick > TickAt )
+							{ 
+								Change.Summary = Line.Substring(TickAt + 1, EndTick - TickAt - 1);
+							}
+						}
+						else
+						{
+							++LineIndex;
+							if( LineIndex >= Lines.Count )
+							{
+								throw new AutomationException("Was expecting a change summary to appear after Change header output from P4, but there were no more lines to read");
+							}
+
+							Line = Lines[ LineIndex ];
+							if( !String.IsNullOrEmpty( Line ) )
+							{
+								throw new AutomationException("Was expecting blank line after Change header output from P4");
+							}
+
+							++LineIndex;
+							for( ; LineIndex < Lines.Count; ++LineIndex )
+							{
+								Line = Lines[ LineIndex ];
+
+								if( String.IsNullOrEmpty( Line ) )
+								{
+									// Summaries end with a blank line (no tabs)
+									break;
+								}
+
+								// Summary lines are supposed to begin with a single tab character (even empty lines)
+								if( Line[0] != '\t' )
+								{
+									throw new AutomationException("Was expecting every line of the P4 changes summary to start with a tab character");
+								}
+
+								// Remove the tab
+								var SummaryLine = Line.Substring( 1 );
+
+								// Add a CR if we already had some summary text
+								if( !String.IsNullOrEmpty( Change.Summary ) )
+								{
+									Change.Summary += "\n";
+								}
+
+								// Append the summary line!
+								Change.Summary += SummaryLine;
+							}
+						}
                         Change.UserEmail = UserToEmail(Change.User);
                         ChangeRecords.Add(Change);
                     }
+					else
+					{
+						throw new AutomationException("Output of 'p4 changes' was not formatted how we expected.  Could not find 'Change', 'on' and 'by' in the output line: " + Line);
+					}
                 }
             }
             catch (Exception)
@@ -643,17 +716,244 @@ namespace AutomationTool
                 return false;
             }
             ChangeRecords.Sort((A, B) => ChangeRecord.Compare(A, B));
-            ChangesCache.Add(CommandLine, ChangeRecords);
+			if( ChangesCache.ContainsKey(CommandLine) )
+			{
+				ChangesCache[CommandLine] = ChangeRecords;
+			}
+			else
+			{ 
+				ChangesCache.Add(CommandLine, ChangeRecords);
+			}
+            return true;
+        }
+
+	
+        public class DescribeRecord
+        {
+            public int CL = 0;
+            public string User = "";
+            public string UserEmail = "";
+            public string Summary = "";
+			
+			public class DescribeFile
+			{
+				public string File;
+				public int Revision;
+				public string ChangeType;
+			}
+			public List<DescribeFile> Files = new List<DescribeFile>();
+            
+			public static int Compare(DescribeRecord A, DescribeRecord B)
+            {
+                return (A.CL < B.CL) ? -1 : (A.CL > B.CL) ? 1 : 0;
+            }
+        }
+
+		/// <summary>
+		/// Wraps P4 describe
+		/// </summary>
+		/// <param name="Changelists">List of changelist numbers to query full descriptions for</param>
+		/// <param name="DescribeRecords">List of records we found.  One for each changelist number.  These will be sorted from oldest to newest.</param>
+		/// <param name="AllowSpew"></param>
+		/// <returns>True if everything went okay</returns>
+        public bool DescribeChangelists(List<int> Changelists, out List<DescribeRecord> DescribeRecords, bool AllowSpew = true)
+        {
+			DescribeRecords = new List<DescribeRecord>();
+            CheckP4Enabled();
+            try
+            {
+				// Change 234641 by This.User@WORKSPACE-C2Q-67_Dev on 2008/05/06 10:32:32
+				// 
+				//         Desc Line 1
+				// 
+				// Affected files ...
+				// 
+				// ... //depot/UnrealEngine3/Development/Src/Engine/Classes/ArrowComponent.uc#8 edit
+				// ... //depot/UnrealEngine3/Development/Src/Engine/Classes/DecalActorBase.uc#4 edit
+
+
+				string Output;
+				string CommandLine = "-s";		// Don't automatically diff the files
+				
+				// Add changelists to the command-line
+				foreach( var Changelist in Changelists )
+				{
+					CommandLine += " " + Changelist.ToString();
+				}
+
+                if (!LogP4Output(out Output, "describe " + CommandLine, null, AllowSpew))
+                {
+                    return false;
+                }
+
+				int ChangelistIndex = 0;
+				var Lines = Output.Split(new string[] { Environment.NewLine }, StringSplitOptions.None);
+                for (var LineIndex = 0; LineIndex < Lines.Length; ++LineIndex)
+                {
+					var Line = Lines[ LineIndex ];
+
+					// If we've hit a blank line, then we're done
+					if( String.IsNullOrEmpty( Line ) )
+					{
+						break;
+					}
+
+                    string MatchChange = "Change ";
+                    string MatchOn = " on "; 
+                    string MatchBy = " by ";
+
+                    int ChangeAt = Line.IndexOf(MatchChange);
+                    int OnAt = Line.IndexOf(MatchOn);
+                    int ByAt = Line.IndexOf(MatchBy);
+                    int AtAt = Line.IndexOf("@");
+                    if (ChangeAt == 0 && OnAt > ChangeAt && ByAt < OnAt)
+                    {
+                        var ChangeString = Line.Substring(ChangeAt + MatchChange.Length, ByAt - ChangeAt - MatchChange.Length);
+
+						var CurrentChangelist = Changelists[ ChangelistIndex++ ];
+
+                        if (!ChangeString.Equals( CurrentChangelist.ToString()))
+                        {
+                            throw new AutomationException("Was expecting changelists to be reported back in the same order we asked for them (CL {0} != {1})", ChangeString, CurrentChangelist.ToString());
+                        }
+
+						var DescribeRecord = new DescribeRecord();
+						DescribeRecords.Add( DescribeRecord );
+
+						DescribeRecord.CL = CurrentChangelist;
+                        DescribeRecord.User = Line.Substring(ByAt + MatchBy.Length, AtAt - ByAt - MatchBy.Length);
+
+						++LineIndex;
+						if( LineIndex >= Lines.Length )
+						{
+							throw new AutomationException("Was expecting a change summary to appear after Change header output from P4, but there were no more lines to read");
+						}
+
+						Line = Lines[ LineIndex ];
+						if( !String.IsNullOrEmpty( Line ) )
+						{
+							throw new AutomationException("Was expecting blank line after Change header output from P4");
+						}
+
+						// Summary
+						++LineIndex;
+						for( ; LineIndex < Lines.Length; ++LineIndex )
+						{
+							Line = Lines[ LineIndex ];
+
+							if( String.IsNullOrEmpty( Line ) )
+							{
+								// Summaries end with a blank line (no tabs)
+								break;
+							}
+
+							// Summary lines are supposed to begin with a single tab character (even empty lines)
+							if( Line[0] != '\t' )
+							{
+								throw new AutomationException("Was expecting every line of the P4 changes summary to start with a tab character");
+							}
+
+							// Remove the tab
+							var SummaryLine = Line.Substring( 1 );
+
+							// Add a CR if we already had some summary text
+							if( !String.IsNullOrEmpty( DescribeRecord.Summary ) )
+							{
+								DescribeRecord.Summary += "\n";
+							}
+
+							// Append the summary line!
+							DescribeRecord.Summary += SummaryLine;
+						}
+
+
+						++LineIndex;
+						if( LineIndex >= Lines.Length )
+						{
+							throw new AutomationException("Was expecting 'Affected files' to appear after the summary output from P4, but there were no more lines to read");
+						}
+
+						Line = Lines[ LineIndex ];
+
+						string MatchAffectedFiles = "Affected files";
+						int AffectedFilesAt = Line.IndexOf(MatchAffectedFiles);
+						if( AffectedFilesAt == 0 )
+						{
+							++LineIndex;
+							if( LineIndex >= Lines.Length )
+							{
+								throw new AutomationException("Was expecting a list of files to appear after Affected Files header output from P4, but there were no more lines to read");
+							}
+
+							Line = Lines[ LineIndex ];
+							if( !String.IsNullOrEmpty( Line ) )
+							{
+								throw new AutomationException("Was expecting blank line after Affected Files header output from P4");
+							}
+
+							// Files
+							++LineIndex;
+							for( ; LineIndex < Lines.Length; ++LineIndex )
+							{
+								Line = Lines[ LineIndex ];
+
+								if( String.IsNullOrEmpty( Line ) )
+								{
+									// Summaries end with a blank line (no tabs)
+									break;
+								}
+
+								// File lines are supposed to begin with a "... " string
+								if( !Line.StartsWith( "... " ) )
+								{
+									throw new AutomationException("Was expecting every line of the P4 describe files to start with a tab character");
+								}
+
+								// Remove the "... " prefix
+								var FilesLine = Line.Substring( 4 );
+
+								var DescribeFile = new DescribeRecord.DescribeFile();
+								DescribeRecord.Files.Add( DescribeFile );
+ 							
+								// Find the revision #
+								var RevisionNumberAt = FilesLine.LastIndexOf( "#" ) + 1;
+								var ChangeTypeAt = 1 + FilesLine.IndexOf( " ", RevisionNumberAt );
+							
+								DescribeFile.File = FilesLine.Substring( 0, RevisionNumberAt - 1 );
+								string RevisionString = FilesLine.Substring( RevisionNumberAt, ChangeTypeAt - RevisionNumberAt );
+								DescribeFile.Revision = int.Parse( RevisionString );
+								DescribeFile.ChangeType = FilesLine.Substring( ChangeTypeAt );															  
+							}
+						}
+						else
+						{
+							throw new AutomationException("Output of 'p4 describe' was not formatted how we expected.  Could not find 'Affected files' in the output line: " + Line);
+						}
+
+                        DescribeRecord.UserEmail = UserToEmail(DescribeRecord.User);
+                    }
+					else
+					{
+						throw new AutomationException("Output of 'p4 describe' was not formatted how we expected.  Could not find 'Change', 'on' and 'by' in the output line: " + Line);
+					}
+                }
+            }
+            catch (Exception)
+            {
+                return false;
+            }
+            DescribeRecords.Sort((A, B) => DescribeRecord.Compare(A, B));
             return true;
         }
+
 		/// <summary>
 		/// Invokes p4 sync command.
 		/// </summary>
 		/// <param name="CommandLine">CommandLine to pass on to the command.</param>
-        public void Sync(string CommandLine, bool AllowSpew = true)
+        public void Sync(string CommandLine, bool AllowSpew = true, bool SpewIsVerbose = false)
 		{
 			CheckP4Enabled();
-			LogP4("sync " + CommandLine, null, AllowSpew);
+			LogP4("sync " + CommandLine, null, AllowSpew, SpewIsVerbose:SpewIsVerbose);
 		}
 
 		/// <summary>
diff --git a/Engine/Source/Programs/AutomationTool/ProcessUtils.cs b/Engine/Source/Programs/AutomationTool/ProcessUtils.cs
--- a/Engine/Source/Programs/AutomationTool/ProcessUtils.cs
+++ b/Engine/Source/Programs/AutomationTool/ProcessUtils.cs
@@ -44,7 +44,7 @@ namespace AutomationTool
 		/// Creates a new process and adds it to the tracking list.
 		/// </summary>
 		/// <returns>New Process objects</returns>
-		public static ProcessResult CreateProcess(bool bAllowSpew, string LogName, Dictionary<string, string> Env = null)
+		public static ProcessResult CreateProcess(bool bAllowSpew, string LogName, Dictionary<string, string> Env = null, TraceEventType SpewVerbosity = TraceEventType.Information)
 		{
 			var NewProcess = HostPlatform.Current.CreateProcess(LogName);
 			if (Env != null)
@@ -61,7 +61,7 @@ namespace AutomationTool
 					}
 				}
 			}
-			var Result = new ProcessResult(NewProcess, bAllowSpew, LogName);
+			var Result = new ProcessResult(NewProcess, bAllowSpew, LogName, SpewVerbosity:SpewVerbosity);
 			NewProcess.Exited += NewProcess_Exited;
 			lock (SyncObject)
 			{
@@ -195,17 +195,19 @@ namespace AutomationTool
 		int ProcessExitCode = -1;
 		StringBuilder ProcessOutput = new StringBuilder();
 		bool AllowSpew = true;
+		TraceEventType SpewVerbosity = TraceEventType.Information;
 		private Process Proc = null;
 		private AutoResetEvent OutputWaitHandle = new AutoResetEvent(false);
 		private AutoResetEvent ErrorWaitHandle = new AutoResetEvent(false);
 		private object ProcSyncObject;
 
-		public ProcessResult(Process InProc, bool bAllowSpew, string LogName)
+		public ProcessResult(Process InProc, bool bAllowSpew, string LogName, TraceEventType SpewVerbosity = TraceEventType.Information)
 		{
 			ProcSyncObject = new object();
 			Proc = InProc;
 			Source = LogName;
 			AllowSpew = bAllowSpew;
+			this.SpewVerbosity = SpewVerbosity;
 		}
 
 		private void LogOutput(TraceEventType Verbosity, string Message)
@@ -229,7 +231,7 @@ namespace AutomationTool
 			{
 				if (AllowSpew)
 				{
-					LogOutput(TraceEventType.Information, e.Data);
+					LogOutput(SpewVerbosity, e.Data);
 				}
 
 				ProcessOutput.Append(e.Data);
@@ -252,7 +254,7 @@ namespace AutomationTool
 			{
 				if (AllowSpew)
 				{
-                    LogOutput(TraceEventType.Information, e.Data);
+                    LogOutput(SpewVerbosity, e.Data);
 				}
 
 				ProcessOutput.Append(e.Data);
@@ -530,6 +532,9 @@ namespace AutomationTool
             NoLoggingOfRunCommand = 1 << 4,
             UTF8Output = 1 << 5,
 
+			/// When specified with AllowSpew, the output will be TraceEventType.Verbose instead of TraceEventType.Information
+			SpewIsVerbose = 1 << 6,
+
 			Default = AllowSpew | AppMustExist,
 		}
 
@@ -557,12 +562,12 @@ namespace AutomationTool
 			}
 			var StartTime = DateTime.UtcNow;
 
+			TraceEventType SpewVerbosity = Options.HasFlag(ERunOptions.SpewIsVerbose) ? TraceEventType.Verbose : TraceEventType.Information;
             if (!Options.HasFlag(ERunOptions.NoLoggingOfRunCommand))
             {
-                Log("Run: " + App + " " + (String.IsNullOrEmpty(CommandLine) ? "" : CommandLine));
-
+                Log(SpewVerbosity,"Run: " + App + " " + (String.IsNullOrEmpty(CommandLine) ? "" : CommandLine));
             }
-			ProcessResult Result = ProcessManager.CreateProcess(Options.HasFlag(ERunOptions.AllowSpew), Path.GetFileNameWithoutExtension(App), Env);
+			ProcessResult Result = ProcessManager.CreateProcess(Options.HasFlag(ERunOptions.AllowSpew), Path.GetFileNameWithoutExtension(App), Env, SpewVerbosity:SpewVerbosity);
 			Process Proc = Result.ProcessObject;
 
 			bool bRedirectStdOut = (Options & ERunOptions.NoStdOutRedirect) != ERunOptions.NoStdOutRedirect;
@@ -604,7 +609,7 @@ namespace AutomationTool
 				AddRunTime(App, (int)(BuildDuration));
                 if (!Options.HasFlag(ERunOptions.NoLoggingOfRunCommand))
                 {
-                    Log("Run: Took {0}s to run " + Path.GetFileName(App), BuildDuration / 1000);
+                    Log(SpewVerbosity,"Run: Took {0}s to run " + Path.GetFileName(App), BuildDuration / 1000);
 
                 }
 				Result.ExitCode = Proc.ExitCode;
