Fix P4Connection.Changes() regression when using '-L' (long comments)

diff --git a/Engine/Source/Programs/AutomationTool/P4Utils.cs b/Engine/Source/Programs/AutomationTool/P4Utils.cs
--- a/Engine/Source/Programs/AutomationTool/P4Utils.cs
+++ b/Engine/Source/Programs/AutomationTool/P4Utils.cs
@@ -577,49 +577,19 @@ namespace AutomationTool
 
 				// If the user specified '-l' or '-L', the summary will appear on subsequent lines (no quotes) instead of the same line (surrounded by single quotes)
 				bool bSummaryIsOnSameLine = CommandLine.IndexOf( "-L", StringComparison.InvariantCultureIgnoreCase ) == -1;
+				if( bSummaryIsOnSameLine && LongComment )
+				{
+					CommandLine = "-L " + CommandLine;
+					bSummaryIsOnSameLine = false;
+				}
 
                 string Output;
-                if (!LogP4Output(out Output, "changes " + (LongComment ? "-L " : "") + CommandLine, null, AllowSpew))
+                if (!LogP4Output(out Output, "changes " + CommandLine, null, AllowSpew))
                 {
                     return false;
                 }
 
-                var TempLines = Output.Split(new string[] { Environment.NewLine }, StringSplitOptions.None);
                 var Lines = new List<string>();
-                if (LongComment)
-                {
-                    string Start = null;
-                    string Comment = "";
-                    foreach (var Line in TempLines)
-                    {
-                        if (Line.StartsWith("\t"))
-                        {
-                            Comment += Line;
-                        }
-                        else
-                        {
-                            if (Start != null)
-                            {
-                                Start = Start + " '" + Comment.Replace("\t", "").Replace("\n", "  ").Replace("\r", "") + "'";
-                                Lines.Add(Start);
-                                Comment = "";
-                            }
-                            Start = Line;
-                        }
-                    }
-                    if (Start != null)
-                    {
-                        Start = Start + " '" + Comment.Replace("\t", "").Replace("\n", "  ").Replace("\r", "") + "'";
-                        Lines.Add(Start);
-                    }
-                }
-                else
-                {
-                    foreach (var Line in TempLines)
-                    {
-                        Lines.Add(Line);
-                    }
-                }
                 for(int LineIndex = 0; LineIndex < Lines.Count; ++LineIndex)
                 {
 					var Line = Lines[ LineIndex ];
