Fix P4Connection.Changes() regression when using '-L' (long comments) (Part 2)

diff --git a/Engine/Source/Programs/AutomationTool/P4Utils.cs b/Engine/Source/Programs/AutomationTool/P4Utils.cs
--- a/Engine/Source/Programs/AutomationTool/P4Utils.cs
+++ b/Engine/Source/Programs/AutomationTool/P4Utils.cs
@@ -589,8 +589,8 @@ namespace AutomationTool
                     return false;
                 }
 
-                var Lines = new List<string>();
-                for(int LineIndex = 0; LineIndex < Lines.Count; ++LineIndex)
+                var Lines = Output.Split(new string[] { Environment.NewLine }, StringSplitOptions.None);
+                for(int LineIndex = 0; LineIndex < Lines.Length; ++LineIndex)
                 {
 					var Line = Lines[ LineIndex ];
 
@@ -631,7 +631,7 @@ namespace AutomationTool
 						else
 						{
 							++LineIndex;
-							if( LineIndex >= Lines.Count )
+							if( LineIndex >= Lines.Length )
 							{
 								throw new AutomationException("Was expecting a change summary to appear after Change header output from P4, but there were no more lines to read");
 							}
@@ -643,7 +643,7 @@ namespace AutomationTool
 							}
 
 							++LineIndex;
-							for( ; LineIndex < Lines.Count; ++LineIndex )
+							for( ; LineIndex < Lines.Length; ++LineIndex )
 							{
 								Line = Lines[ LineIndex ];
 
