Fixed packaging Shipping games on Mac

diff --git a/Engine/Source/Programs/AutomationTool/Mac/MacPlatform.Automation.cs b/Engine/Source/Programs/AutomationTool/Mac/MacPlatform.Automation.cs
--- a/Engine/Source/Programs/AutomationTool/Mac/MacPlatform.Automation.cs
+++ b/Engine/Source/Programs/AutomationTool/Mac/MacPlatform.Automation.cs
@@ -60,13 +60,14 @@ public class MacPlatform : Platform
 
 	public override void ProcessArchivedProject(ProjectParams Params, DeploymentContext SC)
 	{
+		string ExeName = SC.StageExecutables[0];
 		string BundlePath = CombinePaths(SC.ArchiveDirectory, SC.ShortProjectName + ".app");
 		if (!Directory.Exists(BundlePath))
 		{
-			string SourceBundlePath = CombinePaths(SC.ArchiveDirectory, SC.ShortProjectName, "Binaries", "Mac", SC.ShortProjectName + ".app");
+			string SourceBundlePath = CombinePaths(SC.ArchiveDirectory, SC.ShortProjectName, "Binaries", "Mac", ExeName + ".app");
 			if (!Directory.Exists(SourceBundlePath))
 			{
-				SourceBundlePath = CombinePaths(SC.ArchiveDirectory, "Engine", "Binaries", "Mac", "UE4Game.app");
+				SourceBundlePath = CombinePaths(SC.ArchiveDirectory, "Engine", "Binaries", "Mac", ExeName + ".app");
 			}
 			Directory.Move(SourceBundlePath, BundlePath);
 		}
@@ -99,7 +100,6 @@ public class MacPlatform : Platform
 		}
 
 		// Update executable name and entry in Info.plist
-		string ExeName = SC.StageExecutables[0];
 		string UE4GamePath = CombinePaths(SC.ArchiveDirectory, SC.ShortProjectName + ".app", "Contents", "MacOS", ExeName);
 		if (ExeName != SC.ShortProjectName && File.Exists(UE4GamePath))
 		{
