#UE4 - UAT performance improvements when p4 is required:

- In local runs UAT will no longer attempt to autodetect the current CL# if not asked to (with -p4cl command line param)
- Small client detection performance improvements - added client filter based on a known path under the client root.
- More logging during p4 autodetection

diff --git a/Engine/Source/Programs/AutomationTool/Automation.cs b/Engine/Source/Programs/AutomationTool/Automation.cs
--- a/Engine/Source/Programs/AutomationTool/Automation.cs
+++ b/Engine/Source/Programs/AutomationTool/Automation.cs
@@ -102,6 +102,7 @@ namespace AutomationTool
 		public static CommandLineArg ForceLocal = new CommandLineArg("-ForceLocal");
 		public static CommandLineArg NoP4 = new CommandLineArg("-NoP4");
 		public static CommandLineArg P4 = new CommandLineArg("-P4");
+		public static CommandLineArg P4CL = new CommandLineArg("-P4CL");
 		public static CommandLineArg Preprocess = new CommandLineArg("-Preprocess");
 		public static CommandLineArg NoCompile = new CommandLineArg("-NoCompile");
 		public static CommandLineArg NoCompileEditor = new CommandLineArg("-NoCompileEditor");
@@ -131,6 +132,7 @@ AutomationTool.exe [-verbose] [-compileonly] [-p4] Command0 [-Arg0 -Arg1 -Arg2 
 	[Help("verbose", "Enables verbose logging")]
 	[Help("nop4", "Disables Perforce functionality (default if not run on a build machine)")]
 	[Help("p4", "Enables Perforce functionality (default if run on a build machine)")]
+	[Help("p4cl", "Enables automatic 'have' changelist detection in local runs")]
 	[Help("compileonly", "Does not run any commands, only compiles them")]
 	[Help("forcelocal", "Forces local execution")]
 	[Help("help", "Displays help")]
diff --git a/Engine/Source/Programs/AutomationTool/LocalP4Environment.cs b/Engine/Source/Programs/AutomationTool/LocalP4Environment.cs
--- a/Engine/Source/Programs/AutomationTool/LocalP4Environment.cs
+++ b/Engine/Source/Programs/AutomationTool/LocalP4Environment.cs
@@ -70,12 +70,15 @@ namespace AutomationTool
 			CommandUtils.ConditionallySetEnvVar(EnvVarNames.BuildRootP4, BuildRootPath);
 			CommandUtils.ConditionallySetEnvVar(EnvVarNames.ClientRoot, ClientRootPath);
 
-			var CLString = CommandUtils.GetEnvVar(EnvVarNames.Changelist);
-			if (String.IsNullOrEmpty(CLString))
+			var CLString = CommandUtils.GetEnvVar(EnvVarNames.Changelist, null);
+			if (String.IsNullOrEmpty(CLString) && GlobalCommandLine.P4CL)
 			{
                 CLString = DetectCurrentCL(Connection, ClientRootPath);
 			}
-			CommandUtils.ConditionallySetEnvVar(EnvVarNames.Changelist, CLString);
+			if (!String.IsNullOrEmpty(CLString))
+			{
+				CommandUtils.ConditionallySetEnvVar(EnvVarNames.Changelist, CLString);
+			}
 
 			CommandUtils.ConditionallySetEnvVar(EnvVarNames.LabelToSync, "");
 			CommandUtils.ConditionallySetEnvVar("P4USER", UserName);
@@ -113,6 +116,8 @@ namespace AutomationTool
 		/// <returns>Changelist number as a string.</returns>
 		private static string DetectCurrentCL(P4Connection Connection, string ClientRootPath)
 		{
+			CommandUtils.Log("uebp_CL not set, detecting 'have' CL...");
+
 			// Retrieve the current changelist 
 			var P4Result = Connection.P4("changes -m 1 " + CommandUtils.CombinePaths(PathSeparator.Depot, ClientRootPath, "/...#have"), AllowSpew: false);
 			var CLTokens = P4Result.Output.Split(new char[] { ' ' }, StringSplitOptions.RemoveEmptyEntries);
@@ -172,8 +177,10 @@ namespace AutomationTool
 		/// <returns>Client to use.</returns>
 		private static P4ClientInfo DetectClient(P4Connection Connection, string UserName, string HostName, string UATLocation)
 		{
+			CommandUtils.Log("uebp_CLIENT not set, detecting current client...");
+
 			var MatchingClients = new List<P4ClientInfo>();
-			P4ClientInfo[] P4Clients = Connection.GetClientsForUser(UserName);
+			P4ClientInfo[] P4Clients = Connection.GetClientsForUser(UserName, UATLocation);
 			foreach (var Client in P4Clients)
 			{
 				if (String.IsNullOrEmpty(Client.Host) || String.Compare(Client.Host, HostName, true) != 0)
@@ -182,16 +189,7 @@ namespace AutomationTool
 					continue;
 				}
 				
-				var ClientRootPathWithSlash = Client.RootPath;
-				if (!ClientRootPathWithSlash.EndsWith("\\") && !ClientRootPathWithSlash.EndsWith("/"))
-				{
-					ClientRootPathWithSlash = CommandUtils.ConvertSeparators(PathSeparator.Default, ClientRootPathWithSlash + "/");
-				}
-				Log.TraceInformation("Checking clientspec {0} {1}", Client.Name, ClientRootPathWithSlash);
-				if (!String.IsNullOrEmpty(Client.RootPath) && UATLocation.StartsWith(ClientRootPathWithSlash, StringComparison.CurrentCultureIgnoreCase))
-				{
-					MatchingClients.Add(Client);
-				}
+				MatchingClients.Add(Client);
 			}
 
 			P4ClientInfo ClientToUse = null;
diff --git a/Engine/Source/Programs/AutomationTool/P4Environment.cs b/Engine/Source/Programs/AutomationTool/P4Environment.cs
--- a/Engine/Source/Programs/AutomationTool/P4Environment.cs
+++ b/Engine/Source/Programs/AutomationTool/P4Environment.cs
@@ -13,13 +13,46 @@ namespace AutomationTool
 	/// </summary>
 	public class P4Environment
 	{
-		#region Perforce Environment Properties
-
+		
 		public string P4Port { get; protected set; }
 		public string ClientRoot { get; protected set; }
 		public string User { get; protected set; }
-		public string ChangelistString { get; protected set; }
-		public int Changelist { get; protected set; }
+
+
+		private string ChangelistStringInternal;
+		public string ChangelistString 
+		{ 
+			get
+			{
+				if (String.IsNullOrEmpty(ChangelistStringInternal))
+				{
+					throw new AutomationException("P4Environment.ChangelistString has not been initialized but is requested. Set uebp_CL env var or run UAT with -P4CL to automatically detect changelist.");
+				}
+				return ChangelistStringInternal;
+			}
+			protected set
+			{
+				ChangelistStringInternal = value;
+			}
+		}
+
+		private int ChangelistInternal;
+		public int Changelist
+		{
+			get
+			{
+				if (ChangelistInternal <= 0)
+				{
+					throw new AutomationException("P4Environment.Changelist has not been initialized but is requested. Set uebp_CL env var or run UAT with -P4CL to automatically detect changelist.");
+				}
+				return Changelist;
+			}
+			protected set
+			{
+				ChangelistInternal = value;
+			}
+		}
+
 		public string Client { get; protected set; }
 		public string BuildRootP4 { get; protected set; }
 		public string BuildRootEscaped { get; protected set; }
@@ -27,7 +60,6 @@ namespace AutomationTool
 		public string LabelPrefix { get; protected set; }
 		public string BranchName { get; protected set; }
 
-		#endregion
 
 		internal P4Environment(P4Connection Connection, CommandEnvironment CmdEnv)
 		{
@@ -43,7 +75,7 @@ namespace AutomationTool
 			P4Port = CommandUtils.GetEnvVar(EnvVarNames.P4Port);
 			ClientRoot = CommandUtils.GetEnvVar(EnvVarNames.ClientRoot);
 			User = CommandUtils.GetEnvVar(EnvVarNames.User);
-			ChangelistString = CommandUtils.GetEnvVar(EnvVarNames.Changelist);
+			ChangelistStringInternal = CommandUtils.GetEnvVar(EnvVarNames.Changelist, null);
 			Client = CommandUtils.GetEnvVar(EnvVarNames.Client);
 			BuildRootP4 = CommandUtils.GetEnvVar(EnvVarNames.BuildRootP4);
 			if (BuildRootP4.EndsWith("/", StringComparison.InvariantCultureIgnoreCase) || BuildRootP4.EndsWith("\\", StringComparison.InvariantCultureIgnoreCase))
@@ -56,12 +88,12 @@ namespace AutomationTool
 			LabelToSync = CommandUtils.GetEnvVar(EnvVarNames.LabelToSync);
 
 			if (((CommandUtils.P4Enabled || CommandUtils.IsBuildMachine) && (ClientRoot == String.Empty || User == String.Empty ||
-				ChangelistString == String.Empty || Client == String.Empty || BuildRootP4 == String.Empty)))
+				(String.IsNullOrEmpty(ChangelistStringInternal) && CommandUtils.IsBuildMachine) || Client == String.Empty || BuildRootP4 == String.Empty)))
 			{
                 Log.TraceInformation("P4Enabled={0}", CommandUtils.P4Enabled);
                 Log.TraceInformation("ClientRoot={0}",ClientRoot );
                 Log.TraceInformation("User={0}", User);
-                Log.TraceInformation("ChangelistString={0}", ChangelistString);
+                Log.TraceInformation("ChangelistString={0}", ChangelistStringInternal);
                 Log.TraceInformation("Client={0}", Client);
                 Log.TraceInformation("BuildRootP4={0}", BuildRootP4);
 
@@ -72,20 +104,13 @@ namespace AutomationTool
 
 			if (CommandUtils.P4Enabled)
 			{
-				if (ChangelistString.Length < 2)
-				{
-					throw new AutomationException("BUILD FAILED bad CL {0}.", ChangelistString);
-				}
-
-				Changelist = int.Parse(ChangelistString);
-
-				if (Changelist <= 10)
+				if (CommandUtils.IsBuildMachine || ChangelistStringInternal != null)
 				{
-					throw new AutomationException("BUILD FAILED bad CL {0} {1}.", ChangelistString, Changelist);
+					// We may not always need the changelist number if we're not a build machine.
+					// In local runs, changelist initialization can be really slow!
+					VerifyChangelistStringAndSetChangelistNumber();
 				}
 
-				ChangelistString = Changelist.ToString();		// Make sure they match
-
 				// Setup branch name
 				string DepotSuffix = "//depot/";
 				if (BuildRootP4.StartsWith(DepotSuffix))
@@ -123,5 +148,27 @@ namespace AutomationTool
 			Log.TraceInformation("{0}={1}", "P4CLIENT", InternalUtils.GetEnvironmentVariable("P4CLIENT", "", bQuiet));
 			Log.TraceInformation("LabelPrefix={0}", LabelPrefix);
 		}
+
+		protected void VerifyChangelistStringAndSetChangelistNumber()
+		{
+			if (String.IsNullOrEmpty(ChangelistStringInternal))
+			{
+				throw new AutomationException("BUILD FAILED null CL ");
+			}
+
+			if (String.IsNullOrEmpty(ChangelistStringInternal) || ChangelistStringInternal.Length < 2)
+			{
+				throw new AutomationException("BUILD FAILED bad CL {0}.", ChangelistString);
+			}
+
+			ChangelistInternal = Int32.Parse(ChangelistStringInternal);
+
+			if (ChangelistInternal <= 10)
+			{
+				throw new AutomationException("BUILD FAILED bad CL {0} {1}.", ChangelistStringInternal, ChangelistInternal);
+			}
+
+			ChangelistStringInternal = ChangelistInternal.ToString();		// Make sure they match
+		}
 	}
 }
diff --git a/Engine/Source/Programs/AutomationTool/P4Utils.cs b/Engine/Source/Programs/AutomationTool/P4Utils.cs
--- a/Engine/Source/Programs/AutomationTool/P4Utils.cs
+++ b/Engine/Source/Programs/AutomationTool/P4Utils.cs
@@ -1974,7 +1974,7 @@ namespace AutomationTool
 		/// </summary>
 		/// <param name="UserName"></param>
 		/// <returns>List of clients owned by the user.</returns>
-		public P4ClientInfo[] GetClientsForUser(string UserName)
+		public P4ClientInfo[] GetClientsForUser(string UserName, string PathUnderClientRoot = null)
 		{
 			CheckP4Enabled();
 
@@ -1995,11 +1995,11 @@ namespace AutomationTool
 				P4ClientInfo Info = null;
 
 				// Retrieve the client name and info.
-				for (int i = 0; i < Tokens.Length; ++i)
+				for (int TokenIndex = 0; TokenIndex < Tokens.Length; ++TokenIndex)
 				{
-					if (Tokens[i] == "Client")
+					if (Tokens[TokenIndex] == "Client")
 					{
-						var ClientName = Tokens[++i];
+						var ClientName = Tokens[++TokenIndex];
 						Info = GetClientInfoInternal(ClientName);
 						break;
 					}
@@ -2009,8 +2009,23 @@ namespace AutomationTool
 				{
 					throw new AutomationException("Failed to retrieve p4 client info for user {0}. Unable to set up local environment", UserName);
 				}
+				
+				bool bAddClient = true;
+				// Filter the client out if the specified path is not under the client root
+				if (!String.IsNullOrEmpty(PathUnderClientRoot) && !String.IsNullOrEmpty(Info.RootPath))
+				{
+					var ClientRootPathWithSlash = Info.RootPath;
+					if (!ClientRootPathWithSlash.EndsWith("\\") && !ClientRootPathWithSlash.EndsWith("/"))
+					{
+						ClientRootPathWithSlash = CommandUtils.ConvertSeparators(PathSeparator.Default, ClientRootPathWithSlash + "/");
+					}
+					bAddClient = PathUnderClientRoot.StartsWith(ClientRootPathWithSlash, StringComparison.CurrentCultureIgnoreCase);
+				}
 
-				ClientList.Add(Info);
+				if (bAddClient)
+				{
+					ClientList.Add(Info);
+				}
 			}
 			return ClientList.ToArray();
 		}
