#UE4 - UAT: Fix for infinite recursion in P4Environment.

diff --git a/Engine/Source/Programs/AutomationTool/P4Environment.cs b/Engine/Source/Programs/AutomationTool/P4Environment.cs
--- a/Engine/Source/Programs/AutomationTool/P4Environment.cs
+++ b/Engine/Source/Programs/AutomationTool/P4Environment.cs
@@ -45,7 +45,7 @@ namespace AutomationTool
 				{
 					throw new AutomationException("P4Environment.Changelist has not been initialized but is requested. Set uebp_CL env var or run UAT with -P4CL to automatically detect changelist.");
 				}
-				return Changelist;
+				return ChangelistInternal;
 			}
 			protected set
 			{
@@ -158,7 +158,7 @@ namespace AutomationTool
 
 			if (String.IsNullOrEmpty(ChangelistStringInternal) || ChangelistStringInternal.Length < 2)
 			{
-				throw new AutomationException("BUILD FAILED bad CL {0}.", ChangelistString);
+				throw new AutomationException("BUILD FAILED bad CL {0}.", ChangelistStringInternal);
 			}
 
 			ChangelistInternal = Int32.Parse(ChangelistStringInternal);
