build script, retry on mac temp storage

diff --git a/Engine/Source/Programs/AutomationTool/TempStorage.cs b/Engine/Source/Programs/AutomationTool/TempStorage.cs
--- a/Engine/Source/Programs/AutomationTool/TempStorage.cs
+++ b/Engine/Source/Programs/AutomationTool/TempStorage.cs
@@ -428,7 +428,22 @@ namespace AutomationTool
                         var NewFile = CombinePaths(NewBaseDir, ThisFileInfo.Name);
                         if (!FileExists_NoExceptions(false, NewFile))
                         {
-                            throw new AutomationException("Rebased manifest file does not exist {0}", NewFile);
+                            bool bFound = false;
+                            // mac is terrible on shares, this isn't a solution, but a stop gap
+                            if (NewFile.StartsWith("/Volumes/"))
+                            {
+                                int Retry = 60;
+                                while (!bFound && --Retry > 0)
+                                {
+                                    CommandUtils.Log(System.Diagnostics.TraceEventType.Warning, "*** Mac temp storage retry {0}", NewFile);
+                                    System.Threading.Thread.Sleep(10000);
+                                    bFound = FileExists_NoExceptions(false, NewFile);
+                                }
+                            }
+                            if (!bFound)
+                            {
+                                throw new AutomationException("Rebased manifest file does not exist {0}", NewFile);
+                            }
                         }
                         FileInfo Info = new FileInfo(NewFile);
 
