build script, fix noop handling in a restarted build

diff --git a/Engine/Source/Programs/AutomationTool/Scripts/GUBP.Automation.cs b/Engine/Source/Programs/AutomationTool/Scripts/GUBP.Automation.cs
--- a/Engine/Source/Programs/AutomationTool/Scripts/GUBP.Automation.cs
+++ b/Engine/Source/Programs/AutomationTool/Scripts/GUBP.Automation.cs
@@ -3550,7 +3550,7 @@ public class GUBP : BuildCommand
 
     string GetJobStepPath(string Dep)
     {
-        if (GUBPNodes[Dep].AgentSharingGroup != "")
+        if (Dep != "Noop" && GUBPNodes[Dep].AgentSharingGroup != "")
         {
             return "jobSteps[" + GUBPNodes[Dep].AgentSharingGroup + "]/jobSteps[" + Dep + "]";
         }
@@ -5056,11 +5056,13 @@ if (HostPlatform == UnrealTargetPlatform.Mac) continue; //temp hack till mac aut
             string ParentPath = ParseParamValue("ParentPath");
             string BaseArgs = String.Format("createJobStep --parentPath {0}", ParentPath);
 
+            bool bHasNoop = false;
             if (LastSticky == "" && bHaveECNodes)
             {
                 // if we don't have any sticky nodes and we have other nodes, we run a fake noop just to release the resource 
                 string Args = String.Format("{0} --subprocedure GUBP_UAT_Node --parallel 0 --jobStepName Noop --actualParameter NodeName=Noop --actualParameter Sticky=1 --releaseMode release", BaseArgs);
                 RunECTool(Args);
+                bHasNoop = true;
             }
 
             var FakeECArgs = new List<string>();
@@ -5153,7 +5155,7 @@ if (HostPlatform == UnrealTargetPlatform.Mac) continue; //temp hack till mac aut
 
                         var MyChain = AgentGroupChains[MyAgentGroup];
                         int MyIndex = MyChain.IndexOf(NodeToDo);
-                                if (MyIndex  > 0)
+                        if (MyIndex  > 0)
                         {
                             PreConditionUncompletedEcDeps.Add(MyChain[MyIndex - 1]);
                         }
@@ -5181,6 +5183,10 @@ if (HostPlatform == UnrealTargetPlatform.Mac) continue; //temp hack till mac aut
                             }
                         }
                     }
+                    if (bHasNoop && PreConditionUncompletedEcDeps.Count == 0)
+                    {
+                        PreConditionUncompletedEcDeps.Add("Noop");
+                    }
                     if (PreConditionUncompletedEcDeps.Count > 0)
                     {
                         PreCondition = "\"$[/javascript if(";
