#UE4 - UAT: Adding a list of process names that can not be killed by UAT.

diff --git a/Engine/Source/Programs/AutomationTool/HostPlatform.cs b/Engine/Source/Programs/AutomationTool/HostPlatform.cs
--- a/Engine/Source/Programs/AutomationTool/HostPlatform.cs
+++ b/Engine/Source/Programs/AutomationTool/HostPlatform.cs
@@ -111,5 +111,10 @@ namespace AutomationTool
 		/// Returns the pdb file extenstion for the host platform.
 		/// </summary>
 		abstract public string PdbExtension { get; }
+
+		/// <summary>
+		/// List of processes that can't not be killed
+		/// </summary>
+		abstract public string[] DontKillProcessList { get; }
 	}
 }
diff --git a/Engine/Source/Programs/AutomationTool/MacHostPlatform.cs b/Engine/Source/Programs/AutomationTool/MacHostPlatform.cs
--- a/Engine/Source/Programs/AutomationTool/MacHostPlatform.cs
+++ b/Engine/Source/Programs/AutomationTool/MacHostPlatform.cs
@@ -138,5 +138,17 @@ namespace AutomationTool
 		{
 			get { return ".exe.mdb"; }
 		}
+
+		static string[] SystemServices = new string[]
+		{
+			// TODO: Add any system process names here
+		};
+		public override string[] DontKillProcessList
+		{
+			get
+			{
+				return SystemServices;
+			}
+		}
 	}
 }
diff --git a/Engine/Source/Programs/AutomationTool/ProcessUtils.cs b/Engine/Source/Programs/AutomationTool/ProcessUtils.cs
--- a/Engine/Source/Programs/AutomationTool/ProcessUtils.cs
+++ b/Engine/Source/Programs/AutomationTool/ProcessUtils.cs
@@ -40,6 +40,7 @@ namespace AutomationTool
 		/// </summary>
 		private static object SyncObject = new object();
 
+
 		/// <summary>
 		/// Creates a new process and adds it to the tracking list.
 		/// </summary>
@@ -70,6 +71,11 @@ namespace AutomationTool
 			return Result;
 		}
 
+		public static bool CanBeKilled(string ProcessName)
+		{
+			return !HostPlatform.Current.DontKillProcessList.Contains(ProcessName, StringComparer.InvariantCultureIgnoreCase);
+		}
+
 		/// <summary>
 		/// Kills all running processes.
 		/// </summary>
@@ -81,6 +87,15 @@ namespace AutomationTool
 				ProcessesToKill = new List<ProcessResult>(ActiveProcesses);
 				ActiveProcesses.Clear();
 			}
+			// Remove processes that can't be killed
+			for (int ProcessIndex = ProcessesToKill.Count - 1; ProcessIndex >= 0; --ProcessIndex )
+			{
+				if (!CanBeKilled(ProcessesToKill[ProcessIndex].ProcessObject.ProcessName))
+				{
+					CommandUtils.Log("Ignoring process \"{0}\" because it can't be killed.", ProcessesToKill[ProcessIndex].ProcessObject.ProcessName);
+					ProcessesToKill.RemoveAt(ProcessIndex);
+				}
+			}
 			CommandUtils.Log("Trying to kill {0} spawned processes.", ProcessesToKill.Count);
 			foreach (var Proc in ProcessesToKill)
 			{
@@ -416,7 +431,7 @@ namespace AutomationTool
 				foreach (Process KillCandidate in AllProcs)
 				{
 					HashSet<int> VisitedPids = new HashSet<int>();
-					if (IsOurDescendant(ProcessToKill, KillCandidate.Id, VisitedPids))
+					if (ProcessManager.CanBeKilled(KillCandidate.ProcessName) && IsOurDescendant(ProcessToKill, KillCandidate.Id, VisitedPids))
 					{
 						CommandUtils.Log("Trying to kill descendant pid={0}, name={1}", KillCandidate.Id, KillCandidate.ProcessName);
 						try
@@ -441,12 +456,11 @@ namespace AutomationTool
         /// <param name="ProcessToCheck">Process to check</param>
         public static bool HasAnyDescendants(Process ProcessToCheck)
         {
-
             Process[] AllProcs = Process.GetProcesses();
             foreach (Process KillCandidate in AllProcs)
             {
                 HashSet<int> VisitedPids = new HashSet<int>();
-                if (IsOurDescendant(ProcessToCheck, KillCandidate.Id, VisitedPids))
+                if (ProcessManager.CanBeKilled(KillCandidate.ProcessName) && IsOurDescendant(ProcessToCheck, KillCandidate.Id, VisitedPids))
                 {
                     CommandUtils.Log("Descendant pid={0}, name={1}, filename={2}", KillCandidate.Id, KillCandidate.ProcessName,
 						KillCandidate.MainModule != null ? KillCandidate.MainModule.FileName : "unknown");
diff --git a/Engine/Source/Programs/AutomationTool/WindowsHostPlatform.cs b/Engine/Source/Programs/AutomationTool/WindowsHostPlatform.cs
--- a/Engine/Source/Programs/AutomationTool/WindowsHostPlatform.cs
+++ b/Engine/Source/Programs/AutomationTool/WindowsHostPlatform.cs
@@ -205,5 +205,25 @@ namespace AutomationTool
 		{
 			get { return ".pdb"; }
 		}
+
+		static string[] SystemServices = new string[]
+		{
+			"winlogon",
+			"system idle process",
+			"taskmgr",
+			"spoolsv",
+			"csrss",
+			"smss",
+			"svchost",
+			"services",
+			"lsass"
+		};
+		public override string[] DontKillProcessList
+		{
+			get 
+			{
+				return SystemServices;
+			}
+		}
 	}
 }
