build script, flesh out the PS4 rocket stuff

diff --git a/Engine/Source/Programs/AutomationTool/Scripts/GUBP.Automation.cs b/Engine/Source/Programs/AutomationTool/Scripts/GUBP.Automation.cs
--- a/Engine/Source/Programs/AutomationTool/Scripts/GUBP.Automation.cs
+++ b/Engine/Source/Programs/AutomationTool/Scripts/GUBP.Automation.cs
@@ -3070,7 +3070,7 @@ public class GUBP : BuildCommand
         bool bShowAllChanges = bp.ParseParam("AllChanges") && GUBPNodesHistory != null;
         bool bShowChanges = (bp.ParseParam("Changes") && GUBPNodesHistory != null) || bShowAllChanges;
         bool bShowDetailedHistory = (bp.ParseParam("History") && GUBPNodesHistory != null) || bShowChanges;
-        bool bShowDependencies = !bp.ParseParam("NoShowDependencies") && !bShowDetailedHistory;
+        bool bShowDependencies = bp.ParseParam("ShowDependencies");
         bool bShowECDependencies = bp.ParseParam("ShowECDependencies");
         bool bShowHistory = !bp.ParseParam("NoHistory") && GUBPNodesHistory != null;
         bool AddEmailProps = bp.ParseParam("ShowEmails");
@@ -3753,7 +3753,7 @@ public class GUBP : BuildCommand
     [Help("AllChanges", "Like changes except includes changes before the last green. Must have -P4 for this to work.")]
     [Help("EmailOnly", "Only emails the folks given in the argument.")]
     [Help("AddEmails", "Add these space delimited emails too all email lists.")]
-    [Help("NoShowDependencies", "Don't show node dependencies.")]
+    [Help("ShowDependencies", "Show node dependencies.")]
     [Help("ShowECDependencies", "Show EC node dependencies instead.")]
     [Help("ShowECProc", "Show EC proc names.")]
     [Help("BuildRocket", "Build in rocket mode.")]
@@ -4021,7 +4021,7 @@ public class GUBP : BuildCommand
             {
                 if (Plat != UnrealTargetPlatform.HTML5 && 
                     Plat != UnrealTargetPlatform.Linux &&
-                    Plat != UnrealTargetPlatform.PS4 &&
+                    (Plat != UnrealTargetPlatform.PS4 || ParseParam("WithPS4")) &&
                     Plat != UnrealTargetPlatform.WinRT &&
                     Plat != UnrealTargetPlatform.WinRT_ARM &&
                     Plat != UnrealTargetPlatform.XboxOne
