Fix CIS warning on Mono causing build failure.

diff --git a/Engine/Source/Programs/AutomationTool/Scripts/BuildDocumentation.Automation.cs b/Engine/Source/Programs/AutomationTool/Scripts/BuildDocumentation.Automation.cs
--- a/Engine/Source/Programs/AutomationTool/Scripts/BuildDocumentation.Automation.cs
+++ b/Engine/Source/Programs/AutomationTool/Scripts/BuildDocumentation.Automation.cs
@@ -56,7 +56,6 @@ public class BuildDocumentation : BuildCommand
 
 		// Get all the file paths
 		string TargetInfoPath = Path.Combine(BuildDir, "targetinfo.xml");
-		string UnrealBuildToolPath = Path.Combine(CmdEnv.LocalRoot, "Engine/Binaries/DotNET/UnrealBuildTool.exe");
 		string UnrealBuildToolProject = Path.Combine(CmdEnv.LocalRoot, "Engine/Source/Programs/UnrealBuildTool/UnrealBuildTool.csproj");
 		string ApiDocToolSolution = Path.Combine(CmdEnv.LocalRoot, "Engine/Source/Programs/UnrealDocTool/APIDocTool/APIDocTool.sln");
 		string ApiDocToolPath = Path.Combine(CmdEnv.LocalRoot, "Engine/Source/Programs/UnrealDocTool/APIDocTool/APIDocTool/bin/x64/Release/APIDocTool.exe");
