build script, rework all email stuff to be hacker-style

diff --git a/Engine/Source/Programs/AutomationTool/Android/AndroidPlatform.Automation.cs b/Engine/Source/Programs/AutomationTool/Android/AndroidPlatform.Automation.cs
--- a/Engine/Source/Programs/AutomationTool/Android/AndroidPlatform.Automation.cs
+++ b/Engine/Source/Programs/AutomationTool/Android/AndroidPlatform.Automation.cs
@@ -433,10 +433,6 @@ public class AndroidPlatform : Platform
 		}
 	}
 	#endregion
-    public override string GUBP_GetPlatformFailureEMails(string Branch)
-    {
-        return "Josh.Adams[epic] JJ.Hoesing[epic]";
-    }
 
     public override List<string> GetDebugFileExtentions()
     {
diff --git a/Engine/Source/Programs/AutomationTool/HTML5/HTML5Platform.Automation.cs b/Engine/Source/Programs/AutomationTool/HTML5/HTML5Platform.Automation.cs
--- a/Engine/Source/Programs/AutomationTool/HTML5/HTML5Platform.Automation.cs
+++ b/Engine/Source/Programs/AutomationTool/HTML5/HTML5Platform.Automation.cs
@@ -192,11 +192,6 @@ public class HTML5Platform : Platform
 
 	public override bool IsSupported { get { return true; } }
 
-    public override string GUBP_GetPlatformFailureEMails(string Branch)
-    {
-        return "Peter.Sauerbrei[epic]";
-    }
-
     public override List<string> GetDebugFileExtentions()
     {
         return new List<string> { ".pdb" };
diff --git a/Engine/Source/Programs/AutomationTool/IOS/IOSPlatform.Automation.cs b/Engine/Source/Programs/AutomationTool/IOS/IOSPlatform.Automation.cs
--- a/Engine/Source/Programs/AutomationTool/IOS/IOSPlatform.Automation.cs
+++ b/Engine/Source/Programs/AutomationTool/IOS/IOSPlatform.Automation.cs
@@ -582,10 +582,6 @@ public class IOSPlatform : Platform
 	{
 		return "cookeddata/" + Dest;
 	}
-    public override string GUBP_GetPlatformFailureEMails(string Branch)
-    {
-        return "Peter.Sauerbrei[epic] Michael.Trepka[epic]";
-    }
     public override List<string> GetDebugFileExtentions()
     {
         return new List<string> { ".dsym" };
diff --git a/Engine/Source/Programs/AutomationTool/Linux/LinuxPlatform.Automation.cs b/Engine/Source/Programs/AutomationTool/Linux/LinuxPlatform.Automation.cs
--- a/Engine/Source/Programs/AutomationTool/Linux/LinuxPlatform.Automation.cs
+++ b/Engine/Source/Programs/AutomationTool/Linux/LinuxPlatform.Automation.cs
@@ -279,10 +279,6 @@ chmod 700 $HOME/Desktop/{1}.desktop", DesiredGLVersion, SC.ShortProjectName.ToLo
             throw new AutomationException("must specify device for Linux target (-serverdevice=<ip>)");
         }
     }
-    public override string GUBP_GetPlatformFailureEMails(string Branch)
-    {
-        return "Dmitry.Rekman[epic]";
-    }
     public override List<string> GetDebugFileExtentions()
     {
         return new List<string> {  };
diff --git a/Engine/Source/Programs/AutomationTool/Mac/MacPlatform.Automation.cs b/Engine/Source/Programs/AutomationTool/Mac/MacPlatform.Automation.cs
--- a/Engine/Source/Programs/AutomationTool/Mac/MacPlatform.Automation.cs
+++ b/Engine/Source/Programs/AutomationTool/Mac/MacPlatform.Automation.cs
@@ -126,10 +126,6 @@ public class MacPlatform : Platform
 		bool bUseManifest = !bIsBuildingRemotely || BuildConfiguration.bCopyAppBundleBackToDevice;
 		return bUseManifest;
 	}
-    public override string GUBP_GetPlatformFailureEMails(string Branch)
-    {
-        return "Michael.Trepka[epic]";
-    }
     public override List<string> GetDebugFileExtentions()
     {
         return new List<string> {".dsym"};
diff --git a/Engine/Source/Programs/AutomationTool/Platform.cs b/Engine/Source/Programs/AutomationTool/Platform.cs
--- a/Engine/Source/Programs/AutomationTool/Platform.cs
+++ b/Engine/Source/Programs/AutomationTool/Platform.cs
@@ -287,15 +287,6 @@ namespace AutomationTool
             return new List<string>();
         }
 
-
-        /// <summary>
-        /// Allow the platform to setup emails for the GUBP for folks that care about node failures relating to this platform
-        /// </summary>
-        /// <param name="Branch">p4 root of the branch we are running</param>
-        public virtual string GUBP_GetPlatformFailureEMails(string Branch)
-        {
-            return "";
-        }
 		/// <summary>
 		/// UnrealTargetPlatform type for this platform.
 		/// </summary>
diff --git a/Engine/Source/Programs/AutomationTool/ProjectUtils.cs b/Engine/Source/Programs/AutomationTool/ProjectUtils.cs
--- a/Engine/Source/Programs/AutomationTool/ProjectUtils.cs
+++ b/Engine/Source/Programs/AutomationTool/ProjectUtils.cs
@@ -547,6 +547,10 @@ namespace AutomationTool
                     }
                 }
             }
+            if (String.IsNullOrEmpty(BaseEngineProject.FilePath))
+            {
+                throw new AutomationException("All branches must have the blank project /Samples/SampleGames/BlankProject");
+            }
 
             CommandUtils.Log("  Base Engine:");
             BaseEngineProject.Dump(InHostPlatforms);
diff --git a/Engine/Source/Programs/AutomationTool/Scripts/GUBP.Automation.cs b/Engine/Source/Programs/AutomationTool/Scripts/GUBP.Automation.cs
--- a/Engine/Source/Programs/AutomationTool/Scripts/GUBP.Automation.cs
+++ b/Engine/Source/Programs/AutomationTool/Scripts/GUBP.Automation.cs
@@ -25,7 +25,6 @@ public class GUBP : BuildCommand
     public static bool bBuildRocket = false;
     public static bool bForceIncrementalCompile = false;
     public static string ECProject = null;
-    public static string ForceAllEmailsTo = null;
 
     Dictionary<string, GUBPNode> GUBPNodes;
     Dictionary<string, bool> GUBPNodesCompleted;
@@ -94,7 +93,15 @@ public class GUBP : BuildCommand
 
     public abstract class GUBPEmailHacker
     {
-        public virtual string HackEmails(string Emails, GUBP bp, string Branch, string NodeName)
+        public virtual List<string> AddEmails(GUBP bp, string Branch, string NodeName)
+        {
+            return new List<string>();
+        }
+        public virtual string ModifyEmail(string Email, GUBP bp, string Branch, string NodeName)
+        {
+            return Email;
+        }
+        public virtual List<string> FinalizeEmails(List<string> Emails, GUBP bp, string Branch, string NodeName)
         {
             return Emails;
         }
@@ -103,7 +110,6 @@ public class GUBP : BuildCommand
     private static List<GUBPEmailHacker> EmailHackers;
     private string HackEmails(string Emails, string Branch, string NodeName)
     {
-        string Result = Emails;
         if (EmailHackers == null)
         {
             EmailHackers = new List<GUBPEmailHacker>();
@@ -124,11 +130,35 @@ public class GUBP : BuildCommand
                 }
             }
         }
+        List<string> Result = new List<string>(Emails.Split(' '));
+
         foreach (var EmailHacker in EmailHackers)
         {
-            Result = EmailHacker.HackEmails(Result, this, Branch, NodeName);
+            Result.AddRange(EmailHacker.AddEmails(this, Branch, NodeName));
         }
-        return Result;
+        foreach (var EmailHacker in EmailHackers)
+        {
+            var NewResult = new List<string>();
+            foreach (var Email in Result)
+            {
+                var NewEmail = EmailHacker.ModifyEmail(Email, this, Branch, NodeName);
+                if (!String.IsNullOrEmpty(NewEmail))
+                {
+                    NewResult.Add(NewEmail);
+                }
+            }
+            Result = NewResult;
+        }
+        foreach (var EmailHacker in EmailHackers)
+        {
+            Result = EmailHacker.FinalizeEmails(Result, this, Branch, NodeName);
+        }
+        string FinalEmails = "";
+        foreach (var Email in Result)
+        {
+            FinalEmails = GUBPNode.MergeSpaceStrings(FinalEmails, Email);
+        }
+        return FinalEmails;
     }
     public abstract class GUBPNode
     {
@@ -247,10 +277,6 @@ public class GUBP : BuildCommand
             }
             return Result;
         }
-        public virtual string FailureEMails(GUBP bp, string Branch)
-        {
-            return "";
-        }
         public static string MergeSpaceStrings(params string[] EmailLists)
         {
             var Emails = new List<string>();
@@ -954,11 +980,6 @@ public class GUBP : BuildCommand
             return Agenda;
         }
 
-        public override string FailureEMails(GUBP bp, string Branch)
-        {
-            return MergeSpaceStrings(base.FailureEMails(bp, Branch), 
-                Platform.Platforms[EditorPlatform].GUBP_GetPlatformFailureEMails(Branch));
-		}
     }
 
     public class EditorGameNode : CompileNode
@@ -1015,11 +1036,6 @@ public class GUBP : BuildCommand
 
             return Agenda;
         }
-        public override string FailureEMails(GUBP bp, string Branch)
-        {
-            return MergeSpaceStrings(base.FailureEMails(bp, Branch),
-               GameProj.Properties.Targets[TargetRules.TargetType.Editor].Rules.GUBP_GetGameFailureEMails_EditorTypeOnly(Branch));
-        }
     }
 
     public class GamePlatformMonolithicsNode : CompileNode
@@ -1143,12 +1159,6 @@ public class GUBP : BuildCommand
 
             return Agenda;
         }
-        public override string FailureEMails(GUBP bp, string Branch)
-        {
-            return MergeSpaceStrings(base.FailureEMails(bp, Branch),
-               GameProj.Properties.Targets[TargetRules.TargetType.Editor].Rules.GUBP_GetGameFailureEMails_EditorTypeOnly(Branch),
-               Platform.Platforms[TargetPlatform].GUBP_GetPlatformFailureEMails(Branch));
-        }
     }
 
     public class RootEditorHeadersNode : HostPlatformNode
@@ -1615,11 +1625,6 @@ public class GUBP : BuildCommand
             }
             return base.TriggerRequiresRecursiveWorkflow();
         }
-        public override string FailureEMails(GUBP bp, string Branch)
-        {
-            return MergeSpaceStrings(base.FailureEMails(bp, Branch),
-               GameProj.Properties.Targets[TargetRules.TargetType.Editor].Rules.GUBP_GetPromotionEMails_EditorTypeOnly(Branch));
-        }
     }
 
     public class WaitForSharedPromotionUserInput : WaitForPromotionUserInput
@@ -1649,11 +1654,6 @@ public class GUBP : BuildCommand
         {
             return WaitForPromotionUserInput.StaticGetFullName("Shared", bInLabelPromoted);
         }
-        public override string FailureEMails(GUBP bp, string Branch)
-        {
-            return MergeSpaceStrings(base.FailureEMails(bp, Branch),
-               bp.Branch.BaseEngineProject.Properties.Targets[TargetRules.TargetType.Editor].Rules.GUBP_GetPromotionEMails_EditorTypeOnly(Branch));
-        }
     }
 
     public class LabelPromotableNode : GUBPNode
@@ -1776,12 +1776,6 @@ public class GUBP : BuildCommand
         {
             return GameProj.GameName;
         }
-        public override string FailureEMails(GUBP bp, string Branch)
-        {
-            return MergeSpaceStrings(base.FailureEMails(bp, Branch),
-               GameProj.Properties.Targets[TargetRules.TargetType.Editor].Rules.GUBP_GetPromotionEMails_EditorTypeOnly(Branch));
-
-        }
     }
 
     public class SharedLabelPromotableNode : LabelPromotableNode
@@ -1795,12 +1789,6 @@ public class GUBP : BuildCommand
         {
             return LabelPromotableNode.StaticGetFullName("Shared", bInLabelPromoted);
         }
-        public override string FailureEMails(GUBP bp, string Branch)
-        {
-            return MergeSpaceStrings(base.FailureEMails(bp, Branch),
-               bp.Branch.BaseEngineProject.Properties.Targets[TargetRules.TargetType.Editor].Rules.GUBP_GetPromotionEMails_EditorTypeOnly(Branch));
-
-        }
     }
 
     public class WaitForTestShared : AggregateNode
@@ -1947,21 +1935,6 @@ public class GUBP : BuildCommand
                 AddBuildProduct(CookedFile);
             }
         }
-        public override string FailureEMails(GUBP bp, string Branch)
-        {
-            if (GameProj.GameName != bp.Branch.BaseEngineProject.GameName && GameProj.Properties.Targets.ContainsKey(TargetRules.TargetType.Editor))
-            {
-                return MergeSpaceStrings(base.FailureEMails(bp, Branch),
-                   GameProj.Properties.Targets[TargetRules.TargetType.Editor].Rules.GUBP_GetGameFailureEMails_EditorTypeOnly(Branch),
-                   Platform.Platforms[TargetPlatform].GUBP_GetPlatformFailureEMails(Branch));
-            }
-            else
-            {
-                return MergeSpaceStrings(base.FailureEMails(bp, Branch),
-                   bp.Branch.BaseEngineProject.Properties.Targets[TargetRules.TargetType.Editor].Rules.GUBP_GetGameFailureEMails_EditorTypeOnly(Branch),
-                   Platform.Platforms[TargetPlatform].GUBP_GetPlatformFailureEMails(Branch));
-            }
-        }
     }
 
 
@@ -2076,21 +2049,6 @@ public class GUBP : BuildCommand
             BuildProducts = new List<string>();
             SaveRecordOfSuccessAndAddToBuildProducts();
         }
-        public override string FailureEMails(GUBP bp, string Branch)
-        {
-            if (GameProj.GameName != bp.Branch.BaseEngineProject.GameName && GameProj.Properties.Targets.ContainsKey(TargetRules.TargetType.Editor))
-            {
-                return MergeSpaceStrings(base.FailureEMails(bp, Branch),
-                   GameProj.Properties.Targets[TargetRules.TargetType.Editor].Rules.GUBP_GetGameFailureEMails_EditorTypeOnly(Branch),
-                   Platform.Platforms[TargetPlatform].GUBP_GetPlatformFailureEMails(Branch));
-            }
-            else
-            {
-                return MergeSpaceStrings(base.FailureEMails(bp, Branch),
-                   bp.Branch.BaseEngineProject.Properties.Targets[TargetRules.TargetType.Editor].Rules.GUBP_GetGameFailureEMails_EditorTypeOnly(Branch),
-                   Platform.Platforms[TargetPlatform].GUBP_GetPlatformFailureEMails(Branch));
-            }
-        }
     }
 
 
@@ -2414,18 +2372,6 @@ public class GUBP : BuildCommand
             string LogFile = CommandUtils.RunUAT(CommandUtils.CmdEnv, Args);
             SaveRecordOfSuccessAndAddToBuildProducts(CommandUtils.ReadAllText(LogFile));
         }
-        public override string FailureEMails(GUBP bp, string Branch)
-        {
-            var WorkingGameProject = GameProj;
-            if (!WorkingGameProject.Properties.Targets.ContainsKey(TargetRules.TargetType.Editor))
-            {
-                // this is a codeless project, use the base project
-                WorkingGameProject = bp.Branch.BaseEngineProject;
-            }
-
-            return MergeSpaceStrings(base.FailureEMails(bp, Branch),
-               WorkingGameProject.Properties.Targets[TargetRules.TargetType.Editor].Rules.GUBP_GetPromotionEMails_EditorTypeOnly(Branch));
-        }
     }
 
 
@@ -2492,15 +2438,6 @@ public class GUBP : BuildCommand
             UE4Build.CheckBuildProducts(Build.BuildProductFiles);
             SaveRecordOfSuccessAndAddToBuildProducts();
         }
-        public override string FailureEMails(GUBP bp, string Branch)
-        {
-            string Emails;
-            Emails = MergeSpaceStrings(base.FailureEMails(bp, Branch),
-                bp.Branch.BaseEngineProject.Properties.Targets[TargetRules.TargetType.Editor].Rules.GUBP_GetGameFailureEMails_EditorTypeOnly(Branch));
-            return Emails;
-        }
-
-
     }
 
     public class UATTestNode : TestNode
@@ -2650,25 +2587,6 @@ public class GUBP : BuildCommand
             string LogFile = RunUAT(CommandUtils.CmdEnv, WorkingCommandline);
             SaveRecordOfSuccessAndAddToBuildProducts(CommandUtils.ReadAllText(LogFile));
         }
-        public override string FailureEMails(GUBP bp, string Branch)
-        {
-            string Emails;
-            if (GameProj.GameName != bp.Branch.BaseEngineProject.GameName && GameProj.Properties.Targets.ContainsKey(TargetRules.TargetType.Editor))
-            {
-                Emails = MergeSpaceStrings(base.FailureEMails(bp, Branch),
-                   GameProj.Properties.Targets[TargetRules.TargetType.Editor].Rules.GUBP_GetGameFailureEMails_EditorTypeOnly(Branch));
-            }
-            else
-            {
-                Emails = MergeSpaceStrings(base.FailureEMails(bp, Branch),
-                   bp.Branch.BaseEngineProject.Properties.Targets[TargetRules.TargetType.Editor].Rules.GUBP_GetGameFailureEMails_EditorTypeOnly(Branch));
-            }
-            foreach (var Plat in DependsOnCooked)
-            {
-                Emails = MergeSpaceStrings(Emails, Platform.Platforms[Plat].GUBP_GetPlatformFailureEMails(Branch));
-            }
-            return Emails;
-        }
     }
 
     public class GameAggregateNode : HostPlatformAggregateNode
@@ -2706,16 +2624,6 @@ public class GUBP : BuildCommand
         {
             return GameProj.GameName;
         }
-        public override string FailureEMails(GUBP bp, string Branch)
-        {
-            if (GameProj.GameName != bp.Branch.BaseEngineProject.GameName && GameProj.Properties.Targets.ContainsKey(TargetRules.TargetType.Editor))
-            {
-                return MergeSpaceStrings(base.FailureEMails(bp, Branch),
-                   GameProj.Properties.Targets[TargetRules.TargetType.Editor].Rules.GUBP_GetGameFailureEMails_EditorTypeOnly(Branch));
-            }
-            return MergeSpaceStrings(base.FailureEMails(bp, Branch),
-               bp.Branch.BaseEngineProject.Properties.Targets[TargetRules.TargetType.Editor].Rules.GUBP_GetGameFailureEMails_EditorTypeOnly(Branch));
-        }
     };
 
     public string AddNode(GUBPNode Node)
@@ -2962,20 +2870,14 @@ public class GUBP : BuildCommand
             RunAndLog("ectool", Args);
         }
     }
-    string GetEMailListForNode(GUBP bp, string NodeToDo)
+    string GetEMailListForNode(GUBP bp, string NodeToDo, string Causers)
     {        
         var BranchForEmail = "";
         if (P4Enabled)
         {
             BranchForEmail = P4Env.BuildRootP4;
         }
-        var Emails = GUBPNodes[NodeToDo].FailureEMails(this, BranchForEmail);
-        if (Emails == "")
-        {
-            Emails = "Engine-QA[epic]";
-        }
-        Emails = HackEmails(Emails, BranchForEmail, NodeToDo);
-        return Emails;
+        return HackEmails(Causers, BranchForEmail, NodeToDo);
     }
 
     List<P4Connection.ChangeRecord> GetChanges(int LastOutputForChanges, int TopCL, int LastGreen)
@@ -3088,7 +2990,7 @@ public class GUBP : BuildCommand
             string EMails = "";
             if (AddEmailProps)
             {
-                EMails = GetEMailListForNode(bp, NodeToDo);
+                EMails = GetEMailListForNode(bp, NodeToDo, "");
             }
             if (bShowTriggers)
             {
@@ -3649,38 +3551,18 @@ public class GUBP : BuildCommand
             ECProps.Add(string.Format("InProgress/{0}=", NodeToDo));
         }
         {
-            string EmailOnly = ParseParamValue("EmailOnly");
-            if (bFake)
+            var AdditonalEmails = "";
+            if (ParseParam("CIS") && !GUBPNodes[NodeToDo].SendSuccessEmail())
             {
-                EMails = "kellan.carr[epic] gil.gribb[epic]";
-            }
-            else if (!String.IsNullOrEmpty(EmailOnly))
-            {
-                EMails = EmailOnly;
-            }
-            else if (!String.IsNullOrEmpty(ForceAllEmailsTo))
-            {
-                EMails = ForceAllEmailsTo;
-            }
-            else
-            {
-                EMails = GetEMailListForNode(this, NodeToDo);
-                if (ParseParam("CIS") && !GUBPNodes[NodeToDo].SendSuccessEmail())
-                {
-                    EMails = GUBPNode.MergeSpaceStrings(FailCauserEMails, EMails);
-                }
+                AdditonalEmails = FailCauserEMails;
             }
             string AddEmails = ParseParamValue("AddEmails");
             if (!String.IsNullOrEmpty(AddEmails))
             {
-                EMails = GUBPNode.MergeSpaceStrings(AddEmails, EMails);
+                AdditonalEmails = GUBPNode.MergeSpaceStrings(AddEmails, AdditonalEmails);
             }
 
-            {
-                EMails = EMails.Trim().Replace("[epic]", "@epicgames.com");
-                EMails = EMails.Replace("[phosphor]", "@phosphorgames.com");
-                EMails = EMails.Replace("[pitbull]", "@pitbullstudio.co.uk");
-            }
+            EMails = GetEMailListForNode(this, NodeToDo, AdditonalEmails);
 
             ECProps.Add("FailEmails/" + NodeToDo + "=" + EMails);
         }
diff --git a/Engine/Source/Programs/UnrealBuildTool/System/RulesCompiler.cs b/Engine/Source/Programs/UnrealBuildTool/System/RulesCompiler.cs
--- a/Engine/Source/Programs/UnrealBuildTool/System/RulesCompiler.cs
+++ b/Engine/Source/Programs/UnrealBuildTool/System/RulesCompiler.cs
@@ -661,21 +661,6 @@ namespace UnrealBuildTool
             Result.Add("GameAutomationTest", "BuildCookRun -run -RunAutomationTests -unattended -nullrhi -NoP4" + MacOption);
             return Result;
         }
-        /// <summary>
-        /// Allow the platform to setup emails for the GUBP for folks that care about node failures relating to this platform
-        /// </summary>
-        /// <param name="Branch">p4 root of the branch we are running</param>
-        public virtual string GUBP_GetGameFailureEMails_EditorTypeOnly(string Branch)
-        {
-            return "";
-        }
-        /// <summary>
-        /// Allow the Game to set up emails for Promotable and Promotion
-        /// </summary>
-        public virtual string GUBP_GetPromotionEMails_EditorTypeOnly(string Branch)
-        {
-            return "";
-        }
 
         /// <summary>
         /// Return a list of "test name", "UAT command" pairs for testing a monolithic
diff --git a/Engine/Source/UE4Editor.Target.cs b/Engine/Source/UE4Editor.Target.cs
--- a/Engine/Source/UE4Editor.Target.cs
+++ b/Engine/Source/UE4Editor.Target.cs
@@ -109,8 +109,4 @@ public class UE4EditorTarget : TargetRules
         NonCodeProjectNames.Add("TappyChicken", new List<KeyValuePair<UnrealTargetPlatform, UnrealTargetConfiguration>>{new KeyValuePair<UnrealTargetPlatform, UnrealTargetConfiguration>(UnrealTargetPlatform.IOS, UnrealTargetConfiguration.Shipping)});
         return NonCodeProjectNames;
     }
-    public override string GUBP_GetPromotionEMails_EditorTypeOnly(string Branch)
-    {
-        return "Engine-QA[epic] Engine-QA-Team[epic]";
-    }
 }
