UAT: Prevent descendant processes from being killed twice on exit.

diff --git a/Engine/Source/Programs/AutomationTool/ProcessUtils.cs b/Engine/Source/Programs/AutomationTool/ProcessUtils.cs
--- a/Engine/Source/Programs/AutomationTool/ProcessUtils.cs
+++ b/Engine/Source/Programs/AutomationTool/ProcessUtils.cs
@@ -427,12 +427,18 @@ namespace AutomationTool
 			do
 			{
 				bKilledAChild = false;
-				Process[] AllProcs = Process.GetProcesses();
+				// For some reason Process.GetProcesses() sometimes returns the same process twice
+				// So keep track of all processes we already tried to kill
+				var KilledPids = new HashSet<int>();
+				var AllProcs = Process.GetProcesses();
 				foreach (Process KillCandidate in AllProcs)
 				{
-					HashSet<int> VisitedPids = new HashSet<int>();
-					if (ProcessManager.CanBeKilled(KillCandidate.ProcessName) && IsOurDescendant(ProcessToKill, KillCandidate.Id, VisitedPids))
+					var VisitedPids = new HashSet<int>();
+					if (ProcessManager.CanBeKilled(KillCandidate.ProcessName) && 
+						!KilledPids.Contains(KillCandidate.Id) && 
+						IsOurDescendant(ProcessToKill, KillCandidate.Id, VisitedPids))
 					{
+						KilledPids.Add(KillCandidate.Id);
 						CommandUtils.Log("Trying to kill descendant pid={0}, name={1}", KillCandidate.Id, KillCandidate.ProcessName);
 						try
 						{
