Added MacPlatform.RunClient to UAT to handle launching of self-contained, packaged games from UFE

diff --git a/Engine/Source/Programs/AutomationTool/Mac/MacPlatform.Automation.cs b/Engine/Source/Programs/AutomationTool/Mac/MacPlatform.Automation.cs
--- a/Engine/Source/Programs/AutomationTool/Mac/MacPlatform.Automation.cs
+++ b/Engine/Source/Programs/AutomationTool/Mac/MacPlatform.Automation.cs
@@ -117,6 +117,30 @@ public class MacPlatform : Platform
 		Directory.CreateDirectory(CombinePaths(TargetPath, ExeName.StartsWith("UE4Game") ? "Engine" : SC.ShortProjectName, "Binaries", "Mac"));
 	}
 
+	public override ProcessResult RunClient(ERunOptions ClientRunFlags, string ClientApp, string ClientCmdLine, ProjectParams Params)
+	{
+		if (!File.Exists(ClientApp))
+		{
+			if (Directory.Exists(ClientApp + ".app"))
+			{
+				ClientApp += ".app/Contents/MacOS/" + Path.GetFileName(ClientApp);
+			}
+			else
+			{
+				Int32 BaseDirLen = Params.BaseStageDirectory.Length;
+				string StageSubDir = ClientApp.Substring(BaseDirLen, ClientApp.IndexOf("/", BaseDirLen + 1) - BaseDirLen);
+				ClientApp = CombinePaths(Params.BaseStageDirectory, StageSubDir, Params.ShortProjectName + ".app/Contents/MacOS/" + Params.ShortProjectName);
+			}
+		}
+
+		PushDir(Path.GetDirectoryName(ClientApp));
+		// Always start client process and don't wait for exit.
+		ProcessResult ClientProcess = Run(ClientApp, ClientCmdLine, null, ClientRunFlags | ERunOptions.NoWaitForExit);
+		PopDir();
+
+		return ClientProcess;
+	}
+
 	public override bool IsSupported { get { return true; } }
 
 	public override bool ShouldUseManifestForUBTBuilds()
