build script, fix multiple configs make tappy use test

diff --git a/Engine/Source/Programs/AutomationTool/Scripts/GUBP.Automation.cs b/Engine/Source/Programs/AutomationTool/Scripts/GUBP.Automation.cs
--- a/Engine/Source/Programs/AutomationTool/Scripts/GUBP.Automation.cs
+++ b/Engine/Source/Programs/AutomationTool/Scripts/GUBP.Automation.cs
@@ -2256,18 +2256,22 @@ public class GUBP : BuildCommand
             }
         }
 
-        public static string StaticGetFullName(BranchInfo.BranchUProject InGameProj, UnrealTargetPlatform InHostPlatform, List<UnrealTargetPlatform> InClientTargetPlatforms)
+        public static string StaticGetFullName(BranchInfo.BranchUProject InGameProj, UnrealTargetPlatform InHostPlatform, List<UnrealTargetPlatform> InClientTargetPlatforms, List<UnrealTargetConfiguration> InClientConfigs)
         {
             string Infix = "";
             if (InClientTargetPlatforms.Count == 1)
             {
                 Infix = "_" + InClientTargetPlatforms[0].ToString();
             }
+            if (InClientConfigs.Count == 1)
+            {
+                Infix += "_" + InClientConfigs[0].ToString();
+            }
             return InGameProj.GameName + Infix + "_MakeBuild" + HostPlatformNode.StaticGetHostPlatformSuffix(InHostPlatform);
         }
         public override string GetFullName()
         {
-            return StaticGetFullName(GameProj, HostPlatform, ClientTargetPlatforms);
+            return StaticGetFullName(GameProj, HostPlatform, ClientTargetPlatforms, ClientConfigs);
         }
         public override string GameNameIfAnyForTempStorage()
         {
@@ -4236,7 +4240,7 @@ public class GUBP : BuildCommand
                                                     string BuildAgentSharingGroup = NonCodeProject.GameName + "_MakeFormalBuild_" + Plat.ToString();
                                                     GUBPNodes[CookNode.StaticGetFullName(HostPlatform, NonCodeProject, CookedPlatform)].AgentSharingGroup = BuildAgentSharingGroup;
                                                     GUBPNodes[GamePlatformCookedAndCompiledNode.StaticGetFullName(HostPlatform, NonCodeProject, Plat)].AgentSharingGroup = BuildAgentSharingGroup;
-                                                    GUBPNodes[FormalBuildNode.StaticGetFullName(NonCodeProject, HostPlatform, new List<UnrealTargetPlatform>() { Plat })].AgentSharingGroup = BuildAgentSharingGroup;
+                                                    GUBPNodes[FormalBuildNode.StaticGetFullName(NonCodeProject, HostPlatform, new List<UnrealTargetPlatform>() { Plat }, new List<UnrealTargetConfiguration>() { PlatPair.Value })].AgentSharingGroup = BuildAgentSharingGroup;
 
                                                 }
                                             }
diff --git a/Engine/Source/UE4Editor.Target.cs b/Engine/Source/UE4Editor.Target.cs
--- a/Engine/Source/UE4Editor.Target.cs
+++ b/Engine/Source/UE4Editor.Target.cs
@@ -110,7 +110,9 @@ public class UE4EditorTarget : TargetRules
             new List<KeyValuePair<UnrealTargetPlatform, UnrealTargetConfiguration>>
             {
                     new KeyValuePair<UnrealTargetPlatform, UnrealTargetConfiguration>(UnrealTargetPlatform.IOS, UnrealTargetConfiguration.Shipping),
-                    new KeyValuePair<UnrealTargetPlatform, UnrealTargetConfiguration>(UnrealTargetPlatform.Android, UnrealTargetConfiguration.Shipping)
+                    new KeyValuePair<UnrealTargetPlatform, UnrealTargetConfiguration>(UnrealTargetPlatform.IOS, UnrealTargetConfiguration.Test),
+                    new KeyValuePair<UnrealTargetPlatform, UnrealTargetConfiguration>(UnrealTargetPlatform.Android, UnrealTargetConfiguration.Shipping),
+                    new KeyValuePair<UnrealTargetPlatform, UnrealTargetConfiguration>(UnrealTargetPlatform.Android, UnrealTargetConfiguration.Test)
             }
         );
         return NonCodeProjectNames;
diff --git a/Engine/Source/UE4Game.Target.cs b/Engine/Source/UE4Game.Target.cs
--- a/Engine/Source/UE4Game.Target.cs
+++ b/Engine/Source/UE4Game.Target.cs
@@ -94,6 +94,6 @@ public class UE4GameTarget : TargetRules
     }
     public override List<UnrealTargetConfiguration> GUBP_GetConfigs_MonolithicOnly(UnrealTargetPlatform HostPlatform, UnrealTargetPlatform Platform)
     {
-        return new List<UnrealTargetConfiguration> { UnrealTargetConfiguration.Development, UnrealTargetConfiguration.Shipping };
+        return new List<UnrealTargetConfiguration> { UnrealTargetConfiguration.Development, UnrealTargetConfiguration.Shipping, UnrealTargetConfiguration.Test };
     }
 }
