UAT: -p4cl will now automatically imply -p4

diff --git a/Engine/Source/Programs/AutomationTool/Automation.cs b/Engine/Source/Programs/AutomationTool/Automation.cs
--- a/Engine/Source/Programs/AutomationTool/Automation.cs
+++ b/Engine/Source/Programs/AutomationTool/Automation.cs
@@ -102,7 +102,7 @@ namespace AutomationTool
 		public static CommandLineArg ForceLocal = new CommandLineArg("-ForceLocal");
 		public static CommandLineArg NoP4 = new CommandLineArg("-NoP4");
 		public static CommandLineArg P4 = new CommandLineArg("-P4");
-		public static CommandLineArg P4CL = new CommandLineArg("-P4CL");
+		public static CommandLineArg P4CL = new CommandLineArg("-P4CL", () => { P4.Set(); });
 		public static CommandLineArg Preprocess = new CommandLineArg("-Preprocess");
 		public static CommandLineArg NoCompile = new CommandLineArg("-NoCompile");
 		public static CommandLineArg NoCompileEditor = new CommandLineArg("-NoCompileEditor");
@@ -132,7 +132,7 @@ AutomationTool.exe [-verbose] [-compileonly] [-p4] Command0 [-Arg0 -Arg1 -Arg2 
 	[Help("verbose", "Enables verbose logging")]
 	[Help("nop4", "Disables Perforce functionality (default if not run on a build machine)")]
 	[Help("p4", "Enables Perforce functionality (default if run on a build machine)")]
-	[Help("p4cl", "Enables automatic 'have' changelist detection in local runs")]
+	[Help("p4cl", "Enables automatic 'have' changelist detection in local runs. Implies -p4")]
 	[Help("compileonly", "Does not run any commands, only compiles them")]
 	[Help("forcelocal", "Forces local execution")]
 	[Help("help", "Displays help")]
