build script, limit CL descriptions to first line, allow agent sharing groups to be used for node selection.Fix IOS on Mac for formal builds.

diff --git a/Engine/Source/Programs/AutomationTool/Scripts/GUBP.Automation.cs b/Engine/Source/Programs/AutomationTool/Scripts/GUBP.Automation.cs
--- a/Engine/Source/Programs/AutomationTool/Scripts/GUBP.Automation.cs
+++ b/Engine/Source/Programs/AutomationTool/Scripts/GUBP.Automation.cs
@@ -2924,7 +2924,12 @@ public class GUBP : BuildCommand
         var ChangeRecords = GetChanges(LastOutputForChanges, TopCL, LastGreen);
         foreach (var Record in ChangeRecords)
         {
-            Log("         {0} {1} {2}", Record.CL, Record.UserEmail, Record.Summary);
+            var Summary = Record.Summary.Replace("\r", "\n");
+            if (Summary.IndexOf("\n") > 0)
+            {
+                Summary = Summary.Substring(0, Summary.IndexOf("\n"));
+            }
+            Log("         {0} {1} {2}", Record.CL, Record.UserEmail, Summary);
         }
         return TopCL;
     }
@@ -4238,6 +4243,10 @@ public class GUBP : BuildCommand
                                                         CookNode.StaticGetFullName(HostPlatform, Branch.BaseEngineProject, CookedPlatform));
 
                                                     string BuildAgentSharingGroup = NonCodeProject.GameName + "_MakeFormalBuild_" + Plat.ToString();
+                                                    if (Plat == UnrealTargetPlatform.IOS) // Mac-IOS trashes build products, so we need to use different agents
+                                                    {
+                                                        BuildAgentSharingGroup = "";
+                                                    }
                                                     GUBPNodes[CookNode.StaticGetFullName(HostPlatform, NonCodeProject, CookedPlatform)].AgentSharingGroup = BuildAgentSharingGroup;
                                                     GUBPNodes[GamePlatformCookedAndCompiledNode.StaticGetFullName(HostPlatform, NonCodeProject, Plat)].AgentSharingGroup = BuildAgentSharingGroup;
                                                     GUBPNodes[FormalBuildNode.StaticGetFullName(NonCodeProject, HostPlatform, new List<UnrealTargetPlatform>() { Plat }, new List<UnrealTargetConfiguration>() { PlatPair.Value })].AgentSharingGroup = BuildAgentSharingGroup;
@@ -4582,18 +4591,23 @@ if (HostPlatform == UnrealTargetPlatform.Mac) continue; //temp hack till mac aut
                 foreach (var NodeArg in Nodes)
                 {
                     var NodeName = NodeArg.Trim();
+                    bool bFoundAnything = false;
                     if (!String.IsNullOrEmpty(NodeName))
                     {
                         foreach (var Node in GUBPNodes)
                         {
-                            if (Node.Value.GetFullName().Equals(NodeArg, StringComparison.InvariantCultureIgnoreCase))
+                            if (Node.Value.GetFullName().Equals(NodeArg, StringComparison.InvariantCultureIgnoreCase) ||
+                                Node.Value.AgentSharingGroup.Equals(NodeArg, StringComparison.InvariantCultureIgnoreCase)
+                                )
                             {
-                                NodesToDo.Add(Node.Key);
-                                NodeName = null;
-                                break;
+                                if (!NodesToDo.Contains(Node.Key))
+                                {
+                                    NodesToDo.Add(Node.Key);
+                                }
+                                bFoundAnything = true;
                             }
                         }
-                        if (NodeName != null)
+                        if (!bFoundAnything)
                         {
                             throw new AutomationException("Could not find node named {0}", NodeName);
                         }
