build script, don't label launcher for CIS, internal tools only every 40m

diff --git a/Engine/Source/Programs/AutomationTool/Scripts/GUBP.Automation.cs b/Engine/Source/Programs/AutomationTool/Scripts/GUBP.Automation.cs
--- a/Engine/Source/Programs/AutomationTool/Scripts/GUBP.Automation.cs
+++ b/Engine/Source/Programs/AutomationTool/Scripts/GUBP.Automation.cs
@@ -824,6 +824,11 @@ public class GUBP : BuildCommand
         {
             return base.Priority() - 2;
         }
+        public override int CISFrequencyQuantumShift(GUBP bp)
+        {
+            int Result = base.CISFrequencyQuantumShift(bp) + 1;
+            return Result;
+        }
         public override bool DeleteBuildProducts()
         {
             return true;
@@ -908,6 +913,11 @@ public class GUBP : BuildCommand
         {
             return StaticGetFullName(HostPlatform, ProgramTarget);
         }
+        public override int CISFrequencyQuantumShift(GUBP bp)
+        {
+            int Result = base.CISFrequencyQuantumShift(bp) + 1;
+            return Result;
+        }
         public override float Priority()
         {
             return base.Priority() + 3;
