Manually integrating RocketBuild and GUBP changes from UE4-Releases/4.1 branch.

diff --git a/Engine/Source/Programs/AutomationTool/Scripts/GUBP.Automation.cs b/Engine/Source/Programs/AutomationTool/Scripts/GUBP.Automation.cs
--- a/Engine/Source/Programs/AutomationTool/Scripts/GUBP.Automation.cs
+++ b/Engine/Source/Programs/AutomationTool/Scripts/GUBP.Automation.cs
@@ -541,7 +541,7 @@ public class GUBP : BuildCommand
         public override UE4Build.BuildAgenda GetAgenda(GUBP bp)
         {
             var Agenda = new UE4Build.BuildAgenda();
-            if (HostPlatform != UnrealTargetPlatform.Mac)
+			if (HostPlatform != UnrealTargetPlatform.Mac && !GUBP.bForceIncrementalCompile)
             {
                 Agenda.DotNetProjects.AddRange(
                     new string[] 
@@ -709,28 +709,34 @@ public class GUBP : BuildCommand
 
             if (HostPlatform != UnrealTargetPlatform.Mac)
             {
-                Agenda.DotNetProjects.AddRange(
-                    new string[] 
-			    {
-                    CombinePaths(@"Engine\Source\Programs\UnrealControls\UnrealControls.csproj"),
-			    }
-                    );
+				if (!GUBP.bForceIncrementalCompile)
+				{
+					Agenda.DotNetProjects.AddRange(
+						new string[] 
+						{
+							CombinePaths(@"Engine\Source\Programs\UnrealControls\UnrealControls.csproj"),
+						}
+					);
+				}
                 Agenda.DotNetSolutions.AddRange(
-                    new string[] 
-			        {
-				        CombinePaths(@"Engine\Source\Programs\NetworkProfiler\NetworkProfiler.sln"),   
-			        }
+						new string[] 
+						{
+							CombinePaths(@"Engine\Source\Programs\NetworkProfiler\NetworkProfiler.sln"),   
+						}
                     );
-                Agenda.SwarmProject = CombinePaths(@"Engine\Source\Programs\UnrealSwarm\UnrealSwarm.sln");
+                if (!GUBP.bForceIncrementalCompile)
+                {
+					Agenda.SwarmProject = CombinePaths(@"Engine\Source\Programs\UnrealSwarm\UnrealSwarm.sln");
+				}
                 Agenda.IOSDotNetProjects.AddRange(
                         new string[]
-                    {
-                        CombinePaths(@"Engine\Source\Programs\IOS\iPhonePackager\iPhonePackager.csproj"),
-                        CombinePaths(@"Engine\Source\Programs\IOS\DeploymentServer\DeploymentServer.csproj"),
-                        CombinePaths(@"Engine\Source\Programs\IOS\MobileDeviceInterface\MobileDeviceInterface.csproj"),
-                        CombinePaths(@"Engine\Source\Programs\IOS\DeploymentInterface\DeploymentInterface.csproj"),
-                    }
-                        );
+						{
+							CombinePaths(@"Engine\Source\Programs\IOS\iPhonePackager\iPhonePackager.csproj"),
+							CombinePaths(@"Engine\Source\Programs\IOS\DeploymentServer\DeploymentServer.csproj"),
+							CombinePaths(@"Engine\Source\Programs\IOS\MobileDeviceInterface\MobileDeviceInterface.csproj"),
+							CombinePaths(@"Engine\Source\Programs\IOS\DeploymentInterface\DeploymentInterface.csproj"),
+						}
+                    );
             }
 
             string AddArgs = "-nobuilduht -skipactionhistory" + bp.RocketUBTArgs(); ;
@@ -1209,7 +1215,7 @@ public class GUBP : BuildCommand
             }
             foreach (var Target in Targets)
             {
-                foreach (var FileToCopy in CommandUtils.FindFiles("*.h", true, CommandUtils.CombinePaths(CmdEnv.LocalRoot, @"Engine\Intermediate\Build\", HostPlatform.ToString(), Target)))
+                foreach (var FileToCopy in CommandUtils.FindFiles("*.h", true, CommandUtils.CombinePaths(CmdEnv.LocalRoot, @"Engine\Intermediate\Build\", HostPlatform.ToString(), Target, "Inc")))
                 {
                     AddBuildProduct(FileToCopy);
                 }
@@ -1260,33 +1266,32 @@ public class GUBP : BuildCommand
                 throw new AutomationException("these rocket header node require real mac path.");
             }
             BuildProducts = new List<string>();
-
-            var PlatformDir = TargetPlatform.ToString();
-            if (TargetPlatform == UnrealTargetPlatform.Android)
-            {
-                PlatformDir = "Android-armv7";
-            }
-            if (TargetPlatform == HostPlatform)
-            {
-                PlatformDir = CommandUtils.CombinePaths(PlatformDir, "UE4Game"); // the other ones overlap with the editor, so we don't want them here
-            }
-            foreach (var FileToCopy in CommandUtils.FindFiles("*.h", true, CommandUtils.CombinePaths(CmdEnv.LocalRoot, @"Engine\Intermediate\Build\", PlatformDir)))
-            {
-                AddBuildProduct(FileToCopy);
-            }
-
-            // these may not be any new build products here, but we will check anyway
-            var EnginePluginsDirectory = Path.Combine(CommandUtils.CmdEnv.LocalRoot, "Engine/Plugins");
-            var EnginePlugins = new List<PluginInfo>();
-            Plugins.FindPluginsIn(EnginePluginsDirectory, PluginInfo.LoadedFromType.Engine, ref EnginePlugins);
-
-            foreach (var EnginePlugin in EnginePlugins)
-            {
-                foreach (var FileToCopy in CommandUtils.FindFiles("*", true, CommandUtils.CombinePaths(EnginePlugin.Directory, @"Intermediate\Build", PlatformDir, "Inc")))
-                {
-                    AddBuildProduct(FileToCopy);
-                }
-            }
+			if (TargetPlatform != HostPlatform)
+			{
+				// host platform overlaps with the editor, so we don't want them here
+				var PlatformDir = TargetPlatform.ToString();
+				if (TargetPlatform == UnrealTargetPlatform.Android)
+				{
+					PlatformDir = "Android-armv7";
+				}
+				foreach (var FileToCopy in CommandUtils.FindFiles("*.h", true, CommandUtils.CombinePaths(CmdEnv.LocalRoot, @"Engine\Intermediate\Build\", PlatformDir, "Inc")))
+				{
+					AddBuildProduct(FileToCopy);
+				}
+
+				// these may not be any new build products here, but we will check anyway
+				var EnginePluginsDirectory = Path.Combine(CommandUtils.CmdEnv.LocalRoot, "Engine/Plugins");
+				var EnginePlugins = new List<PluginInfo>();
+				Plugins.FindPluginsIn(EnginePluginsDirectory, PluginInfo.LoadedFromType.Engine, ref EnginePlugins);
+
+				foreach (var EnginePlugin in EnginePlugins)
+				{
+					foreach (var FileToCopy in CommandUtils.FindFiles("*", true, CommandUtils.CombinePaths(EnginePlugin.Directory, @"Intermediate\Build", PlatformDir, "Inc")))
+					{
+						AddBuildProduct(FileToCopy);
+					}
+				}
+			}
             RemoveOveralppingBuildProducts();
             if (BuildProducts.Count == 0)
             {
