build script, add StompCheck to the GUBP

diff --git a/Engine/Source/Programs/AutomationTool/Scripts/GUBP.Automation.cs b/Engine/Source/Programs/AutomationTool/Scripts/GUBP.Automation.cs
--- a/Engine/Source/Programs/AutomationTool/Scripts/GUBP.Automation.cs
+++ b/Engine/Source/Programs/AutomationTool/Scripts/GUBP.Automation.cs
@@ -166,6 +166,7 @@ public class GUBP : BuildCommand
         public List<string> FullNamesOfPseudosependencies = new List<string>(); //these are really only used for sorting. We want the editor to fail before the monolithics. Think of it as "can't possibly be useful without".
         public List<string> BuildProducts = null;
         public List<string> AllDependencyBuildProducts = null;
+        public List<string> AllDependencies = null;
         public string AgentSharingGroup = "";
         public int ComputedDependentCISFrequencyQuantumShift = -1;
 
@@ -358,6 +359,13 @@ public class GUBP : BuildCommand
                 AllDependencyBuildProducts.Add(Info.FullName);
             }
         }
+        public void AddAllDependent(string Node)
+        {
+            if (!AllDependencies.Contains(Node))
+            {
+                AllDependencies.Add(Node);
+            }
+        }
         public void RemoveOveralppingBuildProducts()
         {
             foreach (var ToRemove in AllDependencyBuildProducts)
@@ -3668,6 +3676,7 @@ public class GUBP : BuildCommand
     [Help("CIS", "This is a CIS run, assign TimeIndex based on the history.")]
     [Help("ForceIncrementalCompile", "make sure all compiles are incremental")]
     [Help("AutomatedTesting", "Allow automated testing, currently disabled.")]
+    [Help("StompCheck", "Look for stomped build products.")]
 
     public override void ExecuteBuild()
     {
@@ -5292,41 +5301,45 @@ if (HostPlatform == UnrealTargetPlatform.Mac) continue; //temp hack till mac aut
             }
 
             GUBPNodes[NodeToDo].AllDependencyBuildProducts = new List<string>();
-            if (GUBPNodes[NodeToDo].FullNamesOfDependencies == null)
+            GUBPNodes[NodeToDo].AllDependencies = new List<string>();
+            foreach (var Dep in GUBPNodes[NodeToDo].FullNamesOfDependencies)
             {
-                if (ParseParamValue("OnlyNode", "") == "")
+                GUBPNodes[NodeToDo].AddAllDependent(Dep);
+                if (GUBPNodes[Dep].AllDependencies == null)
                 {
-                    throw new AutomationException("Node {0} was not processed yet?", NodeToDo);
+                    if (!bOnlyNode)
+                    {
+                        throw new AutomationException("Node {0} was not processed yet3?  Processing {1}", Dep, NodeToDo);
+                    }
                 }
-            }
-            else
-            {
-                foreach (var Dep in GUBPNodes[NodeToDo].FullNamesOfDependencies)
+                foreach (var DepDep in GUBPNodes[Dep].AllDependencies)
                 {
-                    if (GUBPNodes[Dep].BuildProducts == null)
+                    GUBPNodes[NodeToDo].AddAllDependent(DepDep);
+                }
+                if (GUBPNodes[Dep].BuildProducts == null)
+                {
+                    if (!bOnlyNode)
                     {
-                        if (ParseParamValue("OnlyNode", "") == "")
-                        {
-                            throw new AutomationException("Node {0} was not processed yet2? Processing {1}", Dep, NodeToDo);
-                        }
+                        throw new AutomationException("Node {0} was not processed yet? Processing {1}", Dep, NodeToDo);
                     }
-                    else
+                }
+                else
+                {
+                    foreach (var Prod in GUBPNodes[Dep].BuildProducts)
                     {
-                        foreach (var Prod in GUBPNodes[Dep].BuildProducts)
-                        {
-                            GUBPNodes[NodeToDo].AddDependentBuildProduct(Prod);
-                        }
-                        if (GUBPNodes[Dep].AllDependencyBuildProducts == null)
-                        {
-                            throw new AutomationException("Node {0} was not processed yet3?  Processing {1}", Dep, NodeToDo);
-                        }
-                        foreach (var Prod in GUBPNodes[Dep].AllDependencyBuildProducts)
-                        {
-                            GUBPNodes[NodeToDo].AddDependentBuildProduct(Prod);
-                        }
+                        GUBPNodes[NodeToDo].AddDependentBuildProduct(Prod);
+                    }
+                    if (GUBPNodes[Dep].AllDependencyBuildProducts == null)
+                    {
+                        throw new AutomationException("Node {0} was not processed yet2?  Processing {1}", Dep, NodeToDo);
+                    }
+                    foreach (var Prod in GUBPNodes[Dep].AllDependencyBuildProducts)
+                    {
+                        GUBPNodes[NodeToDo].AddDependentBuildProduct(Prod);
                     }
                 }
             }
+
             string NodeStoreName = StoreName + "-" + GUBPNodes[NodeToDo].GetFullName();
             
             string GameNameIfAny = GUBPNodes[NodeToDo].GameNameIfAnyForTempStorage();
@@ -5386,6 +5399,27 @@ if (HostPlatform == UnrealTargetPlatform.Mac) continue; //temp hack till mac aut
                     if (!GUBPNodes[NodeToDo].IsAggregate())
                     {
                         StoreToTempStorage(CmdEnv, NodeStoreName, GUBPNodes[NodeToDo].BuildProducts, !bSaveSharedTempStorage, GameNameIfAny, StorageRootIfAny);
+
+                        if (ParseParam("StompCheck"))
+                        {
+                            foreach (var Dep in GUBPNodes[NodeToDo].AllDependencies)
+                            {
+                                try
+                                {
+                                    bool WasLocal;
+                                    RetrieveFromTempStorage(CmdEnv, NodeStoreName, out WasLocal, GameNameIfAny, StorageRootIfAny);
+                                    if (!WasLocal)
+                                    {
+                                        throw new AutomationException("Retrieve was not local?");
+                                    }
+                                }
+                                catch(Exception Ex)
+                                {
+                                    throw new AutomationException("Node {0} stomped Node {1}   Ex: {2}", NodeToDo, Dep, LogUtils.FormatException(Ex));
+                                }
+                            }
+                        }
+
                     }
                 }
                 catch (Exception Ex)
