Fix for double build of iOS when trying to codesign
#ue4
#ios

diff --git a/Engine/Source/Programs/AutomationTool/IOS/IOSPlatform.Automation.cs b/Engine/Source/Programs/AutomationTool/IOS/IOSPlatform.Automation.cs
--- a/Engine/Source/Programs/AutomationTool/IOS/IOSPlatform.Automation.cs
+++ b/Engine/Source/Programs/AutomationTool/IOS/IOSPlatform.Automation.cs
@@ -203,34 +203,28 @@ public class IOSPlatform : Platform
 	{
 		// first check for ue4.xcodeproj
 		bWasGenerated = false;
-		string XcodeProj = RawProjectPath.Replace(".uproject", ".xcodeproj");
+		string XcodeProj = RawProjectPath.Replace(".uproject", "_IOS.xcodeproj");
 		Console.WriteLine ("Project: " + XcodeProj);
 		if (!Directory.Exists (XcodeProj))
 		{
-			// project.xcodeproj doesn't exist, so check for ue4.xcodeproj
-			XcodeProj = CombinePaths(CmdEnv.LocalRoot, "UE4.xcodeproj");
+			// project.xcodeproj doesn't exist, so generate temp project
+			string Arguments = "-project=\"" + RawProjectPath + "\"";
+			Arguments += " -platforms=IOS -game -nointellisense -iosdeployonly -ignorejunk";
+			string Script = CombinePaths(CmdEnv.LocalRoot, "Engine/Build/BatchFiles/Mac/GenerateProjectFiles.sh");
+			if (GlobalCommandLine.Rocket)
+			{
+				Script = CombinePaths(CmdEnv.LocalRoot, "Engine/Build/BatchFiles/Mac/RocketGenerateProjectFiles.sh");
+			}
+			string CWD = Directory.GetCurrentDirectory ();
+			Directory.SetCurrentDirectory (Path.GetDirectoryName (Script));
+			Run (Script, Arguments, null, ERunOptions.Default);
+			bWasGenerated = true;
+			Directory.SetCurrentDirectory (CWD);
+
 			if (!Directory.Exists (XcodeProj))
 			{
-				// ue4.xcodeproj doesn't exist, so generate temp project
-				string Arguments = "-project=\"" + RawProjectPath + "\"";
-				Arguments += " -game -nointellisense";
-				string Script = CombinePaths(CmdEnv.LocalRoot, "Engine/Build/BatchFiles/Mac/GenerateProjectFiles.sh");
-				if (GlobalCommandLine.Rocket)
-				{
-					Script = CombinePaths(CmdEnv.LocalRoot, "Engine/Build/BatchFiles/Mac/RocketGenerateProjectFiles.sh");
-				}
-				string CWD = Directory.GetCurrentDirectory ();
-				Directory.SetCurrentDirectory (Path.GetDirectoryName (Script));
-				Run (Script, Arguments, null, ERunOptions.Default);
-				bWasGenerated = true;
-				Directory.SetCurrentDirectory (CWD);
-
-				XcodeProj = RawProjectPath.Replace(".uproject", ".xcodeproj"); 
-				if (!Directory.Exists (XcodeProj))
-				{
-					// something very bad happened
-					throw new AutomationException("iOS couldn't find the appropriate Xcode Project");
-				}
+				// something very bad happened
+				throw new AutomationException("iOS couldn't find the appropriate Xcode Project");
 			}
 		}
 
diff --git a/Engine/Source/Programs/UnrealBuildTool/IOS/IOSToolChain.cs b/Engine/Source/Programs/UnrealBuildTool/IOS/IOSToolChain.cs
--- a/Engine/Source/Programs/UnrealBuildTool/IOS/IOSToolChain.cs
+++ b/Engine/Source/Programs/UnrealBuildTool/IOS/IOSToolChain.cs
@@ -824,7 +824,7 @@ namespace UnrealBuildTool
 					// generate the dummy project so signing works
 					if (Target.GameName == "UE4Game" || Target.GameName == "UE4Client" || Utils.IsFileUnderDirectory(Target.ProjectDirectory + "/" + Target.GameName + ".uproject", Path.GetFullPath("../..")))
 					{
-						UnrealBuildTool.GenerateProjectFiles (new XcodeProjectFileGenerator (), new string[] {"-platforms=IOS", "-NoIntellIsense"});
+						UnrealBuildTool.GenerateProjectFiles (new XcodeProjectFileGenerator (), new string[] {"-platforms=IOS", "-NoIntellIsense", "-iosdeployonly", "-ignorejunk"});
 						Project = Path.GetFullPath("../..") + "/UE4_IOS.xcodeproj";
 					}
 					else
diff --git a/Engine/Source/Programs/UnrealBuildTool/System/UnrealBuildTool.cs b/Engine/Source/Programs/UnrealBuildTool/System/UnrealBuildTool.cs
--- a/Engine/Source/Programs/UnrealBuildTool/System/UnrealBuildTool.cs
+++ b/Engine/Source/Programs/UnrealBuildTool/System/UnrealBuildTool.cs
@@ -633,6 +633,7 @@ namespace UnrealBuildTool
                     var bGenerateMakefiles = false;
 					var bValidPlatformsOnly = false;
 					var bSpecificModulesOnly = false;
+					var bIgnoreJunk = false;
 
 					// We need to be able to identify if one of the arguments is the platform...
 					// Leverage the existing parser function in UEBuildTarget to get this information.
@@ -766,6 +767,10 @@ namespace UnrealBuildTool
 						{
 							UEBuildConfiguration.bCompileSpeedTree = false;
 						}
+						else if (LowercaseArg == "-ignorejunk")
+						{
+							bIgnoreJunk = true;
+						}
 						else if (LowercaseArg == "-define")
 						{
 							// Skip -define
@@ -831,7 +836,7 @@ namespace UnrealBuildTool
 					ProjectFileGenerator.bGenerateProjectFiles = false;
 
 					// now that we know the available platforms, we can delete other platforms' junk. if we're only building specific modules from the editor, don't touch anything else (it may be in use).
-					if (!bSpecificModulesOnly)
+					if (!bSpecificModulesOnly && !bIgnoreJunk)
 					{
 						JunkDeleter.DeleteJunk();
 					}
diff --git a/Engine/Source/Programs/UnrealBuildTool/System/XcodeProjectFileGenerator.cs b/Engine/Source/Programs/UnrealBuildTool/System/XcodeProjectFileGenerator.cs
--- a/Engine/Source/Programs/UnrealBuildTool/System/XcodeProjectFileGenerator.cs
+++ b/Engine/Source/Programs/UnrealBuildTool/System/XcodeProjectFileGenerator.cs
@@ -1624,12 +1624,16 @@ namespace UnrealBuildTool
 
 						// @todo: Remove target platform param and merge Mac and iOS targets. For now BuildTarget knows how to build iOS, but cannot run iOS apps, so we need separate DeployTarget.
 						bool bIsMacOnly = !SupportedPlatforms.Contains(UnrealTargetPlatform.IOS);
+
 						XcodeProjectTarget BuildTarget = new XcodeProjectTarget(TargetName + " - Mac", TargetName, "Legacy", "", UnrealTargetPlatform.Mac, bIsMacOnly);
-						ProjectTargets.Add(BuildTarget);
+						if (!bGeneratingRunIOSProject)
+						{
+							ProjectTargets.Add(BuildTarget);
+						}
 
 						if (ProjectFilePlatform.HasFlag(XcodeProjectFilePlatform.iOS) && SupportedPlatforms.Contains(UnrealTargetPlatform.IOS))
 						{
-							if (bGeneratingRocketProjectFiles && TargetName == "UE4Game")
+							if ((bGeneratingRocketProjectFiles && TargetName == "UE4Game") || bGeneratingRunIOSProject)
 							{
 								// Generate Framework references.
 								List<XcodeFrameworkRef> FrameworkRefs = new List<XcodeFrameworkRef>();
