UAT: Removing -p4cl, adding DoesNotNettP4CL attribute

diff --git a/Engine/Source/Programs/AutomationTool/Automation.cs b/Engine/Source/Programs/AutomationTool/Automation.cs
--- a/Engine/Source/Programs/AutomationTool/Automation.cs
+++ b/Engine/Source/Programs/AutomationTool/Automation.cs
@@ -102,7 +102,6 @@ namespace AutomationTool
 		public static CommandLineArg ForceLocal = new CommandLineArg("-ForceLocal");
 		public static CommandLineArg NoP4 = new CommandLineArg("-NoP4");
 		public static CommandLineArg P4 = new CommandLineArg("-P4");
-		public static CommandLineArg P4CL = new CommandLineArg("-P4CL", () => { P4.Set(); });
 		public static CommandLineArg Preprocess = new CommandLineArg("-Preprocess");
 		public static CommandLineArg NoCompile = new CommandLineArg("-NoCompile");
 		public static CommandLineArg NoCompileEditor = new CommandLineArg("-NoCompileEditor");
@@ -132,7 +131,6 @@ AutomationTool.exe [-verbose] [-compileonly] [-p4] Command0 [-Arg0 -Arg1 -Arg2 
 	[Help("verbose", "Enables verbose logging")]
 	[Help("nop4", "Disables Perforce functionality (default if not run on a build machine)")]
 	[Help("p4", "Enables Perforce functionality (default if run on a build machine)")]
-	[Help("p4cl", "Enables automatic 'have' changelist detection in local runs. Implies -p4")]
 	[Help("compileonly", "Does not run any commands, only compiles them")]
 	[Help("forcelocal", "Forces local execution")]
 	[Help("help", "Displays help")]
diff --git a/Engine/Source/Programs/AutomationTool/LocalP4Environment.cs b/Engine/Source/Programs/AutomationTool/LocalP4Environment.cs
--- a/Engine/Source/Programs/AutomationTool/LocalP4Environment.cs
+++ b/Engine/Source/Programs/AutomationTool/LocalP4Environment.cs
@@ -71,7 +71,7 @@ namespace AutomationTool
 			CommandUtils.ConditionallySetEnvVar(EnvVarNames.ClientRoot, ClientRootPath);
 
 			var CLString = CommandUtils.GetEnvVar(EnvVarNames.Changelist, null);
-			if (String.IsNullOrEmpty(CLString) && GlobalCommandLine.P4CL)
+			if (String.IsNullOrEmpty(CLString) && CommandUtils.P4CLRequired)
 			{
                 CLString = DetectCurrentCL(Connection, ClientRootPath);
 			}
diff --git a/Engine/Source/Programs/AutomationTool/P4Utils.cs b/Engine/Source/Programs/AutomationTool/P4Utils.cs
--- a/Engine/Source/Programs/AutomationTool/P4Utils.cs
+++ b/Engine/Source/Programs/AutomationTool/P4Utils.cs
@@ -14,6 +14,10 @@ namespace AutomationTool
 	{
 	}
 
+	public class DoesNotNeedP4CLAttribute : Attribute
+	{
+	}
+
 	public class P4Exception : AutomationException
 	{
 		public P4Exception() { }
@@ -246,7 +250,7 @@ namespace AutomationTool
 		#endregion
 
 		/// <summary>
-		/// Check is P4 is commands are supported.
+		/// Check if P4 is supported.
 		/// </summary>		
 		public static bool P4Enabled
 		{
@@ -265,6 +269,26 @@ namespace AutomationTool
 		}
 		private static bool? bP4Enabled;
 
+		/// <summary>
+		/// Check if P4CL is required.
+		/// </summary>		
+		public static bool P4CLRequired
+		{
+			get
+			{
+				if (!bP4CLRequired.HasValue)
+				{
+					throw new AutomationException("Trying to access P4CLRequired property before it was initialized.");
+				}
+				return (bool)bP4CLRequired;
+			}
+			private set
+			{
+				bP4CLRequired = value;
+			}
+		}
+		private static bool? bP4CLRequired;
+
 		/// <summary>
 		/// Throws an exception when P4 is disabled. This should be called in every P4 function.
 		/// </summary>
@@ -323,17 +347,19 @@ namespace AutomationTool
 			if (Automation.IsBuildMachine)
 			{
 				P4Enabled = !GlobalCommandLine.NoP4;
+				P4CLRequired = P4Enabled;
 			}
 			else
 			{
-				P4Enabled = GlobalCommandLine.P4;
-				if (!P4Enabled && !GlobalCommandLine.NoP4)
-				{
-					// Check if any of the commands to execute require P4
-					P4Enabled = CheckIfCommandsRequireP4(CommandsToExecute, Commands);
-				}
+				bool bRequireP4;
+				bool bRequireCL;
+				CheckIfCommandsRequireP4(CommandsToExecute, Commands, out bRequireP4, out bRequireCL);
+
+				P4Enabled = GlobalCommandLine.P4 || bRequireP4;
+				P4CLRequired = GlobalCommandLine.P4 || bRequireCL;
 			}
 			Log("P4Enabled={0}", P4Enabled);
+			Log("P4CLRequired={0}", P4CLRequired);
 		}
 
 		/// <summary>
@@ -341,23 +367,30 @@ namespace AutomationTool
 		/// </summary>
 		/// <param name="CommandsToExecute">List of commands to be executed.</param>
 		/// <param name="Commands">Commands.</param>
-		/// <returns>True if any of the commands to execute has [RequireP4] attribute.</returns>
-		private static bool CheckIfCommandsRequireP4(List<CommandInfo> CommandsToExecute, CaselessDictionary<Type> Commands)
+		private static void CheckIfCommandsRequireP4(List<CommandInfo> CommandsToExecute, CaselessDictionary<Type> Commands, out bool bRequireP4, out bool bRequireCL)
 		{
+			bRequireP4 = false;
+			bRequireCL = false;
+
 			foreach (var CommandInfo in CommandsToExecute)
 			{
 				Type Command;
 				if (Commands.TryGetValue(CommandInfo.CommandName, out Command))
 				{
-					var RequireP4Attributes = Command.GetCustomAttributes(typeof(RequireP4Attribute), true);
+					var RequireP4Attributes = Command.GetCustomAttributes(typeof(RequireP4Attribute), true);	
 					if (!CommandUtils.IsNullOrEmpty(RequireP4Attributes))
 					{
 						Log("Command {0} requires P4 functionality.", Command.Name);
-						return true;
+						bRequireP4 = true;
+
+						var DoesNotNeedP4CLAttributes = Command.GetCustomAttributes(typeof(DoesNotNeedP4CLAttribute), true);
+						if (CommandUtils.IsNullOrEmpty(DoesNotNeedP4CLAttributes))
+						{
+							bRequireCL = true;
+						}
 					}
 				}
 			}
-			return false;
 		}
 	}
 
diff --git a/Engine/Source/Programs/AutomationTool/Scripts/Tests.Automation.cs b/Engine/Source/Programs/AutomationTool/Scripts/Tests.Automation.cs
--- a/Engine/Source/Programs/AutomationTool/Scripts/Tests.Automation.cs
+++ b/Engine/Source/Programs/AutomationTool/Scripts/Tests.Automation.cs
@@ -15,6 +15,7 @@ using UnrealBuildTool;
 
 [Help("Tests P4 functionality. Runs 'p4 info'.")]
 [RequireP4]
+[DoesNotNeedP4CL]
 class TestP4_Info : BuildCommand
 {
 	public override void ExecuteBuild()
