Fix for iOS build not having the appropriate plist
#ue4
#ios

diff --git a/Engine/Source/Programs/AutomationTool/IOS/IOSPlatform.Automation.cs b/Engine/Source/Programs/AutomationTool/IOS/IOSPlatform.Automation.cs
--- a/Engine/Source/Programs/AutomationTool/IOS/IOSPlatform.Automation.cs
+++ b/Engine/Source/Programs/AutomationTool/IOS/IOSPlatform.Automation.cs
@@ -190,7 +190,7 @@ public class IOSPlatform : Platform
 		else
 		{
 			// code sign the app
-			CodeSign(Path.GetDirectoryName(Params.ProjectGameExeFilename), Path.GetFileNameWithoutExtension(Params.ProjectGameExeFilename), Params.RawProjectPath, SC.StageTargetConfigurations[0], Params.Distribution);
+			CodeSign(Path.GetDirectoryName(Params.ProjectGameExeFilename), Path.GetFileNameWithoutExtension(Params.ProjectGameExeFilename), Params.RawProjectPath, SC.StageTargetConfigurations[0], SC.LocalRoot, Params.ShortProjectName, Path.GetDirectoryName(Params.RawProjectPath), SC.IsCodeBasedProject, Params.Distribution);
 
 			// now generate the ipa
 			PackageIPA(Path.GetDirectoryName(Params.ProjectGameExeFilename), Path.GetFileNameWithoutExtension(Params.ProjectGameExeFilename), Params.ShortProjectName, Path.GetDirectoryName(Params.RawProjectPath), SC.StageTargetConfigurations[0], Params.Distribution);
@@ -199,7 +199,7 @@ public class IOSPlatform : Platform
 		PrintRunTime();
 	}
 
-	private string EnsureXcodeProjectExists(string RawProjectPath, out bool bWasGenerated)
+	private string EnsureXcodeProjectExists(string RawProjectPath, string LocalRoot, string ShortProjectName, string ProjectRoot, bool IsCodeBasedProject, out bool bWasGenerated)
 	{
 		// first check for ue4.xcodeproj
 		bWasGenerated = false;
@@ -228,14 +228,60 @@ public class IOSPlatform : Platform
 			}
 		}
 
+		// copy the appropriate plist file over
+		string SourcePListFile = CombinePaths(LocalRoot, "Engine", "Build", "IOS", "UE4Game-Info.plist");
+		if (File.Exists(ProjectRoot + "/Build/IOS/" + ShortProjectName + "-Info.plist"))
+		{
+			SourcePListFile = CombinePaths(ProjectRoot, "Build", "IOS", ShortProjectName + "-Info.plist");
+		}
+
+		//@TODO: This is writing to the engine directory!
+		string SourcePath = CombinePaths((IsCodeBasedProject ? ProjectRoot : LocalRoot + "\\Engine"), "Intermediate", "IOS");
+		string TargetPListFile = Path.Combine(SourcePath, (IsCodeBasedProject ? ShortProjectName : "UE4Game") + "-Info.plist");
+
+		Dictionary<string, string> Replacements = new Dictionary<string, string>();
+		Replacements.Add("${EXECUTABLE_NAME}", (IsCodeBasedProject ? ShortProjectName : "UE4Game"));
+		Replacements.Add("${BUNDLE_IDENTIFIER}", ShortProjectName.Replace("_", ""));
+		CopyFileWithReplacements(SourcePListFile, TargetPListFile, Replacements);
+
+		// Now do the .mobileprovision
+		//@TODO: Remove this mobileprovision copy, and move to a library approach like Xcode/codesign does
+		string SourceProvision = CombinePaths(LocalRoot, "Engine", "Build", "IOS", "UE4Game.mobileprovision");
+		string GameSourceProvision = ProjectRoot + "/Build/IOS/" + ShortProjectName + ".mobileprovision";
+		if (File.Exists(GameSourceProvision))
+		{
+			SourceProvision = GameSourceProvision;
+		}
+		if (File.Exists (SourceProvision))
+		{
+			Directory.CreateDirectory (Environment.GetEnvironmentVariable ("HOME") + "/Library/MobileDevice/Provisioning Profiles/");
+			File.Copy (SourceProvision, Environment.GetEnvironmentVariable ("HOME") + "/Library/MobileDevice/Provisioning Profiles/" + ShortProjectName + ".mobileprovision", true);
+			FileInfo DestFileInfo = new FileInfo (Environment.GetEnvironmentVariable ("HOME") + "/Library/MobileDevice/Provisioning Profiles/" + ShortProjectName + ".mobileprovision");
+			DestFileInfo.Attributes = DestFileInfo.Attributes & ~FileAttributes.ReadOnly;
+		}
+
+		// install the distribution provision
+		SourceProvision = CombinePaths(LocalRoot, "Engine", "Build", "IOS", "UE4Game_Distro.mobileprovision");
+		GameSourceProvision = ProjectRoot + "/Build/IOS/" + ShortProjectName + "_Distro.mobileprovision";
+		if (File.Exists(GameSourceProvision))
+		{
+			SourceProvision = GameSourceProvision;
+		}
+		if (File.Exists (SourceProvision))
+		{
+			File.Copy (SourceProvision, Environment.GetEnvironmentVariable ("HOME") + "/Library/MobileDevice/Provisioning Profiles/" + ShortProjectName + "_Distro.mobileprovision", true);
+			FileInfo DestFileInfo = new FileInfo (Environment.GetEnvironmentVariable ("HOME") + "/Library/MobileDevice/Provisioning Profiles/" + ShortProjectName + "_Distro.mobileprovision");
+			DestFileInfo.Attributes = DestFileInfo.Attributes & ~FileAttributes.ReadOnly;
+		}
+
 		return XcodeProj;
 	}
 
-	private void CodeSign(string BaseDirectory, string GameName, string RawProjectPath, UnrealTargetConfiguration TargetConfig, bool Distribution = false)
+	private void CodeSign(string BaseDirectory, string GameName, string RawProjectPath, UnrealTargetConfiguration TargetConfig, string LocalRoot, string ProjectName, string ProjectDirectory, bool IsCode, bool Distribution = false)
 	{
 		// check for the proper xcodeproject
 		bool bWasGenerated = false;
-		string XcodeProj = EnsureXcodeProjectExists (RawProjectPath, out bWasGenerated);
+		string XcodeProj = EnsureXcodeProjectExists (RawProjectPath, LocalRoot, ProjectName, ProjectDirectory, IsCode, out bWasGenerated);
 		string Arguments = "UBT_NO_POST_DEPLOY=true";
 		Arguments += " /usr/bin/xcrun xcodebuild build -project \"" + XcodeProj + "\"";
 		Arguments += " -scheme '";
@@ -595,7 +641,7 @@ public class IOSPlatform : Platform
 			}
 			string GameApp = AppDirectory + "/" + GameName;
 			bool bWasGenerated = false;
-			string XcodeProj = EnsureXcodeProjectExists (Params.RawProjectPath, out bWasGenerated);
+			string XcodeProj = EnsureXcodeProjectExists (Params.RawProjectPath, CmdEnv.LocalRoot, Params.ShortProjectName, GetDirectoryName(Params.RawProjectPath), Params.IsCodeBasedProject, out bWasGenerated);
 			string Arguments = "UBT_NO_POST_DEPLOY=true /usr/bin/xcrun xcodebuild test -project \"" + XcodeProj + "\"";
 			Arguments += " -scheme '";
 			Arguments += GameName;
