build script, more robustification of mac shares

diff --git a/Engine/Source/Programs/AutomationTool/TempStorage.cs b/Engine/Source/Programs/AutomationTool/TempStorage.cs
--- a/Engine/Source/Programs/AutomationTool/TempStorage.cs
+++ b/Engine/Source/Programs/AutomationTool/TempStorage.cs
@@ -71,6 +71,34 @@ namespace AutomationTool
                 return String.IsNullOrEmpty(Name) ? "" : Name;
             }
         }
+
+        public static void Robust_FileExists_NoExceptions(string Filename, string Message)
+        {
+            Robust_FileExists_NoExceptions(false, Filename, Message);
+        }
+        public static void Robust_FileExists_NoExceptions(bool bQuiet, string Filename, string Message)
+        {
+            if (!FileExists_NoExceptions(bQuiet, Filename))
+            {
+                bool bFound = false;
+                // mac is terrible on shares, this isn't a solution, but a stop gap
+                if (Filename.StartsWith("/Volumes/"))
+                {
+                    int Retry = 60;
+                    while (!bFound && --Retry > 0)
+                    {
+                        CommandUtils.Log(System.Diagnostics.TraceEventType.Warning, "*** Mac temp storage retry {0}", Filename);
+                        System.Threading.Thread.Sleep(10000);
+                        bFound = FileExists_NoExceptions(bQuiet, Filename);
+                    }
+                }
+                if (!bFound)
+                {
+                    throw new AutomationException(Message, Filename);
+                }
+            }
+        }
+
         public class TempStorageManifest
         {
             private static readonly string RootElementName = "tempstorage";
@@ -93,10 +121,8 @@ namespace AutomationTool
                 foreach (string InFilename in InFiles)
                 {
                     var Filename = CombinePaths(InFilename);
-                    if (!FileExists_NoExceptions(true, Filename))
-                    {
-                        throw new AutomationException("Could not add {0} to manifest because it does not exist", Filename);
-                    }
+                    Robust_FileExists_NoExceptions(true, Filename, "Could not add {0} to manifest because it does not exist");
+
                     if (!Filename.StartsWith(BaseFolder, StringComparison.InvariantCultureIgnoreCase))
                     {
                         throw new AutomationException("Could not add {0} to manifest because it does not start with the base folder {1}", Filename, BaseFolder);
@@ -426,25 +452,8 @@ namespace AutomationTool
                     foreach (var ThisFileInfo in Dir.Value)
                     {
                         var NewFile = CombinePaths(NewBaseDir, ThisFileInfo.Name);
-                        if (!FileExists_NoExceptions(false, NewFile))
-                        {
-                            bool bFound = false;
-                            // mac is terrible on shares, this isn't a solution, but a stop gap
-                            if (NewFile.StartsWith("/Volumes/"))
-                            {
-                                int Retry = 60;
-                                while (!bFound && --Retry > 0)
-                                {
-                                    CommandUtils.Log(System.Diagnostics.TraceEventType.Warning, "*** Mac temp storage retry {0}", NewFile);
-                                    System.Threading.Thread.Sleep(10000);
-                                    bFound = FileExists_NoExceptions(false, NewFile);
-                                }
-                            }
-                            if (!bFound)
-                            {
-                                throw new AutomationException("Rebased manifest file does not exist {0}", NewFile);
-                            }
-                        }
+                        Robust_FileExists_NoExceptions(false, NewFile, "Rebased manifest file does not exist {0}");
+
                         FileInfo Info = new FileInfo(NewFile);
 
                         Result.Add(Info.FullName);
@@ -681,10 +690,8 @@ namespace AutomationTool
                 foreach (string InFilename in Files)
                 {
                     var Filename = CombinePaths(InFilename);
-                    if (!FileExists_NoExceptions(false, Filename))
-                    {
-                        throw new AutomationException("Could not add {0} to manifest because it does not exist", Filename);
-                    }
+                    Robust_FileExists_NoExceptions(false, Filename, "Could not add {0} to manifest because it does not exist");
+
                     if (!Filename.StartsWith(BaseFolder, StringComparison.InvariantCultureIgnoreCase))
                     {
                         throw new AutomationException("Could not add {0} to manifest because it does not start with the base folder {1}", Filename, BaseFolder);
@@ -696,10 +703,8 @@ namespace AutomationTool
                         throw new AutomationException("Dest file {0} already exists.", DestFile);
                     }
                     CopyFile(Filename, DestFile, true);
-                    if (!FileExists_NoExceptions(true, DestFile))
-                    {
-                        throw new AutomationException("Could not copy {0} to {1}", Filename, DestFile);
-                    }
+                    Robust_FileExists_NoExceptions(true, DestFile, "Could not copy to {0}");
+
                     DestFiles.Add(DestFile);
                 }
                 var Shared = SaveSharedTempStorageManifest(Env, StorageBlockName, GameFolder, DestFiles);
@@ -757,10 +762,8 @@ namespace AutomationTool
                 throw new AutomationException("Storage Block Does Not Exists! {0}", BlockPath);
             }
             var SharedManifest = SharedTempStorageManifestFilename(Env, StorageBlockName, GameFolder);
-            if (!FileExists_NoExceptions(SharedManifest))
-            {
-                throw new AutomationException("Storage Block Manifest Does Not Exists! {0}", SharedManifest);
-            }
+            Robust_FileExists_NoExceptions(SharedManifest, "Storage Block Manifest Does Not Exists! {0}");
+
             var Shared = new TempStorageManifest();
             Shared.Load(SharedManifest);
 
@@ -770,10 +773,8 @@ namespace AutomationTool
             foreach (string InFilename in SharedFiles)
             {
                 var Filename = CombinePaths(InFilename);
-                if (!FileExists_NoExceptions(true, Filename))
-                {
-                    throw new AutomationException("Could not add {0} to manifest because it does not exist", Filename);
-                }
+                Robust_FileExists_NoExceptions(true, Filename, "Could not add {0} to manifest because it does not exist");
+
                 if (!Filename.StartsWith(BlockPath, StringComparison.InvariantCultureIgnoreCase))
                 {
                     throw new AutomationException("Could not add {0} to manifest because it does not start with the base folder {1}", Filename, BlockPath);
@@ -787,10 +788,8 @@ namespace AutomationTool
                 }
                 CopyFile(Filename, DestFile, true);
 
-                if (!FileExists_NoExceptions(true, DestFile))
-                {
-                    throw new AutomationException("Could not copy {0} to {1}", Filename, DestFile);
-                }
+                Robust_FileExists_NoExceptions(true, DestFile, "Could not copy to {0}");
+
                 if (UnrealBuildTool.Utils.IsRunningOnMono && IsProbablyAMacOrIOSExe(DestFile))
                 {
                     SetExecutableBit(DestFile);
