build script, for app bundles to be copied back for the GUBP

diff --git a/Engine/Source/Programs/AutomationTool/Mac/MacPlatform.Automation.cs b/Engine/Source/Programs/AutomationTool/Mac/MacPlatform.Automation.cs
--- a/Engine/Source/Programs/AutomationTool/Mac/MacPlatform.Automation.cs
+++ b/Engine/Source/Programs/AutomationTool/Mac/MacPlatform.Automation.cs
@@ -143,11 +143,11 @@ public class MacPlatform : Platform
 
 	public override bool IsSupported { get { return true; } }
 
-	public override bool ShouldUseManifestForUBTBuilds()
+    public override bool ShouldUseManifestForUBTBuilds(string AddArgs)
 	{
 		// don't use the manifest to set up build products if we are compiling Mac under Windows and we aren't going to copy anything back to the PC
 		bool bIsBuildingRemotely = UnrealBuildTool.ExternalExecution.GetRuntimePlatform() != UnrealTargetPlatform.Mac;
-		bool bUseManifest = !bIsBuildingRemotely || BuildConfiguration.bCopyAppBundleBackToDevice;
+        bool bUseManifest = !bIsBuildingRemotely || AddArgs.IndexOf("-CopyAppBundleBackToDevice", StringComparison.InvariantCultureIgnoreCase) > 0;
 		return bUseManifest;
 	}
     public override List<string> GetDebugFileExtentions()
diff --git a/Engine/Source/Programs/AutomationTool/Platform.cs b/Engine/Source/Programs/AutomationTool/Platform.cs
--- a/Engine/Source/Programs/AutomationTool/Platform.cs
+++ b/Engine/Source/Programs/AutomationTool/Platform.cs
@@ -355,7 +355,7 @@ namespace AutomationTool
 
 		}
 
-		public virtual bool ShouldUseManifestForUBTBuilds()
+		public virtual bool ShouldUseManifestForUBTBuilds(string AddArgs)
 		{
 			return true;
 		}
diff --git a/Engine/Source/Programs/AutomationTool/Scripts/GUBP.Automation.cs b/Engine/Source/Programs/AutomationTool/Scripts/GUBP.Automation.cs
--- a/Engine/Source/Programs/AutomationTool/Scripts/GUBP.Automation.cs
+++ b/Engine/Source/Programs/AutomationTool/Scripts/GUBP.Automation.cs
@@ -559,7 +559,9 @@ public class GUBP : BuildCommand
 			    }
                     );
             }
-            Agenda.AddTargets(new string[] { "UnrealHeaderTool" }, HostPlatform, UnrealTargetConfiguration.Development);
+            string AddArgs = "-CopyAppBundleBackToDevice" + bp.RocketUBTArgs();
+
+            Agenda.AddTargets(new string[] { "UnrealHeaderTool" }, HostPlatform, UnrealTargetConfiguration.Development, InAddArgs: AddArgs);
             return Agenda;
         }
         public override void PostBuild(GUBP bp, UE4Build UE4Build)
@@ -747,7 +749,7 @@ public class GUBP : BuildCommand
                     );
             }
 
-            string AddArgs = "-nobuilduht -skipactionhistory" + bp.RocketUBTArgs(); ;
+            string AddArgs = "-nobuilduht -skipactionhistory -CopyAppBundleBackToDevice" + bp.RocketUBTArgs(); ;
 
             foreach (var ProgramTarget in bp.Branch.BaseEngineProject.Properties.Programs)
             {
@@ -802,7 +804,7 @@ public class GUBP : BuildCommand
         {
             var Agenda = new UE4Build.BuildAgenda();
 
-            string AddArgs = "-nobuilduht -skipactionhistory" + bp.RocketUBTArgs(); ;
+            string AddArgs = "-nobuilduht -skipactionhistory -CopyAppBundleBackToDevice" + bp.RocketUBTArgs(); ;
 
             foreach (var Plat in ProgramTarget.Rules.GUBP_ToolPlatforms(HostPlatform))
             {
@@ -879,7 +881,7 @@ public class GUBP : BuildCommand
 			        }
                     );
             }
-            string AddArgs = "-nobuilduht -skipactionhistory" + bp.RocketUBTArgs(); ;
+            string AddArgs = "-nobuilduht -skipactionhistory -CopyAppBundleBackToDevice" + bp.RocketUBTArgs(); ;
 
             foreach (var ProgramTarget in bp.Branch.BaseEngineProject.Properties.Programs)
             {
@@ -944,7 +946,7 @@ public class GUBP : BuildCommand
         {
             var Agenda = new UE4Build.BuildAgenda();
 
-            string AddArgs = "-nobuilduht -skipactionhistory" + bp.RocketUBTArgs(); ;
+            string AddArgs = "-nobuilduht -skipactionhistory -CopyAppBundleBackToDevice" + bp.RocketUBTArgs(); ;
 
             foreach (var Plat in ProgramTarget.Rules.GUBP_ToolPlatforms(HostPlatform))
             {
@@ -993,12 +995,12 @@ public class GUBP : BuildCommand
 
             Agenda.AddTargets(
                 new string[] { bp.Branch.BaseEngineProject.Properties.Targets[TargetRules.TargetType.Editor].TargetName },
-                HostPlatform, UnrealTargetConfiguration.Development, InAddArgs: "-nobuilduht -skipactionhistory -onlyplatformspecificfor=" + EditorPlatform.ToString());
+                HostPlatform, UnrealTargetConfiguration.Development, InAddArgs: "-nobuilduht -skipactionhistory -CopyAppBundleBackToDevice -onlyplatformspecificfor=" + EditorPlatform.ToString());
             foreach (var ProgramTarget in bp.Branch.BaseEngineProject.Properties.Programs)
             {
                 if (ProgramTarget.Rules.GUBP_AlwaysBuildWithBaseEditor() && ProgramTarget.Rules.SupportsPlatform(HostPlatform) && ProgramTarget.Rules.GUBP_NeedsPlatformSpecificDLLs())
                 {
-                    Agenda.AddTargets(new string[] { ProgramTarget.TargetName }, HostPlatform, UnrealTargetConfiguration.Development, InAddArgs: "-nobuilduht -skipactionhistory -onlyplatformspecificfor=" + EditorPlatform.ToString());
+                    Agenda.AddTargets(new string[] { ProgramTarget.TargetName }, HostPlatform, UnrealTargetConfiguration.Development, InAddArgs: "-nobuilduht -skipactionhistory -CopyAppBundleBackToDevice -onlyplatformspecificfor=" + EditorPlatform.ToString());
                 }
             }
             return Agenda;
@@ -1052,7 +1054,7 @@ public class GUBP : BuildCommand
         {
             var Agenda = new UE4Build.BuildAgenda();
 
-            string Args = "-nobuilduht -skipactionhistory -skipnonhostplatforms -forceheadergeneration" + bp.RocketUBTArgs(true);
+            string Args = "-nobuilduht -skipactionhistory -skipnonhostplatforms -CopyAppBundleBackToDevice -forceheadergeneration" + bp.RocketUBTArgs(true);
 
             Agenda.AddTargets(
                 new string[] { GameProj.Properties.Targets[TargetRules.TargetType.Editor].TargetName },
@@ -1138,7 +1140,7 @@ public class GUBP : BuildCommand
             }
             var Agenda = new UE4Build.BuildAgenda();
 
-            string Args = "-nobuilduht -skipactionhistory" + bp.RocketUBTArgs();
+            string Args = "-nobuilduht -skipactionhistory -CopyAppBundleBackToDevice" + bp.RocketUBTArgs();
 
             foreach (var Kind in BranchInfo.MonolithicKinds)
             {
diff --git a/Engine/Source/Programs/AutomationTool/UE4Build.cs b/Engine/Source/Programs/AutomationTool/UE4Build.cs
--- a/Engine/Source/Programs/AutomationTool/UE4Build.cs
+++ b/Engine/Source/Programs/AutomationTool/UE4Build.cs
@@ -304,7 +304,7 @@ namespace AutomationTool
 			PrepareUBT();
 
 			// let the platform determine when to use the manifest
-			bool UseManifest = Platform.Platforms[TargetPlatform].ShouldUseManifestForUBTBuilds();
+            bool UseManifest = Platform.Platforms[TargetPlatform].ShouldUseManifestForUBTBuilds(AddArgs);
 
 			if (UseManifest)
 			{
diff --git a/Engine/Source/Programs/UnrealBuildTool/System/UnrealBuildTool.cs b/Engine/Source/Programs/UnrealBuildTool/System/UnrealBuildTool.cs
--- a/Engine/Source/Programs/UnrealBuildTool/System/UnrealBuildTool.cs
+++ b/Engine/Source/Programs/UnrealBuildTool/System/UnrealBuildTool.cs
@@ -1303,6 +1303,11 @@ namespace UnrealBuildTool
 				BuildConfiguration.bDeployAfterCompile = true;
 			}
 
+            if (Utils.ParseCommandLineFlag(Arguments, "-CopyAppBundleBackToDevice", out ArgumentIndex))
+            {
+                BuildConfiguration.bCopyAppBundleBackToDevice = true;
+            }
+
 			if (Utils.ParseCommandLineFlag(Arguments, "-flushmac", out ArgumentIndex))
 			{
 				BuildConfiguration.bFlushBuildDirOnRemoteMac = true;
