build script, fix P4 changes cache, add timings, make the temp storage search faster

diff --git a/Engine/Source/Programs/AutomationTool/Scripts/GUBP.Automation.cs b/Engine/Source/Programs/AutomationTool/Scripts/GUBP.Automation.cs
--- a/Engine/Source/Programs/AutomationTool/Scripts/GUBP.Automation.cs
+++ b/Engine/Source/Programs/AutomationTool/Scripts/GUBP.Automation.cs
@@ -110,6 +110,8 @@ public class GUBP : BuildCommand
     private static List<GUBPEmailHacker> EmailHackers;
     private string HackEmails(string Emails, string Branch, string NodeName)
     {
+        var StartTime = DateTime.UtcNow;
+
         if (EmailHackers == null)
         {
             EmailHackers = new List<GUBPEmailHacker>();
@@ -154,10 +156,15 @@ public class GUBP : BuildCommand
             Result = EmailHacker.FinalizeEmails(Result, this, Branch, NodeName);
         }
         string FinalEmails = "";
+        int Count = 0;
         foreach (var Email in Result)
         {
             FinalEmails = GUBPNode.MergeSpaceStrings(FinalEmails, Email);
+            Count++;
         }
+        var BuildDuration = (DateTime.UtcNow - StartTime).TotalMilliseconds;
+        Log("Took {0}s to hack {1} emails", BuildDuration / 1000, Count);
+
         return FinalEmails;
     }
     public abstract class GUBPNode
@@ -2916,6 +2923,8 @@ public class GUBP : BuildCommand
         {
             if (LastOutputForChanges > 1990000)
             {
+                var StartTime = DateTime.UtcNow;
+
                 string Cmd = String.Format("{0}@{1},{2} {3}@{4},{5}",
                     CombinePaths(PathSeparator.Slash, P4Env.BuildRootP4, "*", "Source", "..."), LastOutputForChanges + 1, TopCL,
                     CombinePaths(PathSeparator.Slash, P4Env.BuildRootP4, "*", "Build", "..."), LastOutputForChanges + 1, TopCL
@@ -2935,6 +2944,10 @@ public class GUBP : BuildCommand
                 {
                     throw new AutomationException("Could not get changes; cmdline: p4 changes {0}", Cmd);
                 }
+
+                var BuildDuration = (DateTime.UtcNow - StartTime).TotalMilliseconds;
+                Log("Took {0}s to get P4 history for {1}-{2}", BuildDuration / 1000, LastOutputForChanges + 1, TopCL);
+
             }
             else
             {
@@ -3241,10 +3254,6 @@ public class GUBP : BuildCommand
         var Result = new List<int>();
         foreach (var ThisString in Strings)
         {
-            if (ThisString.Contains("-"))
-            {
-                continue;
-            }
             int ThisInt = int.Parse(ThisString);
             if (ThisInt < 1960000 || ThisInt > 3000000)
             {
@@ -3285,6 +3294,8 @@ public class GUBP : BuildCommand
 
     List<string> TopologicalSort(HashSet<string> NodesToDo, string ExplicitTrigger, bool LocalOnly, bool SubSort = false)
     {
+        var StartTime = DateTime.UtcNow;
+
         var OrdereredToDo = new List<string>();
 
         var SortedAgentGroupChains = new Dictionary<string, List<string>>();
@@ -3481,6 +3492,12 @@ public class GUBP : BuildCommand
                 throw new AutomationException("Cycle in GUBP");
             }
         }
+        if (!SubSort)
+        {
+            var BuildDuration = (DateTime.UtcNow - StartTime).TotalMilliseconds;
+            Log("Took {0}s to sort {1} nodes", BuildDuration / 1000, OrdereredToDo.Count);
+        }
+
         return OrdereredToDo;
     }
 
@@ -4834,15 +4851,21 @@ if (HostPlatform == UnrealTargetPlatform.Mac) continue; //temp hack till mac aut
         GUBPNodesHistory = new Dictionary<string, NodeHistory>();
 
         Log("******* Caching completion");
-        foreach (var Node in NodesToDo)
         {
-            Log("** {0}", Node);
-            NodeIsAlreadyComplete(Node, LocalOnly); // cache these now to avoid spam later
-            GetControllingTriggerDotName(Node);
+            var StartTime = DateTime.UtcNow;
+            foreach (var Node in NodesToDo)
+            {
+                Log("** {0}", Node);
+                NodeIsAlreadyComplete(Node, LocalOnly); // cache these now to avoid spam later
+                GetControllingTriggerDotName(Node);
+            }
+            var BuildDuration = (DateTime.UtcNow - StartTime).TotalMilliseconds;
+            Log("Took {0}s to cache completion for {1} nodes", BuildDuration / 1000, NodesToDo.Count);
         }
         if (CLString != "" && StoreName.Contains(CLString) && !ParseParam("NoHistory"))
         {
             Log("******* Updating history");
+            var StartTime = DateTime.UtcNow;
             foreach (var Node in NodesToDo)
             {
                 if (!NodeIsAlreadyComplete(Node, LocalOnly))
@@ -4850,6 +4873,8 @@ if (HostPlatform == UnrealTargetPlatform.Mac) continue; //temp hack till mac aut
                     UpdateNodeHistory(Node, CLString);
                 }
             }
+            var BuildDuration = (DateTime.UtcNow - StartTime).TotalMilliseconds;
+            Log("Took {0}s to get history for {1} nodes", BuildDuration / 1000, NodesToDo.Count);
         }
 
         var OrdereredToDo = TopologicalSort(NodesToDo, ExplicitTrigger, LocalOnly);
diff --git a/Engine/Source/Programs/AutomationTool/TempStorage.cs b/Engine/Source/Programs/AutomationTool/TempStorage.cs
--- a/Engine/Source/Programs/AutomationTool/TempStorage.cs
+++ b/Engine/Source/Programs/AutomationTool/TempStorage.cs
@@ -888,6 +888,11 @@ namespace AutomationTool
                             {
                                 throw new AutomationException("Dir {0} didn't have any string to fit the star in the wildcard {1}", ThisSubDir, WildCard);
                             }
+                            // these are a bunch of false positives
+                            if (StarReplacement.Contains("-"))
+                            {
+                                continue;
+                            }
                             if (!Files.Contains(StarReplacement))
                             {
                                 Files.Add(StarReplacement);
