build script, rename the tools agent sharing group to avoid a clash with a node name

diff --git a/Engine/Source/Programs/AutomationTool/Scripts/GUBP.Automation.cs b/Engine/Source/Programs/AutomationTool/Scripts/GUBP.Automation.cs
--- a/Engine/Source/Programs/AutomationTool/Scripts/GUBP.Automation.cs
+++ b/Engine/Source/Programs/AutomationTool/Scripts/GUBP.Automation.cs
@@ -702,7 +702,7 @@ public class GUBP : BuildCommand
             {
                 AddPseudodependency(RootEditorNode.StaticGetFullName(HostPlatform));
             }
-            AgentSharingGroup = "Tools" + StaticGetHostPlatformSuffix(HostPlatform);
+            AgentSharingGroup = "ToolsGroup" + StaticGetHostPlatformSuffix(HostPlatform);
         }
         public static string StaticGetFullName(UnrealTargetPlatform InHostPlatform)
         {
@@ -789,7 +789,7 @@ public class GUBP : BuildCommand
             {
                 AddPseudodependency(RootEditorNode.StaticGetFullName(HostPlatform));
             }
-            AgentSharingGroup = "Tools" + StaticGetHostPlatformSuffix(HostPlatform);
+            AgentSharingGroup = "ToolsGroup" + StaticGetHostPlatformSuffix(HostPlatform);
         }
         public static string StaticGetFullName(UnrealTargetPlatform InHostPlatform, SingleTargetProperties InProgramTarget)
         {
@@ -833,7 +833,7 @@ public class GUBP : BuildCommand
             {
                 AddPseudodependency(RootEditorNode.StaticGetFullName(HostPlatform));
             }
-            AgentSharingGroup = "Tools" + StaticGetHostPlatformSuffix(HostPlatform);
+            AgentSharingGroup = "ToolsGroup" + StaticGetHostPlatformSuffix(HostPlatform);
         }
         public static string StaticGetFullName(UnrealTargetPlatform InHostPlatform)
         {
@@ -926,7 +926,7 @@ public class GUBP : BuildCommand
             {
                 AddPseudodependency(RootEditorNode.StaticGetFullName(HostPlatform));
             }
-            AgentSharingGroup = "Tools" + StaticGetHostPlatformSuffix(HostPlatform);
+            AgentSharingGroup = "ToolsGroup" + StaticGetHostPlatformSuffix(HostPlatform);
         }
         public static string StaticGetFullName(UnrealTargetPlatform InHostPlatform, SingleTargetProperties InProgramTarget)
         {
