Fixing UAT crash when trying to get command help.

diff --git a/Engine/Source/Programs/AutomationTool/HelpUtils.cs b/Engine/Source/Programs/AutomationTool/HelpUtils.cs
--- a/Engine/Source/Programs/AutomationTool/HelpUtils.cs
+++ b/Engine/Source/Programs/AutomationTool/HelpUtils.cs
@@ -165,7 +165,7 @@ namespace AutomationTool
 		/// <returns></returns>
 		static string FormatParams(List<string> Params, int Indent, int DefaultRightPadding)
 		{
-			Dictionary<string, string> ParamDict = new Dictionary<string, string>();
+			Dictionary<string, string> ParamDict = new Dictionary<string, string>(StringComparer.InvariantCultureIgnoreCase);
 
 			// Extract Params/Descriptions into Key/Value pairs
 			foreach (var param in Params)
@@ -178,7 +178,14 @@ namespace AutomationTool
 				string ParamDesc = param.Substring(SplitPoint + 1, param.Length - (SplitPoint + 1));
 
 				// build dictionary using Name and Desc as Key and Value
-				ParamDict.Add(ParamName, ParamDesc);
+				if (!ParamDict.ContainsKey(ParamName))
+				{
+					ParamDict.Add(ParamName, ParamDesc);
+				}
+				else
+				{
+					LogWarning("Duplicated help parameter \"{0}\"", ParamName);
+				}
 			}
 
 			// string used to intent the param
