build script, fix typo in nobuilduht

diff --git a/Engine/Source/Programs/AutomationTool/Scripts/GUBP.Automation.cs b/Engine/Source/Programs/AutomationTool/Scripts/GUBP.Automation.cs
--- a/Engine/Source/Programs/AutomationTool/Scripts/GUBP.Automation.cs
+++ b/Engine/Source/Programs/AutomationTool/Scripts/GUBP.Automation.cs
@@ -604,7 +604,7 @@ public class GUBP : BuildCommand
         {
             var Agenda = new UE4Build.BuildAgenda();
 
-            string AddArgs = "-nobuiltuht" + bp.RocketUBTArgs();
+            string AddArgs = "-nobuilduht" + bp.RocketUBTArgs();
             if (bp.bOrthogonalizeEditorPlatforms)
             {
                 AddArgs += " -skipnonhostplatforms";
