build script, rearrange timing logs to keep them out of the P4 history proper

diff --git a/Engine/Source/Programs/AutomationTool/Scripts/GUBP.Automation.cs b/Engine/Source/Programs/AutomationTool/Scripts/GUBP.Automation.cs
--- a/Engine/Source/Programs/AutomationTool/Scripts/GUBP.Automation.cs
+++ b/Engine/Source/Programs/AutomationTool/Scripts/GUBP.Automation.cs
@@ -2923,8 +2923,6 @@ public class GUBP : BuildCommand
         {
             if (LastOutputForChanges > 1990000)
             {
-                var StartTime = DateTime.UtcNow;
-
                 string Cmd = String.Format("{0}@{1},{2} {3}@{4},{5}",
                     CombinePaths(PathSeparator.Slash, P4Env.BuildRootP4, "*", "Source", "..."), LastOutputForChanges + 1, TopCL,
                     CombinePaths(PathSeparator.Slash, P4Env.BuildRootP4, "*", "Build", "..."), LastOutputForChanges + 1, TopCL
@@ -2944,10 +2942,6 @@ public class GUBP : BuildCommand
                 {
                     throw new AutomationException("Could not get changes; cmdline: p4 changes {0}", Cmd);
                 }
-
-                var BuildDuration = (DateTime.UtcNow - StartTime).TotalMilliseconds;
-                Log("Took {0}s to get P4 history for {1}-{2}", BuildDuration / 1000, LastOutputForChanges + 1, TopCL);
-
             }
             else
             {
@@ -2974,6 +2968,8 @@ public class GUBP : BuildCommand
 
     void PrintDetailedChanges(NodeHistory History, bool bShowAllChanges = false)
     {
+        var StartTime = DateTime.UtcNow;
+
         string Me = String.Format("{0}   <<<< local sync", P4Env.Changelist);
         int LastOutputForChanges = 0;
         int LastGreen = History.LastSucceeded;
@@ -3013,6 +3009,9 @@ public class GUBP : BuildCommand
             LastOutputForChanges = PrintChanges(LastOutputForChanges, P4Env.Changelist, LastGreen);
             Log("         {0}", Me);
         }
+        var BuildDuration = (DateTime.UtcNow - StartTime).TotalMilliseconds;
+        Log("Took {0}s to get P4 history", BuildDuration / 1000);
+
     }
     void PrintNodes(GUBP bp, List<string> Nodes, bool LocalOnly, List<string> UnfinishedTriggers = null)
     {
@@ -3561,6 +3560,8 @@ public class GUBP : BuildCommand
 
     List<string> GetECPropsForNode(string NodeToDo, out string EMails, bool OnlyLateUpdates = false)
     {
+        var StartTime = DateTime.UtcNow;
+
         var ECProps = new List<string>();
         EMails = "";
         string FailCauserEMails = "";
@@ -3632,6 +3633,8 @@ public class GUBP : BuildCommand
             ECProps.Add(string.Format("Timeouts/{0}={1}", NodeToDo, GUBPNodes[NodeToDo].TimeoutInMinutes()));
             ECProps.Add(string.Format("JobStepPath/{0}={1}", NodeToDo, GetJobStepPath(NodeToDo)));
         }
+        var BuildDuration = (DateTime.UtcNow - StartTime).TotalMilliseconds;
+        Log("Took {0}s to get P4 history for node {1}", BuildDuration / 1000, NodeToDo);
 
         return ECProps;
     }
