build script, 48 retries on P4 submits

diff --git a/Engine/Source/Programs/AutomationTool/P4Utils.cs b/Engine/Source/Programs/AutomationTool/P4Utils.cs
--- a/Engine/Source/Programs/AutomationTool/P4Utils.cs
+++ b/Engine/Source/Programs/AutomationTool/P4Utils.cs
@@ -1138,7 +1138,7 @@ namespace AutomationTool
 			SubmittedCL = 0;
 			int Retry = 0;
 			string LastCmdOutput = "none?";
-			while (Retry++ < 24)
+			while (Retry++ < 48)
 			{
 				bool Pending;
 				if (!ChangeExists(CL, out Pending))
@@ -1277,7 +1277,7 @@ namespace AutomationTool
 				RevertAll(CL);
 				CommandUtils.Log(TraceEventType.Error, "Submit CL {0} failed, reverting files\n", CL);
 			}
-			throw new P4Exception("Change {0} failed to submit after 12 retries??.\n{1}", CL, LastCmdOutput);
+			throw new P4Exception("Change {0} failed to submit after 48 retries??.\n{1}", CL, LastCmdOutput);
 		}
 
 		/// <summary>
