UAT: Quick fix for module sorting sometimes throwing index out of range exception

diff --git a/Engine/Source/Programs/AutomationTool/ScriptCompiler.cs b/Engine/Source/Programs/AutomationTool/ScriptCompiler.cs
--- a/Engine/Source/Programs/AutomationTool/ScriptCompiler.cs
+++ b/Engine/Source/Programs/AutomationTool/ScriptCompiler.cs
@@ -156,30 +156,39 @@ namespace AutomationTool
 			Log.TraceInformation("Building script modules");
 			// Make sure DefaultScriptsDLLName is compiled first
 			var DefaultScriptsProjName = Path.ChangeExtension(DefaultScriptsDLLName, "csproj");
-			Modules.Sort((string A, string B) =>
+			foreach (var ModuleName in Modules)
 			{
-				if (A.IndexOf(DefaultScriptsProjName, StringComparison.InvariantCultureIgnoreCase) >= 0)
+				if (ModuleName.IndexOf(DefaultScriptsProjName, StringComparison.InvariantCultureIgnoreCase) >= 0)
 				{
-					return -1;
-				}
-				else if (B.IndexOf(DefaultScriptsProjName, StringComparison.InvariantCultureIgnoreCase) >= 0)
-				{
-					return 1;
+					Log.TraceInformation("Building script module: {0}", ModuleName);
+					try
+					{
+						CompileScriptModule(ModuleName);
+					}
+					catch (Exception Ex)
+					{
+						CommandUtils.Log(TraceEventType.Error, Ex);
+						throw new AutomationException("Failed to compile module {0}", ModuleName);
+					}
+					break;
 				}
-				return String.Compare(A, B, true); 
-			});
+			}
 
-			foreach (var ModuleFile in Modules)
+			// Second pass, compile everything else
+			foreach (var ModuleName in Modules)
 			{
-				Log.TraceInformation("Building script module: {0}", ModuleFile);
-				try
+				if (ModuleName.IndexOf(DefaultScriptsProjName, StringComparison.InvariantCultureIgnoreCase) < 0)
 				{
-					CompileScriptModule(ModuleFile);
-				}
-				catch (Exception Ex)
-				{
-					CommandUtils.Log(TraceEventType.Error, Ex);
-					throw new AutomationException("Failed to compile module {0}", ModuleFile);
+					Log.TraceInformation("Building script module: {0}", ModuleName);
+					try
+					{
+						CompileScriptModule(ModuleName);
+					}
+					catch (Exception Ex)
+					{
+						CommandUtils.Log(TraceEventType.Error, Ex);
+						throw new AutomationException("Failed to compile module {0}", ModuleName);
+					}
 				}
 			}
 		}
