build script, spew reduction

diff --git a/Engine/Source/Programs/AutomationTool/CommandUtils.cs b/Engine/Source/Programs/AutomationTool/CommandUtils.cs
--- a/Engine/Source/Programs/AutomationTool/CommandUtils.cs
+++ b/Engine/Source/Programs/AutomationTool/CommandUtils.cs
@@ -569,6 +569,23 @@ namespace AutomationTool
 			}
 		}
 
+        /// <summary>
+        /// Creates a directory(or directories).
+        /// If the creation of the directory fails, this function throws an Exception.
+        /// </summary>
+        /// <param name="Directory">Directory</param>
+        public static void CreateDirectory(bool bQuiet, params string[] Directories)
+        {
+            foreach (var DirectoryName in Directories)
+            {
+                var NormalizedDirectory = ConvertSeparators(PathSeparator.Default, DirectoryName);
+                if (!InternalUtils.SafeCreateDirectory(NormalizedDirectory, bQuiet))
+                {
+                    throw new AutomationException(String.Format("Failed to create directory '{0}'", NormalizedDirectory));
+                }
+            }
+        }
+
 		/// <summary>
 		/// Creates a directory (or directories).
 		/// If the creation of the directory fails, this function prints a warning.
@@ -595,11 +612,11 @@ namespace AutomationTool
 		/// </summary>
 		/// <param name="OldName">Old name</param>
 		/// <param name="NewName">new name</param>
-		public static void RenameFile(string OldName, string NewName)
+        public static void RenameFile(string OldName, string NewName, bool bQuiet = false)
 		{
 			var OldNormalized = ConvertSeparators(PathSeparator.Default, OldName);
 			var NewNormalized = ConvertSeparators(PathSeparator.Default, NewName);
-			if (!InternalUtils.SafeRenameFile(OldNormalized, NewNormalized))
+			if (!InternalUtils.SafeRenameFile(OldNormalized, NewNormalized, bQuiet))
 			{
 				throw new AutomationException(String.Format("Failed to rename/move file '{0}' to '{1}'", OldNormalized, NewNormalized));
 			}
diff --git a/Engine/Source/Programs/AutomationTool/Scripts/GUBP.Automation.cs b/Engine/Source/Programs/AutomationTool/Scripts/GUBP.Automation.cs
--- a/Engine/Source/Programs/AutomationTool/Scripts/GUBP.Automation.cs
+++ b/Engine/Source/Programs/AutomationTool/Scripts/GUBP.Automation.cs
@@ -2878,11 +2878,11 @@ public class GUBP : BuildCommand
         bool Result;
         if (LocalOnly)
         {
-            Result = LocalTempStorageExists(CmdEnv, NodeStoreName);
+            Result = LocalTempStorageExists(CmdEnv, NodeStoreName, bQuiet : true);
         }
         else
         {
-            Result = TempStorageExists(CmdEnv, NodeStoreName, GameNameIfAny);
+            Result = TempStorageExists(CmdEnv, NodeStoreName, GameNameIfAny, bQuiet: true);
         }
         if (Result)
         {
@@ -2895,7 +2895,7 @@ public class GUBP : BuildCommand
         GUBPNodesCompleted.Add(NodeToDo, Result);
         return Result;
     }
-    void RunECTool(string Args)
+    void RunECTool(string Args, bool bQuiet = false)
     {
         if (ParseParam("FakeEC"))
         {
@@ -2903,7 +2903,12 @@ public class GUBP : BuildCommand
         }
         else
         {
-            RunAndLog("ectool", Args);
+            ERunOptions Opts = ERunOptions.Default;
+            if (bQuiet)
+            {
+                Opts = Opts & ~ERunOptions.AllowSpew;
+            }
+            RunAndLog("ectool", Args, Options: Opts);
         }
     }
     string GetEMailListForNode(GUBP bp, string NodeToDo, string Causers)
@@ -3650,7 +3655,7 @@ public class GUBP : BuildCommand
             foreach (var Prop in Props)
             {
                 var Parts = Prop.Split("=".ToCharArray());
-                RunECTool(String.Format("setProperty \"/myWorkflow/{0}\" \"{1}\"", Parts[0], Parts[1]));
+                RunECTool(String.Format("setProperty \"/myWorkflow/{0}\" \"{1}\"", Parts[0], Parts[1]), true);
             }
         }
         catch (Exception Ex)
diff --git a/Engine/Source/Programs/AutomationTool/TempStorage.cs b/Engine/Source/Programs/AutomationTool/TempStorage.cs
--- a/Engine/Source/Programs/AutomationTool/TempStorage.cs
+++ b/Engine/Source/Programs/AutomationTool/TempStorage.cs
@@ -143,17 +143,14 @@ namespace AutomationTool
                     ManifestDirectory.Add(new TempStorageFileInfo(Filename, RelativeFile));
                 }
 
-                Log("Created manifest:");
-                Stats();
+                Stats("Created manifest");
 
             }
 
-            public void Stats()
+            public void Stats(string Description)
             {
 
-                Log("Directories={0}", Directories.Count);
-                Log("Files={0}", GetFileCount());
-                Log("Size={0}", GetTotalSize());
+                Log("{0}: Directories={1} Files={2} Size={3}", Description, Directories.Count, GetFileCount(), GetTotalSize());
             }
             
             public bool Compare(TempStorageManifest Other)
@@ -264,7 +261,7 @@ namespace AutomationTool
             }
 
 
-            public void Load(string Filename)
+            public void Load(string Filename, bool bQuiet = false)
             {
                 XmlDocument File = new XmlDocument();
                 try
@@ -296,9 +293,10 @@ namespace AutomationTool
                     throw new AutomationException("Bad root node ({0}).", RootNode != null ? RootNode.Name : "null");
                 }
 
-                Log("Loaded manifest {0}:", Filename);
-
-                Stats();
+                if (!bQuiet)
+                {
+                    Stats(String.Format("Loaded manifest {0}", Filename));
+                }
                 if (GetFileCount() <= 0 || GetTotalSize() <= 0)
                 {
                     throw new AutomationException("Attempt to load empty manifest.");
@@ -587,24 +585,24 @@ namespace AutomationTool
             {
                 throw new AutomationException("Temp manifest file already exists {0}", TempFilename);
             }
-            CreateDirectory(Path.GetDirectoryName(FinalFilename));
+            CreateDirectory(true, Path.GetDirectoryName(FinalFilename));
             Saver.Save(TempFilename);
 
             var Tester = new TempStorageManifest();
-            Tester.Load(TempFilename);
+            Tester.Load(TempFilename, true);
 
             if (!Saver.Compare(Tester))
             {
                 throw new AutomationException("Temp manifest differs {0}", TempFilename);
             }
 
-            RenameFile(TempFilename, FinalFilename);
-            if (FileExists_NoExceptions(TempFilename))
+            RenameFile(TempFilename, FinalFilename, true);
+            if (FileExists_NoExceptions(true, TempFilename))
             {
                 throw new AutomationException("Temp manifest didn't go away {0}", TempFilename);
             }
             var FinalTester = new TempStorageManifest();
-            FinalTester.Load(FinalFilename);
+            FinalTester.Load(FinalFilename, true);
 
             if (!Saver.Compare(FinalTester))
             {
@@ -637,27 +635,27 @@ namespace AutomationTool
         {
             DeleteDirectory(true, SharedTempStorageDirectory(StorageBlockName, GameFolder));
         }
-        public static bool LocalTempStorageExists(CommandEnvironment Env, string StorageBlockName)
+        public static bool LocalTempStorageExists(CommandEnvironment Env, string StorageBlockName, bool bQuiet = false)
         {
             var LocalManifest = LocalTempStorageManifestFilename(Env, StorageBlockName);
-            if (FileExists_NoExceptions(LocalManifest))
+            if (FileExists_NoExceptions(bQuiet, LocalManifest))
             {
                 return true;
             }
             return false;
         }
-        public static bool SharedTempStorageExists(CommandEnvironment Env, string StorageBlockName, string GameFolder = "")
+        public static bool SharedTempStorageExists(CommandEnvironment Env, string StorageBlockName, string GameFolder = "", bool bQuiet = false)
         {
             var SharedManifest = SharedTempStorageManifestFilename(Env, StorageBlockName, GameFolder, false);
-            if (FileExists_NoExceptions(SharedManifest))
+            if (FileExists_NoExceptions(bQuiet, SharedManifest))
             {
                 return true;
             }
             return false;
         }
-        public static bool TempStorageExists(CommandEnvironment Env, string StorageBlockName, string GameFolder = "", bool bLocalOnly = false)
+        public static bool TempStorageExists(CommandEnvironment Env, string StorageBlockName, string GameFolder = "", bool bLocalOnly = false, bool bQuiet = false)
         {
-            return LocalTempStorageExists(Env, StorageBlockName) || (!bLocalOnly && SharedTempStorageExists(Env, StorageBlockName, GameFolder));
+            return LocalTempStorageExists(Env, StorageBlockName, bQuiet) || (!bLocalOnly && SharedTempStorageExists(Env, StorageBlockName, GameFolder, bQuiet));
         }
         public static void StoreToTempStorage(CommandEnvironment Env, string StorageBlockName, List<string> Files, bool bLocalOnly = false, string GameFolder = "", string BaseFolder = "")
         {
@@ -681,7 +679,7 @@ namespace AutomationTool
                 {
                     throw new AutomationException("Storage Block Already Exists! {0}", BlockPath);
                 }
-                CreateDirectory(BlockPath);
+                CreateDirectory(true, BlockPath);
                 if (!DirectoryExists_NoExceptions(BlockPath))
                 {
                     throw new AutomationException("Storage Block Could Not Be Created! {0}", BlockPath);
@@ -907,7 +905,7 @@ namespace AutomationTool
             {
                 var NewBlock = StorageBlockName.Replace("*", StarReplacement);
 
-                if (TempStorageExists(Env, NewBlock, GameFolder, LocalOnly))
+                if (TempStorageExists(Env, NewBlock, GameFolder, LocalOnly, true))
                 {
                     OutFiles.Add(StarReplacement);
                 }
