Update NoRedist refernces for CrashReportServer to NotForLicensees.

diff --git a/Engine/Source/Programs/AutomationTool/Scripts/GUBP.Automation.cs b/Engine/Source/Programs/AutomationTool/Scripts/GUBP.Automation.cs
--- a/Engine/Source/Programs/AutomationTool/Scripts/GUBP.Automation.cs
+++ b/Engine/Source/Programs/AutomationTool/Scripts/GUBP.Automation.cs
@@ -868,9 +868,9 @@ public class GUBP : BuildCommand
                     new string[] 
 			    {
 				    CombinePaths(@"Engine\Source\Programs\Distill\Distill.csproj"),		  
-                    CombinePaths(@"Engine\Source\Programs\NoRedist\CrashReportServer\CrashReportCommon\CrashReportCommon.csproj"),
-					CombinePaths(@"Engine\Source\Programs\NoRedist\CrashReportServer\CrashReportReceiver\CrashReportReceiver.csproj"),
-					CombinePaths(@"Engine\Source\Programs\NoRedist\CrashReportServer\CrashReportProcess\CrashReportProcess.csproj"),
+                    CombinePaths(@"Engine\Source\Programs\NotForLicensees\CrashReportServer\CrashReportCommon\CrashReportCommon.csproj"),
+					CombinePaths(@"Engine\Source\Programs\NotForLicensees\CrashReportServer\CrashReportReceiver\CrashReportReceiver.csproj"),
+					CombinePaths(@"Engine\Source\Programs\NotForLicensees\CrashReportServer\CrashReportProcess\CrashReportProcess.csproj"),
                     CombinePaths(@"Engine\Source\Programs\CrashReporter\RegisterPII\RegisterPII.csproj"),
 			    });
                 Agenda.DotNetSolutions.AddRange(
