build script, performance improvements to temp storage, P4 tweak to give better errors

diff --git a/Engine/Source/Programs/AutomationTool/P4Utils.cs b/Engine/Source/Programs/AutomationTool/P4Utils.cs
--- a/Engine/Source/Programs/AutomationTool/P4Utils.cs
+++ b/Engine/Source/Programs/AutomationTool/P4Utils.cs
@@ -620,9 +620,8 @@ namespace AutomationTool
                 string Output;
                 if (!LogP4Output(out Output, "changes " + CommandLine, null, AllowSpew))
                 {
-                    return false;
+                    throw new AutomationException("P4 returned failure.");
                 }
-
                 var Lines = Output.Split(new string[] { Environment.NewLine }, StringSplitOptions.None);
                 for(int LineIndex = 0; LineIndex < Lines.Length; ++LineIndex)
                 {
@@ -715,8 +714,10 @@ namespace AutomationTool
 					}
                 }
             }
-            catch (Exception)
+            catch (Exception Ex)
             {
+                CommandUtils.Log(System.Diagnostics.TraceEventType.Warning, "Unable to get P4 changes with {0}", CommandLine);
+                CommandUtils.Log(System.Diagnostics.TraceEventType.Warning, " Exception was {0}", LogUtils.FormatException(Ex));
                 return false;
             }
             ChangeRecords.Sort((A, B) => ChangeRecord.Compare(A, B));
diff --git a/Engine/Source/Programs/AutomationTool/Scripts/GUBP.Automation.cs b/Engine/Source/Programs/AutomationTool/Scripts/GUBP.Automation.cs
--- a/Engine/Source/Programs/AutomationTool/Scripts/GUBP.Automation.cs
+++ b/Engine/Source/Programs/AutomationTool/Scripts/GUBP.Automation.cs
@@ -5012,6 +5012,18 @@ if (HostPlatform == UnrealTargetPlatform.Mac) continue; //temp hack till mac aut
                     }
                 }
             }
+            {
+                var StartTime = DateTime.UtcNow;
+                foreach (var NodeToDo in OrdereredToDo)
+                {
+                    if (GUBPNodes[NodeToDo].RunInEC() && !NodeIsAlreadyComplete(NodeToDo, LocalOnly)) // if something is already finished, we don't put it into EC
+                    {
+                        CleanSharedTempStorageDirectory(GUBPNodes[NodeToDo].GameNameIfAnyForTempStorage());
+                    }
+                }
+                var BuildDuration = (DateTime.UtcNow - StartTime).TotalMilliseconds;
+                Log("Took {0}s to clear temp storage of old files.", BuildDuration / 1000);
+            }
             string ParentPath = ParseParamValue("ParentPath");
             string BaseArgs = String.Format("createJobStep --parentPath {0}", ParentPath);
 
diff --git a/Engine/Source/Programs/AutomationTool/TempStorage.cs b/Engine/Source/Programs/AutomationTool/TempStorage.cs
--- a/Engine/Source/Programs/AutomationTool/TempStorage.cs
+++ b/Engine/Source/Programs/AutomationTool/TempStorage.cs
@@ -470,15 +470,8 @@ namespace AutomationTool
             return CombinePaths(Env.LocalRoot, "Engine", "Saved", "GUBP");
         }
 
-        static HashSet<string> TestedForClean = new HashSet<string>();
         static void CleanSharedTempStorage(string Directory)
         {
-            if (!IsBuildMachine || TestedForClean.Contains(Directory) ||
-                UnrealBuildTool.Utils.IsRunningOnMono)  // saw a hang on this, anyway it isn't necessary to clean with macs, they are slow anyway
-            {
-                return;
-            }
-            TestedForClean.Add(Directory);
             const int MaximumDaysToKeepTempStorage = 2;
             DirectoryInfo DirInfo = new DirectoryInfo(Directory);
             var TopLevelDirs = DirInfo.GetDirectories();
@@ -537,8 +530,13 @@ namespace AutomationTool
             return true;
         }
 
-        public static string ResolveSharedTempStorageDirectory(string GameFolder, bool bClean = true)
+        static Dictionary<string, string> ResolveCache = new Dictionary<string, string>();
+        public static string ResolveSharedTempStorageDirectory(string GameFolder)
         {
+            if (ResolveCache.ContainsKey(GameFolder))
+            {
+                return ResolveCache[GameFolder];
+            }
             string SharedSubdir = SharedTempStorageDirectory();
             string Root = RootSharedTempStorageDirectory();
             string Result = CombinePaths(Root, GameFolder, SharedSubdir);
@@ -561,15 +559,33 @@ namespace AutomationTool
                     }
                 }
             }
-            if (bClean)
+            ResolveCache.Add(GameFolder, Result);
+            return Result;
+        }
+
+        static HashSet<string> TestedForClean = new HashSet<string>();
+        public static void CleanSharedTempStorageDirectory(string GameFolder)
+        {
+
+            if (!IsBuildMachine || TestedForClean.Contains(GameFolder) ||
+                UnrealBuildTool.Utils.IsRunningOnMono)  // saw a hang on this, anyway it isn't necessary to clean with macs, they are slow anyway
             {
-                CleanSharedTempStorage(Result);
+                return;
+            }
+            TestedForClean.Add(GameFolder);
+            try
+            {
+                CleanSharedTempStorage(ResolveSharedTempStorageDirectory(GameFolder));
+            }
+            catch (Exception Ex)
+            {
+                Log(System.Diagnostics.TraceEventType.Warning, "Unable to Clean Directory with GameName {0}", GameFolder);
+                Log(System.Diagnostics.TraceEventType.Warning, " Exception was {0}", LogUtils.FormatException(Ex));
             }
-            return Result;
         }
-        public static string SharedTempStorageDirectory(string StorageBlock, string GameFolder = "", bool bClean = true)
+        public static string SharedTempStorageDirectory(string StorageBlock, string GameFolder = "")
         {
-            return CombinePaths(ResolveSharedTempStorageDirectory(GameFolder, bClean), StorageBlock);
+            return CombinePaths(ResolveSharedTempStorageDirectory(GameFolder), StorageBlock);
         }
 
         public static TempStorageManifest SaveTempStorageManifest(string RootDir, string FinalFilename, List<string> Files)
@@ -619,9 +635,9 @@ namespace AutomationTool
         {
             return SaveTempStorageManifest(BaseFolder, LocalTempStorageManifestFilename(Env, StorageBlockName), Files);
         }
-        public static string SharedTempStorageManifestFilename(CommandEnvironment Env, string StorageBlockName, string GameFolder, bool bClean = true)
+        public static string SharedTempStorageManifestFilename(CommandEnvironment Env, string StorageBlockName, string GameFolder)
         {
-            return CombinePaths(SharedTempStorageDirectory(StorageBlockName, GameFolder, bClean), StorageBlockName + ".TempManifest");
+            return CombinePaths(SharedTempStorageDirectory(StorageBlockName, GameFolder), StorageBlockName + ".TempManifest");
         }
         public static TempStorageManifest SaveSharedTempStorageManifest(CommandEnvironment Env, string StorageBlockName, string GameFolder, List<string> Files)
         {
@@ -646,7 +662,7 @@ namespace AutomationTool
         }
         public static bool SharedTempStorageExists(CommandEnvironment Env, string StorageBlockName, string GameFolder = "", bool bQuiet = false)
         {
-            var SharedManifest = SharedTempStorageManifestFilename(Env, StorageBlockName, GameFolder, false);
+            var SharedManifest = SharedTempStorageManifestFilename(Env, StorageBlockName, GameFolder);
             if (FileExists_NoExceptions(bQuiet, SharedManifest))
             {
                 return true;
@@ -749,7 +765,7 @@ namespace AutomationTool
             }
             WasLocal = false;
 
-            var BlockPath = CombinePaths(SharedTempStorageDirectory(StorageBlockName, GameFolder, false), "/");
+            var BlockPath = CombinePaths(SharedTempStorageDirectory(StorageBlockName, GameFolder), "/");
             if (!BlockPath.EndsWith("/") && !BlockPath.EndsWith("\\"))
             {
                 throw new AutomationException("base folder {0} should end with a separator", BlockPath);
