build script, change criteria for the branch having a shared promotable

diff --git a/Engine/Source/Programs/AutomationTool/Scripts/GUBP.Automation.cs b/Engine/Source/Programs/AutomationTool/Scripts/GUBP.Automation.cs
--- a/Engine/Source/Programs/AutomationTool/Scripts/GUBP.Automation.cs
+++ b/Engine/Source/Programs/AutomationTool/Scripts/GUBP.Automation.cs
@@ -4057,73 +4057,45 @@ public class GUBP : BuildCommand
 
             bool DoASharedPromotable = false;
 
-            int NumShared = 0;
+            int NumSharedCode = 0;
             foreach (var CodeProj in Branch.CodeProjects)
             {
                 var Options = CodeProj.Options(HostPlatform);
 
                 if (Options.bIsPromotable && !Options.bSeparateGamePromotion)
                 {
-                    NumShared++;
+                    NumSharedCode++;
                 }
             }
+
+            var NonCodeProjectNames = new Dictionary<string, List<UnrealTargetPlatform>>();
+            var NonCodeFormalBuilds = new Dictionary<string, List<KeyValuePair<UnrealTargetPlatform, UnrealTargetConfiguration>>>();
             {
-                var Options = Branch.BaseEngineProject.Options(HostPlatform);
+                var Target = Branch.BaseEngineProject.Properties.Targets[TargetRules.TargetType.Editor];
 
-                if (!Options.bIsPromotable || Options.bSeparateGamePromotion)
+                foreach (var Codeless in Target.Rules.GUBP_NonCodeProjects_BaseEditorTypeOnly(HostPlatform))
                 {
-                    if (NumShared > 0)
+                    var Proj = Branch.FindGame(Codeless.Key);
+                    if (Proj == null)
                     {
-                        throw new AutomationException("We assume that if we have shared promotable, the base engine is in it. Some games are marked as shared and promotable, but the base engine is not.");
+                        Log(System.Diagnostics.TraceEventType.Warning, "{0} was listed as a codeless project by GUBP_NonCodeProjects_BaseEditorTypeOnly, however it does not exist in this branch.", Codeless.Key);
                     }
-                }
-                else if (NumShared > 0)
-                {
-                    DoASharedPromotable = true;
-                }
-            }
-
-            if (!DoASharedPromotable && bBuildRocket)
-            {
-                throw new AutomationException("we were asked to make a rocket build, but this branch does not have a shared promotable.");
-            }
-
-            if (DoASharedPromotable)
-            {
-                AddNode(new NonUnityTestNode(HostPlatform));
-
-                var AgentSharingGroup = "Shared_EditorTests" + HostPlatformNode.StaticGetHostPlatformSuffix(HostPlatform);
-
-                Dictionary<string, List<UnrealTargetPlatform>> NonCodeProjectNames;
-                var NonCodeFormalBuilds = new Dictionary<string, List<KeyValuePair<UnrealTargetPlatform, UnrealTargetConfiguration>>>();
-                {
-                    var Target = Branch.BaseEngineProject.Properties.Targets[TargetRules.TargetType.Editor];
-                    NonCodeProjectNames = Target.Rules.GUBP_NonCodeProjects_BaseEditorTypeOnly(HostPlatform);
-
-                    foreach (var Codeless in NonCodeProjectNames)
+                    else if (Proj.Properties.bIsCodeBasedProject)
                     {
-                        var Proj = Branch.FindGame(Codeless.Key);
-                        if (Proj == null)
-                        {
-                            Log(System.Diagnostics.TraceEventType.Warning, "{0} was listed as a codeless project by GUBP_NonCodeProjects_BaseEditorTypeOnly, however it does not exist in this branch.", Codeless.Key);
-                        }
-                        else if (Proj.Properties.bIsCodeBasedProject)
-                        {
-                            throw new AutomationException("{0} was listed as a codeless project by GUBP_NonCodeProjects_BaseEditorTypeOnly, however it is a code based project.", Codeless.Key);
-                        }
+                        throw new AutomationException("{0} was listed as a codeless project by GUBP_NonCodeProjects_BaseEditorTypeOnly, however it is a code based project.", Codeless.Key);
                     }
-
-                    var Options = Branch.BaseEngineProject.Options(HostPlatform);
-
-                    if (!Options.bIsPromotable || Options.bSeparateGamePromotion)
+                    else
                     {
-                        throw new AutomationException("We assume that if we have shared promotable, the base engine is in it.");
+                        NonCodeProjectNames.Add(Codeless.Key, Codeless.Value);
                     }
+                }
 
-                    var TempNonCodeFormalBuilds = Target.Rules.GUBP_NonCodeFormalBuilds_BaseEditorTypeOnly();
-                    var HostMonos = GetMonolithicPlatformsForUProject(HostPlatform, Branch.BaseEngineProject, true);
+                var TempNonCodeFormalBuilds = Target.Rules.GUBP_NonCodeFormalBuilds_BaseEditorTypeOnly();
+                var HostMonos = GetMonolithicPlatformsForUProject(HostPlatform, Branch.BaseEngineProject, true);
 
-                    foreach (var Codeless in TempNonCodeFormalBuilds)
+                foreach (var Codeless in TempNonCodeFormalBuilds)
+                {
+                    if (NonCodeProjectNames.ContainsKey(Codeless.Key))
                     {
                         var PlatList = Codeless.Value;
                         var NewPlatList = new List<KeyValuePair<UnrealTargetPlatform, UnrealTargetConfiguration>>();
@@ -4139,7 +4111,35 @@ public class GUBP : BuildCommand
                             NonCodeFormalBuilds.Add(Codeless.Key, NewPlatList);
                         }
                     }
+                    else
+                    {
+                        Log(System.Diagnostics.TraceEventType.Warning, "{0} was listed as a codeless formal build GUBP_NonCodeFormalBuilds_BaseEditorTypeOnly, however it does not exist in this branch.", Codeless.Key);
+                    }
                 }
+            }
+
+            DoASharedPromotable = NumSharedCode > 0 || NonCodeProjectNames.Count > 0 || NonCodeFormalBuilds.Count > 0;
+
+            if (!DoASharedPromotable && bBuildRocket)
+            {
+                throw new AutomationException("we were asked to make a rocket build, but this branch does not have a shared promotable.");
+            }
+
+            if (DoASharedPromotable)
+            {
+                AddNode(new NonUnityTestNode(HostPlatform));
+
+                var AgentSharingGroup = "Shared_EditorTests" + HostPlatformNode.StaticGetHostPlatformSuffix(HostPlatform);
+
+                var Options = Branch.BaseEngineProject.Options(HostPlatform);
+
+                if (!Options.bIsPromotable || Options.bSeparateGamePromotion)
+                {
+                    throw new AutomationException("We assume that if we have shared promotable, the base engine is in it.");
+                }
+
+
+
                 if (bAutomatedTesting)
                 {
                     var EditorTests = Branch.BaseEngineProject.Properties.Targets[TargetRules.TargetType.Editor].Rules.GUBP_GetEditorTests_EditorTypeOnly(HostPlatform);
