build script, move cleaning to be at the end of the editor agent group so that it doesn't cause any latency

diff --git a/Engine/Source/Programs/AutomationTool/Scripts/GUBP.Automation.cs b/Engine/Source/Programs/AutomationTool/Scripts/GUBP.Automation.cs
--- a/Engine/Source/Programs/AutomationTool/Scripts/GUBP.Automation.cs
+++ b/Engine/Source/Programs/AutomationTool/Scripts/GUBP.Automation.cs
@@ -2650,7 +2650,7 @@ public class GUBP : BuildCommand
 
         public static string StaticGetFullName(UnrealTargetPlatform InHostPlatform, BranchInfo.BranchUProject InGameProj, string InAggregateName)
         {
-            return InGameProj.GameName + "_" + InAggregateName + StaticGetHostPlatformSuffix(InHostPlatform); ;
+            return InGameProj.GameName + "_" + InAggregateName + StaticGetHostPlatformSuffix(InHostPlatform);
         }
         public override string GetFullName()
         {
@@ -2662,6 +2662,47 @@ public class GUBP : BuildCommand
         }
     };
 
+    public class CleanSharedTempStorageNode : GUBPNode
+    {
+        public CleanSharedTempStorageNode(GUBP bp)
+        {
+            var ToolsNode = bp.FindNode(ToolsForCompileNode.StaticGetFullName(UnrealTargetPlatform.Win64));
+            AgentSharingGroup = ToolsNode.AgentSharingGroup;
+        }
+        public override float Priority()
+        {
+            return -1E15f;
+        }
+
+        public static string StaticGetFullName()
+        {
+            return "CleanSharedTempStorage";
+        }
+        public override string GetFullName()
+        {
+            return StaticGetFullName();
+        }
+        public override void DoBuild(GUBP bp)
+        {
+            {
+                var StartTime = DateTime.UtcNow;
+                foreach (var NodeToDo in bp.GUBPNodes)
+                {
+                    CleanSharedTempStorageDirectory(NodeToDo.Value.GameNameIfAnyForTempStorage());
+                }
+                var BuildDuration = (DateTime.UtcNow - StartTime).TotalMilliseconds;
+                Log("Took {0}s to clear temp storage of old files.", BuildDuration / 1000);
+            }
+
+            BuildProducts = new List<string>();
+            SaveRecordOfSuccessAndAddToBuildProducts();
+        }
+        public override int CISFrequencyQuantumShift(GUBP bp)
+        {
+            return base.CISFrequencyQuantumShift(bp) + 4;
+        }
+    };
+
     public string AddNode(GUBPNode Node)
     {
         string Name = Node.GetFullName();
@@ -4542,6 +4583,11 @@ if (HostPlatform == UnrealTargetPlatform.Mac) continue; //temp hack till mac aut
             AddCustomNodes(HostPlatform);
         }
 
+        if (HasNode(ToolsForCompileNode.StaticGetFullName(UnrealTargetPlatform.Win64)))
+        {
+            AddNode(new CleanSharedTempStorageNode(this));
+        }
+
         foreach (var NodeToDo in GUBPNodes)
         {
             foreach (var Dep in GUBPNodes[NodeToDo.Key].FullNamesOfDependencies)
@@ -5012,18 +5058,6 @@ if (HostPlatform == UnrealTargetPlatform.Mac) continue; //temp hack till mac aut
                     }
                 }
             }
-            {
-                var StartTime = DateTime.UtcNow;
-                foreach (var NodeToDo in OrdereredToDo)
-                {
-                    if (GUBPNodes[NodeToDo].RunInEC() && !NodeIsAlreadyComplete(NodeToDo, LocalOnly)) // if something is already finished, we don't put it into EC
-                    {
-                        CleanSharedTempStorageDirectory(GUBPNodes[NodeToDo].GameNameIfAnyForTempStorage());
-                    }
-                }
-                var BuildDuration = (DateTime.UtcNow - StartTime).TotalMilliseconds;
-                Log("Took {0}s to clear temp storage of old files.", BuildDuration / 1000);
-            }
             string ParentPath = ParseParamValue("ParentPath");
             string BaseArgs = String.Format("createJobStep --parentPath {0}", ParentPath);
 
