build script, suppress dotnetutilities temp storage warning

diff --git a/Engine/Source/Programs/AutomationTool/TempStorage.cs b/Engine/Source/Programs/AutomationTool/TempStorage.cs
--- a/Engine/Source/Programs/AutomationTool/TempStorage.cs
+++ b/Engine/Source/Programs/AutomationTool/TempStorage.cs
@@ -52,14 +52,23 @@ namespace AutomationTool
                     //temp hack until the mac build products work correctly
                     bOkToBeDifferent = bOkToBeDifferent || Name.Contains("Engine/Binaries/Mac/UE4Editor.app/Contents/MacOS/");
 
+
+                    bool bSilentOkToBeDifferent = (Name == "Engine/Binaries/DotNET/DotNETUtilities.dll");
+
+                    System.Diagnostics.TraceEventType LogType = bOkToBeDifferent ? System.Diagnostics.TraceEventType.Warning : System.Diagnostics.TraceEventType.Error;
+                    if (bSilentOkToBeDifferent && bOkToBeDifferent)
+                    {
+                        LogType = System.Diagnostics.TraceEventType.Information;
+                    }
+
                     if (!((Timestamp - Other.Timestamp).TotalSeconds < 1 && (Timestamp - Other.Timestamp).TotalSeconds > -1))
                     {
-                        CommandUtils.Log(bOkToBeDifferent ? System.Diagnostics.TraceEventType.Warning : System.Diagnostics.TraceEventType.Error, "File date mismatch {0} {1} {2} {3}", Name, Timestamp.ToString(), Other.Name, Other.Timestamp.ToString());
+                        CommandUtils.Log(LogType, "File date mismatch {0} {1} {2} {3}", Name, Timestamp.ToString(), Other.Name, Other.Timestamp.ToString());
                         bOk = bOkToBeDifferent;
                     }
                     if (!(Size == Other.Size))
                     {
-                        CommandUtils.Log(bOkToBeDifferent ? System.Diagnostics.TraceEventType.Warning : System.Diagnostics.TraceEventType.Error, "File size mismatch {0} {1} {2} {3}", Name, Size, Other.Name, Other.Size);
+                        CommandUtils.Log(LogType, "File size mismatch {0} {1} {2} {3}", Name, Size, Other.Name, Other.Size);
                         bOk = bOkToBeDifferent;
                     }
                 }
