build script, suppress more warnings in the GUBP

diff --git a/Engine/Source/Programs/AutomationTool/TempStorage.cs b/Engine/Source/Programs/AutomationTool/TempStorage.cs
--- a/Engine/Source/Programs/AutomationTool/TempStorage.cs
+++ b/Engine/Source/Programs/AutomationTool/TempStorage.cs
@@ -53,7 +53,12 @@ namespace AutomationTool
                     bOkToBeDifferent = bOkToBeDifferent || Name.Contains("Engine/Binaries/Mac/UE4Editor.app/Contents/MacOS/");
 
 
+                    // DotNETUtilities.dll is built by a tons of other things
                     bool bSilentOkToBeDifferent = (Name == "Engine/Binaries/DotNET/DotNETUtilities.dll");
+                    bSilentOkToBeDifferent = bSilentOkToBeDifferent || (Name == "Engine/Binaries/DotNET/DotNETUtilities.pdb");
+                    // RPCUtility is build by IPP and maybe other things
+                    bSilentOkToBeDifferent = bSilentOkToBeDifferent || (Name == "Engine/Binaries/DotNET/RPCUtility.exe");
+                    bSilentOkToBeDifferent = bSilentOkToBeDifferent || (Name == "Engine/Binaries/DotNET/RPCUtility.pdb");
 
                     System.Diagnostics.TraceEventType LogType = bOkToBeDifferent ? System.Diagnostics.TraceEventType.Warning : System.Diagnostics.TraceEventType.Error;
                     if (bSilentOkToBeDifferent && bOkToBeDifferent)
