L10N: Updated loc ini configurations to specify simplified Chinese as output. Added a C# wrapper around OneSky's web service api. Beginning steps of setting up some automation scripts to integrate our localization pipeline with OneSky's.

diff --git a/Engine/Config/Localization/Editor.ini b/Engine/Config/Localization/Editor.ini
--- a/Engine/Config/Localization/Editor.ini
+++ b/Engine/Config/Localization/Editor.ini
@@ -8,7 +8,7 @@ SourceCulture=en
 CulturesToGenerate=en
 CulturesToGenerate=ja
 CulturesToGenerate=ko
-CulturesToGenerate=zh
+CulturesToGenerate=zh-CN
 
 ;Gather text from source code
 [GatherTextStep0]
@@ -28,6 +28,7 @@ IncludePaths=*/Engine/Source/Runtime/*
 IncludePaths=*/Engine/Source/Developer/*
 ExcludePaths=*/Engine/Config/Localization/*
 ExcludePaths=*/Engine/Source/Developer/NoRedist/CommunityPortalServices/*
+ExcludePaths=*/Engine/Source/Developer/NoRedist/BuildPatchServices/*
 InputKeys=Category
 OutputNamespaces=UObjectCategories
 OutputKeys="{MetaDataValue}"
diff --git a/Engine/Config/Localization/Engine.ini b/Engine/Config/Localization/Engine.ini
--- a/Engine/Config/Localization/Engine.ini
+++ b/Engine/Config/Localization/Engine.ini
@@ -7,7 +7,7 @@ SourceCulture=en
 CulturesToGenerate=en
 CulturesToGenerate=ja
 CulturesToGenerate=ko
-CulturesToGenerate=zh
+CulturesToGenerate=zh-CN
 
 ;Gather text from source code
 [GatherTextStep0]
@@ -18,6 +18,7 @@ IncludePaths=Engine/Source/Developer
 IncludePaths=Engine/Config/
 ExcludePaths=*/Engine/Config/Localization/*
 ExcludePaths=*/Engine/Source/Developer/NoRedist/CommunityPortalServices/*
+ExcludePaths=*/Engine/Source/Developer/NoRedist/BuildPatchServices/*
 SourceFileSearchFilters=*.cpp
 SourceFileSearchFilters=*.h
 SourceFileSearchFilters=*.c
diff --git a/Engine/Config/Localization/PortableObjectExport.ini b/Engine/Config/Localization/PortableObjectExport.ini
--- a/Engine/Config/Localization/PortableObjectExport.ini
+++ b/Engine/Config/Localization/PortableObjectExport.ini
@@ -1,7 +1,8 @@
 [CommonSettings]
+CulturesToGenerate=en
 CulturesToGenerate=ja
 CulturesToGenerate=ko
-CulturesToGenerate=zh
+CulturesToGenerate=zh-CN
 
 [GatherTextStep0]
 CommandletClass=InternationalizationExport
@@ -10,6 +11,7 @@ SourcePath=../../../Engine/Content/Localization/Engine
 DestinationPath=../../../Engine/Content/Localization/Engine
 ManifestName=Engine.manifest
 ArchiveName=Engine.archive
+PortableObjectName=Engine.po
 
 [GatherTextStep1]
 CommandletClass=InternationalizationExport
@@ -18,6 +20,7 @@ SourcePath=../../../Engine/Content/Localization/Editor
 DestinationPath=../../../Engine/Content/Localization/Editor
 ManifestName=Editor.manifest
 ArchiveName=Editor.archive
+PortableObjectName=Editor.po
 
 [GatherTextStep2]
 CommandletClass=InternationalizationExport
@@ -26,6 +29,7 @@ SourcePath=../../../Engine/Content/Localization/Editor
 DestinationPath=../../../Engine/Content/Localization/Editor
 ManifestName=ToolTips.manifest
 ArchiveName=ToolTips.archive
+PortableObjectName=ToolTips.po
 
 [GatherTextStep3]
 CommandletClass=InternationalizationExport
@@ -33,4 +37,5 @@ bExportLoc=true
 SourcePath=../../../Engine/Content/Localization/Editor
 DestinationPath=../../../Engine/Content/Localization/Editor
 ManifestName=PropertyNames.manifest
-ArchiveName=PropertyNames.archive
\ No newline at end of file
+ArchiveName=PropertyNames.archive
+PortableObjectName=PropertyNames.po
\ No newline at end of file
diff --git a/Engine/Config/Localization/PortableObjectImport.ini b/Engine/Config/Localization/PortableObjectImport.ini
--- a/Engine/Config/Localization/PortableObjectImport.ini
+++ b/Engine/Config/Localization/PortableObjectImport.ini
@@ -1,7 +1,8 @@
 [CommonSettings]
+CulturesToGenerate=en
 CulturesToGenerate=ja
 CulturesToGenerate=ko
-CulturesToGenerate=zh
+CulturesToGenerate=zh-CN
 
 [GatherTextStep0]
 CommandletClass=InternationalizationExport
diff --git a/Engine/Config/Localization/PropertyNames.ini b/Engine/Config/Localization/PropertyNames.ini
--- a/Engine/Config/Localization/PropertyNames.ini
+++ b/Engine/Config/Localization/PropertyNames.ini
@@ -7,7 +7,7 @@ SourceCulture=en
 CulturesToGenerate=en
 CulturesToGenerate=ja
 CulturesToGenerate=ko
-CulturesToGenerate=zh
+CulturesToGenerate=zh-CN
 
 ;Gather text from metadata
 [GatherTextStep0]
diff --git a/Engine/Config/Localization/ResourceFileGen.ini b/Engine/Config/Localization/ResourceFileGen.ini
--- a/Engine/Config/Localization/ResourceFileGen.ini
+++ b/Engine/Config/Localization/ResourceFileGen.ini
@@ -7,7 +7,7 @@ ManifestName=Engine.manifest
 CulturesToGenerate=en
 CulturesToGenerate=ja
 CulturesToGenerate=ko
-CulturesToGenerate=zh
+CulturesToGenerate=zh-CN
 ResourceName=Engine.locres
 
 ;Write Localized Text Resource for Editor
@@ -19,7 +19,7 @@ ManifestName=Editor.manifest
 CulturesToGenerate=en
 CulturesToGenerate=ja
 CulturesToGenerate=ko
-CulturesToGenerate=zh
+CulturesToGenerate=zh-CN
 ResourceName=Editor.locres
 
 ;Write Localized Text Resource for Editor Tool Tips
@@ -31,7 +31,7 @@ ManifestName=ToolTips.manifest
 CulturesToGenerate=en
 CulturesToGenerate=ja
 CulturesToGenerate=ko
-CulturesToGenerate=zh
+CulturesToGenerate=zh-CN
 ResourceName=ToolTips.locres
 
 ;Write Localized Text Resource for Editor Property Names
@@ -43,5 +43,5 @@ ManifestName=PropertyNames.manifest
 CulturesToGenerate=en
 CulturesToGenerate=ja
 CulturesToGenerate=ko
-CulturesToGenerate=zh
+CulturesToGenerate=zh-CN
 ResourceName=PropertyNames.locres
\ No newline at end of file
diff --git a/Engine/Config/Localization/ToolTips.ini b/Engine/Config/Localization/ToolTips.ini
--- a/Engine/Config/Localization/ToolTips.ini
+++ b/Engine/Config/Localization/ToolTips.ini
@@ -7,7 +7,7 @@ SourceCulture=en
 CulturesToGenerate=en
 CulturesToGenerate=ja
 CulturesToGenerate=ko
-CulturesToGenerate=zh
+CulturesToGenerate=zh-CN
 
 ;Gather text from metadata
 [GatherTextStep0]
diff --git a/Engine/Config/Localization/WordCount.ini b/Engine/Config/Localization/WordCount.ini
--- a/Engine/Config/Localization/WordCount.ini
+++ b/Engine/Config/Localization/WordCount.ini
@@ -9,7 +9,7 @@ SourcePath=../../../Engine/Content/Localization/Engine
 ManifestName=Engine.manifest
 CulturesToGenerate=ja
 CulturesToGenerate=ko
-CulturesToGenerate=zh
+CulturesToGenerate=zh-CN
 WordCountReportName=Engine.csv
 
 ;Write Text Localization Report for Editor
@@ -20,7 +20,7 @@ SourcePath=../../../Engine/Content/Localization/Editor
 ManifestName=Editor.manifest
 CulturesToGenerate=ja
 CulturesToGenerate=ko
-CulturesToGenerate=zh
+CulturesToGenerate=zh-CN
 WordCountReportName=Editor.csv
 
 ;Write Text Localization Report for Editor Tool Tips
@@ -31,7 +31,7 @@ SourcePath=../../../Engine/Content/Localization/ToolTips
 ManifestName=ToolTips.manifest
 CulturesToGenerate=ja
 CulturesToGenerate=ko
-CulturesToGenerate=zh
+CulturesToGenerate=zh-CN
 WordCountReportName=ToolTips.csv
 
 ;Write Text Localization Report for Editor Property Names
@@ -42,5 +42,5 @@ SourcePath=../../../Engine/Content/Localization/PropertyNames
 ManifestName=PropertyNames.manifest
 CulturesToGenerate=ja
 CulturesToGenerate=ko
-CulturesToGenerate=zh
+CulturesToGenerate=zh-CN
 WordCountReportName=PropertyNames.csv
\ No newline at end of file
diff --git a/Engine/Source/Editor/UnrealEd/Classes/Commandlets/InternationalizationExportCommandlet.h b/Engine/Source/Editor/UnrealEd/Classes/Commandlets/InternationalizationExportCommandlet.h
--- a/Engine/Source/Editor/UnrealEd/Classes/Commandlets/InternationalizationExportCommandlet.h
+++ b/Engine/Source/Editor/UnrealEd/Classes/Commandlets/InternationalizationExportCommandlet.h
@@ -372,8 +372,8 @@ public:
 	virtual int32 Main(const FString& Params) OVERRIDE;
 	// End UCommandlet Interface
 private:
-	bool DoExport( const FString& SourcePath, const FString& DestinationPath );
-	bool DoImport( const FString& SourcePath, const FString& DestinationPath );
+	bool DoExport(const FString& SourcePath, const FString& DestinationPath, const FString& Filename);
+	bool DoImport(const FString& SourcePath, const FString& DestinationPath, const FString& Filename);
 
 	TArray<FString> CulturesToGenerate;
 	FString ConfigPath;
diff --git a/Engine/Source/Editor/UnrealEd/Private/Commandlets/InternationalizationExportCommandlet.cpp b/Engine/Source/Editor/UnrealEd/Private/Commandlets/InternationalizationExportCommandlet.cpp
--- a/Engine/Source/Editor/UnrealEd/Private/Commandlets/InternationalizationExportCommandlet.cpp
+++ b/Engine/Source/Editor/UnrealEd/Private/Commandlets/InternationalizationExportCommandlet.cpp
@@ -880,13 +880,15 @@ TSharedPtr<FPortableObjectEntry> FPortableObjectFormatDOM::FindEntry( const FStr
 
 void FPortableObjectEntry::AddReference( const FString& InReference )
 {
-	// Reference comments can contain multiple references in a single line so we parse those out.
-	TArray<FString> ReferencesToProcess;
-	InReference.ParseIntoArray( &ReferencesToProcess, TEXT(" "), true );
-	for( const FString& Reference : ReferencesToProcess )
-	{
-		ReferenceComments.AddUnique( Reference );
-	}
+	ReferenceComments.AddUnique(InReference);
+
+	//// Reference comments can contain multiple references in a single line so we parse those out.
+	//TArray<FString> ReferencesToProcess;
+	//InReference.ParseIntoArray( &ReferencesToProcess, TEXT(" "), true );
+	//for( const FString& Reference : ReferencesToProcess )
+	//{
+	//	ReferenceComments.AddUnique( Reference );
+	//}
 }
 
 void FPortableObjectEntry::AddReferences( const TArray<FString>& InReferences )
@@ -961,7 +963,7 @@ UInternationalizationExportCommandlet::UInternationalizationExportCommandlet(con
 }
 
 
-bool UInternationalizationExportCommandlet::DoExport( const FString& SourcePath, const FString& DestinationPath  )
+bool UInternationalizationExportCommandlet::DoExport( const FString& SourcePath, const FString& DestinationPath, const FString& Filename )
 {
 	// Get manifest name.
 	FString ManifestName;
@@ -1067,7 +1069,7 @@ bool UInternationalizationExportCommandlet::DoExport( const FString& SourcePath,
 					FString OutputFileName;
 					if (FPaths::GetExtension(DestinationPath) != "po")
 					{
-						OutputFileName = DestinationPath / CultureName + TEXT(".po");
+						OutputFileName = DestinationPath / CultureName / Filename;
 					}
 					else
 					{
@@ -1087,7 +1089,7 @@ bool UInternationalizationExportCommandlet::DoExport( const FString& SourcePath,
 	return true;
 }
 
-bool UInternationalizationExportCommandlet::DoImport( const FString& SourcePath, const FString& DestinationPath  )
+bool UInternationalizationExportCommandlet::DoImport(const FString& SourcePath, const FString& DestinationPath, const FString& Filename)
 {
 	// Get manifest name.
 	FString ManifestName;
@@ -1113,7 +1115,7 @@ bool UInternationalizationExportCommandlet::DoImport( const FString& SourcePath,
 		FString POFilePath;
 		if (FPaths::GetExtension(SourcePath) != "po")
 		{
-			POFilePath = SourcePath / CultureName + TEXT(".po");
+			POFilePath = SourcePath / CultureName / Filename;
 		}
 		else
 		{
@@ -1122,7 +1124,7 @@ bool UInternationalizationExportCommandlet::DoImport( const FString& SourcePath,
 
 		if( !FPaths::FileExists(POFilePath) )
 		{
-			UE_LOG( LogInternationalizationExportCommandlet, Error, TEXT("Could not find file %s"), *POFilePath );
+			UE_LOG( LogInternationalizationExportCommandlet, Warning, TEXT("Could not find file %s"), *POFilePath );
 			continue;
 		}
 
@@ -1265,6 +1267,13 @@ int32 UInternationalizationExportCommandlet::Main( const FString& Params )
 		return -1;
 	}
 
+	FString Filename; // Name of the file to read or write from
+	if (!GetConfigString(*SectionName, TEXT("PortableObjectName"), Filename, ConfigPath))
+	{
+		UE_LOG(LogInternationalizationExportCommandlet, Error, TEXT("No portable object name specified."));
+		return -1;
+	}
+
 	// Get cultures to generate.
 	if( GetConfigArray(*SectionName, TEXT("CulturesToGenerate"), CulturesToGenerate, ConfigPath) == 0 )
 	{
@@ -1287,7 +1296,7 @@ int32 UInternationalizationExportCommandlet::Main( const FString& Params )
 
 	if( bDoExport )
 	{
-		if( !DoExport( SourcePath, DestinationPath ) )
+		if (!DoExport(SourcePath, DestinationPath, Filename))
 		{
 			UE_LOG(LogInternationalizationExportCommandlet, Error, TEXT("Failed to export localization files."));
 			return -1;
@@ -1296,7 +1305,7 @@ int32 UInternationalizationExportCommandlet::Main( const FString& Params )
 
 	if( bDoImport )
 	{
-		if( !DoImport( SourcePath, DestinationPath ) )
+		if (!DoImport(SourcePath, DestinationPath, Filename))
 		{
 			UE_LOG(LogInternationalizationExportCommandlet, Error, TEXT("Failed to import localization files."));
 			return -1;
diff --git a/Engine/Source/Programs/AutomationTool/Scripts/AutomationScripts.Automation.csproj b/Engine/Source/Programs/AutomationTool/Scripts/AutomationScripts.Automation.csproj
--- a/Engine/Source/Programs/AutomationTool/Scripts/AutomationScripts.Automation.csproj
+++ b/Engine/Source/Programs/AutomationTool/Scripts/AutomationScripts.Automation.csproj
@@ -38,6 +38,10 @@
       <HintPath>..\..\..\..\Binaries\DotNET\Ionic.Zip.Reduced.dll</HintPath>
       <Private>False</Private>
     </Reference>
+    <Reference Include="OneSky, Version=1.0.0.0, Culture=neutral, processorArchitecture=MSIL">
+      <SpecificVersion>False</SpecificVersion>
+      <HintPath>..\..\..\ThirdParty\OneSky\bin\Release\OneSky.dll</HintPath>
+    </Reference>
     <Reference Include="System" />
     <Reference Include="System.Core" />
     <Reference Include="System.Runtime.Serialization" />
diff --git a/Engine/Source/Programs/AutomationTool/Scripts/LauncherLocalization.Automation.cs b/Engine/Source/Programs/AutomationTool/Scripts/LauncherLocalization.Automation.cs
--- a/Engine/Source/Programs/AutomationTool/Scripts/LauncherLocalization.Automation.cs
+++ b/Engine/Source/Programs/AutomationTool/Scripts/LauncherLocalization.Automation.cs
@@ -1,10 +1,11 @@
 ﻿// Copyright 1998-2014 Epic Games, Inc. All Rights Reserved.
 using System;
 using System.Collections.Generic;
-using System.Text;
+using System.Globalization;
+using System.Linq;
 using System.IO;
 using AutomationTool;
-using UnrealBuildTool;
+using OneSky;
 
 class LauncherLocalization : BuildCommand
 {
@@ -26,11 +27,26 @@ class LauncherLocalization : BuildCommand
             Log("Localize from label {0}", P4Env.LabelToSync);            
         }
 
+        //var oneSkyService = new OneSkyService("Reo5M8fnKcqVANHb4Zjc5nVUeXeehhGZ", "ByaWoHTjLM6clyxWLrrhhgdTY1btatdB");
+
+        //// Export Launcher text from OneSky
+        //{
+        //    var launcherGroup = GetLauncherGroup(oneSkyService);
+        //    var appProject = GetAppProject(oneSkyService);
+        //    var appFile = appProject.UploadedFiles.FirstOrDefault(f => f.Filename == "App.po");
+
+        //    //Export
+        //    if (appFile != null)
+        //    {
+        //        ExportFileToDirectory(appFile, new DirectoryInfo(CmdEnv.LocalRoot + "/Engine/Programs/NoRedist/UnrealEngineLauncher/Content/Localization/App"), launcherGroup.EnabledCultures);
+        //    }
+        //}
+
         // Setup editor arguments for SCC.
         string EditorArguments = String.Empty;
         if (P4Enabled)
         {
-			EditorArguments = String.Format("-SCCProvider={0} -P4Port={1} -P4User={2} -P4Client={3} -P4Passwd={4}", "Perforce", P4Env.P4Port, P4Env.User, P4Env.Client, P4.GetAuthenticationToken());
+            EditorArguments = String.Format("-SCCProvider={0} -P4Port={1} -P4User={2} -P4Client={3} -P4Passwd={4}", "Perforce", P4Env.P4Port, P4Env.User, P4Env.Client, P4.GetAuthenticationToken());
         }
         else
         {
@@ -44,10 +60,10 @@ class LauncherLocalization : BuildCommand
 
         // Setup commandlet arguments with configurations.
         var CommandletArgumentSets = new string[] 
-				{
+                {
                     String.Format("-config={0}", @"../../../Engine/Programs/NoRedist/UnrealEngineLauncher/Config/Localization/App.ini") + (string.IsNullOrEmpty(CommandletSCCArguments) ? "" : " " + CommandletSCCArguments),
                     String.Format("-config={0}", @"../../../Engine/Programs/NoRedist/UnrealEngineLauncher/Config/Localization/WordCount.ini"),
-				};
+                };
 
         // Execute commandlet for each set of arguments.
         foreach (var CommandletArguments in CommandletArgumentSets)
@@ -64,6 +80,110 @@ class LauncherLocalization : BuildCommand
                 throw new AutomationException("Error while executing localization commandlet '{0}'", Arguments);
             }
         }
+
+        // Upload Launcher text to OneSky
+        //UploadDirectoryToProject(GetAppProject(oneSkyService), new DirectoryInfo(CmdEnv.LocalRoot + "/Engine/Programs/NoRedist/UnrealEngineLauncher/Content/Localization/App"), "*.po");
+    }
+
+    private static ProjectGroup GetLauncherGroup(OneSkyService oneSkyService)
+    {
+        var launcherGroup = oneSkyService.ProjectGroups.FirstOrDefault(g => g.Name == "Launcher");
+
+        if (launcherGroup == null)
+        {
+            launcherGroup = new ProjectGroup("Launcher", new CultureInfo("en"));
+            oneSkyService.ProjectGroups.Add(launcherGroup);
+        }
+
+        return launcherGroup;
+    }
+
+    private static OneSky.Project GetAppProject(OneSkyService oneSkyService)
+    {
+        var launcherGroup = GetLauncherGroup(oneSkyService);
+
+        OneSky.Project appProject = launcherGroup.Projects.FirstOrDefault(p => p.Name == "App");
+
+        if (appProject == null)
+        {
+            ProjectType projectType = oneSkyService.ProjectTypes.First(pt => pt.Code == "website");
+
+            appProject = new OneSky.Project("App", "The core application text that ships with the Launcher", projectType);
+            launcherGroup.Projects.Add(appProject);
+        }
+
+        return appProject;
+    }
+
+    private static void ExportFileToDirectory(UploadedFile file, DirectoryInfo destination, IEnumerable<CultureInfo> cultures)
+    {
+        foreach (var culture in cultures)
+        {
+            var cultureDirectory = new DirectoryInfo(Path.Combine(destination.FullName, culture.Name));
+            if (!cultureDirectory.Exists)
+            {
+                cultureDirectory.Create();
+            }
+
+            using (var memoryStream = new MemoryStream())
+            {
+                var exportFile = new FileInfo(Path.Combine(cultureDirectory.FullName, file.Filename));
+
+                var exportTranslationState = file.ExportTranslation(culture, memoryStream).Result;
+                if (exportTranslationState == UploadedFile.ExportTranslationState.Success)
+                {
+                    memoryStream.Position = 0;
+                    using (Stream fileStream = File.OpenWrite(exportFile.FullName))
+                    {
+                        memoryStream.CopyTo(fileStream);
+                        Console.WriteLine("[SUCCESS] Exporting: " + exportFile.FullName + " Locale: " + culture.Name);
+                    }
+                }
+                else if (exportTranslationState == UploadedFile.ExportTranslationState.NoContent)
+                {
+                    Console.WriteLine("[WARNING] Exporting: " + exportFile.FullName + " Locale: " + culture.Name + " has no translations!");
+                }
+                else
+                {
+                    Console.WriteLine("[FAILED] Exporting: " + exportFile.FullName + " Locale: " + culture.Name);
+                }
+            }
+        }
+    }
+
+    static void UploadDirectoryToProject(OneSky.Project project, DirectoryInfo directory, string fileExtension)
+    {
+        foreach (var file in Directory.GetFiles(directory.FullName, fileExtension, SearchOption.AllDirectories))
+        {
+            DirectoryInfo parentDirectory = Directory.GetParent(file);
+            string localeName = parentDirectory.Name;
+            string currentFile = file;
+
+            using (var fileStream = File.OpenRead(currentFile))
+            {
+                // Read the BOM
+                var bom = new byte[3];
+                fileStream.Read(bom, 0, 3);
+
+                //We want to ignore the utf8 BOM
+                if (bom[0] != 0xef || bom[1] != 0xbb || bom[2] != 0xbf)
+                {
+                    fileStream.Position = 0;
+                }
+
+                Console.WriteLine("Uploading: " + currentFile + " Locale: " + localeName);
+                var uploadedFile = project.Upload(Path.GetFileName(currentFile), fileStream, new CultureInfo(localeName)).Result;
+
+                if (uploadedFile == null)
+                {
+                    Console.WriteLine("[FAILED] Uploading: " + currentFile + " Locale: " + localeName);
+                }
+                else
+                {
+                    Console.WriteLine("[SUCCESS] Uploading: " + currentFile + " Locale: " + localeName);
+                }
+            }
+        }
     }
 }
 
