build script, fix problems with the funky checkins wes has

diff --git a/Engine/Source/Programs/AutomationTool/P4Utils.cs b/Engine/Source/Programs/AutomationTool/P4Utils.cs
--- a/Engine/Source/Programs/AutomationTool/P4Utils.cs
+++ b/Engine/Source/Programs/AutomationTool/P4Utils.cs
@@ -622,6 +622,9 @@ namespace AutomationTool
                 {
                     throw new AutomationException("P4 returned failure.");
                 }
+                
+                //CommandUtils.Log("{0}", Output.Replace("\t", "[tab]").Replace("\r", "[r]").Replace("\n", "[n]"));
+                
                 var Lines = Output.Split(new string[] { Environment.NewLine }, StringSplitOptions.None);
                 for(int LineIndex = 0; LineIndex < Lines.Length; ++LineIndex)
                 {
@@ -680,11 +683,35 @@ namespace AutomationTool
 							{
 								Line = Lines[ LineIndex ];
 
-								if( String.IsNullOrEmpty( Line ) )
-								{
-									// Summaries end with a blank line (no tabs)
-									break;
-								}
+                                bool NextLineIsTab = false;
+
+                                for (int LookAhead = LineIndex; LookAhead < Lines.Length; LookAhead++ )
+                                {
+                                    if (Lines[LookAhead].StartsWith("\t"))
+                                    {
+                                        NextLineIsTab = true;
+                                        break;
+                                    }
+                                    if (Lines[LookAhead] != "")
+                                    {
+                                        break;
+                                    }
+                                }
+
+                                if (String.IsNullOrEmpty(Line))
+                                {
+                                    if (NextLineIsTab)
+                                    {
+                                        // Add a CR if we already had some summary text
+                                        if (!String.IsNullOrEmpty(Change.Summary))
+                                        {
+                                            Change.Summary += "\n";
+                                        } 
+                                        continue;
+                                    }
+                                    // Summaries end with a blank line (no tabs)
+                                    break;
+                                }
 
 								// Summary lines are supposed to begin with a single tab character (even empty lines)
 								if( Line[0] != '\t' )
