Android: Remove the -f from the rm commands added to the batch file. Some devices (eg Galaxy S III) do not recognize -f and fail to remove the file, causing weird behavior.

diff --git a/Engine/Source/Programs/AutomationTool/Android/AndroidPlatform.Automation.cs b/Engine/Source/Programs/AutomationTool/Android/AndroidPlatform.Automation.cs
--- a/Engine/Source/Programs/AutomationTool/Android/AndroidPlatform.Automation.cs
+++ b/Engine/Source/Programs/AutomationTool/Android/AndroidPlatform.Automation.cs
@@ -127,9 +127,9 @@ public class AndroidPlatform : Platform
 			"set ADB=%ANDROID_HOME%\\platform-tools\\adb.exe",
 			"%ADB% uninstall " + PackageName,
 			"%ADB% install " + Path.GetFileName(ApkName),
-			"%ADB% shell rm -rf /mnt/sdcard/" + Params.ShortProjectName,
-			SC.IsCodeBasedProject ? "" : "%ADB% shell rm -rf /mnt/sdcard/UE4Game/UE4CommandLine.txt", // we need to delete the commandline in UE4Game or it will mess up loading
-			"%ADB% shell rm -rf /mnt/sdcard/obb/" + PackageName,
+			"%ADB% shell rm -r /mnt/sdcard/" + Params.ShortProjectName,
+			SC.IsCodeBasedProject ? "" : "%ADB% shell rm /mnt/sdcard/UE4Game/UE4CommandLine.txt", // we need to delete the commandline in UE4Game or it will mess up loading
+			"%ADB% shell rm -r /mnt/sdcard/obb/" + PackageName,
 			"%ADB% push " + Path.GetFileName(LocalObbName) + " " + DeviceObbName,
 		};
 		File.WriteAllLines(BatchName, BatchLines);
@@ -217,8 +217,8 @@ public class AndroidPlatform : Platform
 			string UE4GameRemoteDir = "/mnt/sdcard/" + Params.ShortProjectName;
 
 			// make sure device is at a clean state
-			Run(CmdEnv.CmdExe, AdbCommand + "shell rm -rf " + RemoteDir);
-			Run(CmdEnv.CmdExe, AdbCommand + "shell rm -rf " + UE4GameRemoteDir);
+			Run(CmdEnv.CmdExe, AdbCommand + "shell rm -r " + RemoteDir);
+			Run(CmdEnv.CmdExe, AdbCommand + "shell rm -r " + UE4GameRemoteDir);
 
 			string[] Files = Directory.GetFiles(SC.StageDirectory, "*", SearchOption.AllDirectories);
 			// copy each UFS file
@@ -244,7 +244,7 @@ public class AndroidPlatform : Platform
 			}
 
 			// delete the .obb file, since it will cause nothing we just deployed to be used
-			Run(CmdEnv.CmdExe, AdbCommand + "shell rm -f " + DeviceObbName);
+			Run(CmdEnv.CmdExe, AdbCommand + "shell rm " + DeviceObbName);
 		}
 		else
 		{
