Licensee Engine Version (Changelist) support. Modified UpdateLocalVersion UAT script to be able to set licensee versions and custom changelists.

diff --git a/Engine/Source/Programs/AutomationTool/Scripts/UpdateLocalVersion.Automation.cs b/Engine/Source/Programs/AutomationTool/Scripts/UpdateLocalVersion.Automation.cs
--- a/Engine/Source/Programs/AutomationTool/Scripts/UpdateLocalVersion.Automation.cs
+++ b/Engine/Source/Programs/AutomationTool/Scripts/UpdateLocalVersion.Automation.cs
@@ -8,12 +8,15 @@ using AutomationTool;
 using UnrealBuildTool;
 
 [Help("Updates your local versions based on your P4 sync")]
+[Help("CL", "Overrides the automatically disovered changelist number with the specified one")]
+[Help("Licensee", "When updating version files, store the changelist number in licensee format")]
 [RequireP4]
 public class UpdateLocalVersion : BuildCommand
 {
 	public override void ExecuteBuild()
 	{
 		var UE4Build = new UE4Build(this);
-		UE4Build.UpdateVersionFiles();
+		var ChangelistOverride = ParseParamInt("cl", -1);
+		UE4Build.UpdateVersionFiles(ChangelistNumberOverride: (ChangelistOverride < 0 ? (int?)null : ChangelistOverride));
 	}
 }
diff --git a/Engine/Source/Programs/AutomationTool/UE4Build.cs b/Engine/Source/Programs/AutomationTool/UE4Build.cs
--- a/Engine/Source/Programs/AutomationTool/UE4Build.cs
+++ b/Engine/Source/Programs/AutomationTool/UE4Build.cs
@@ -12,6 +12,7 @@ namespace AutomationTool
 	[Help("NoXGE", "Toggle to disable the distributed build process")]
 	[Help("ForceNonUnity", "Toggle to disable the unity build system")]
 	[Help("ForceUnity", "Toggle to force enable the unity build system")]
+	[Help("Licensee", "If set, this build is being compiled by a licensee")]
 	public class UE4Build : CommandUtils
 	{
 		private BuildCommand OwnerCommand;
@@ -415,23 +416,33 @@ namespace AutomationTool
 		/// <summary>
 		/// Updates the engine version files
 		/// </summary>
-		public List<string> UpdateVersionFiles(bool ActuallyUpdateVersionFiles = true)
+		public List<string> UpdateVersionFiles(bool ActuallyUpdateVersionFiles = true, int? ChangelistNumberOverride = null)
 		{
+			bool bIsLicenseeVersion = OwnerCommand.ParseParam("Licensee");
+			bool bDoUpdateVersionFiles = CommandUtils.P4Enabled && ActuallyUpdateVersionFiles;		
+			int ChangelistNumber = 0;
+			string ChangelistString = String.Empty;
+			if (bDoUpdateVersionFiles)
+			{
+				ChangelistNumber = ChangelistNumberOverride.GetValueOrDefault(P4Env.Changelist);
+				ChangelistString = ChangelistNumber.ToString();
+			}
+
 			var Result = new List<String>();
 			string Branch = P4Enabled ? P4Env.BuildRootEscaped : "";
 			{
 				string VerFile = CombinePaths(CmdEnv.LocalRoot, "Engine", "Build", "build.properties");
-                if (CommandUtils.P4Enabled && ActuallyUpdateVersionFiles)
+				if (bDoUpdateVersionFiles)
 				{
 					Log("Updating {0} with:", VerFile);
 					Log("  TimestampForBVT={0}", CmdEnv.TimestampAsString);
-					Log("  EngineVersion={0}", P4Env.ChangelistString);
-					Log("  BranchName={0}", Branch);
+					Log("  EngineVersion={0}", ChangelistNumber.ToString());
+					Log("  BranchName={0}", Branch);					
 
 					PrepareBuildProduct(VerFile, false);
 					VersionFileUpdater BuildProperties = new VersionFileUpdater(VerFile);
 					BuildProperties.ReplaceLine("TimestampForBVT=", CmdEnv.TimestampAsString);
-					BuildProperties.ReplaceLine("EngineVersion=", P4Env.ChangelistString);
+					BuildProperties.ReplaceLine("EngineVersion=", ChangelistNumber.ToString());
 					BuildProperties.ReplaceLine("BranchName=", Branch);
                     BuildProperties.Commit();
 				}
@@ -449,13 +460,13 @@ namespace AutomationTool
 
 			{
 				string VerFile = CombinePaths(CmdEnv.LocalRoot, "Engine", "Source", "Runtime", "Core", "Private", "UObject", "ObjectVersion.cpp");
-                if (CommandUtils.P4Enabled && ActuallyUpdateVersionFiles)
+				if (bDoUpdateVersionFiles)
 				{
 					Log("Updating {0} with:", VerFile);
-					Log(" #define	ENGINE_VERSION  {0}", P4Env.ChangelistString);
+					Log(" #define	ENGINE_VERSION  {0}", ChangelistNumber.ToString());
 
 					VersionFileUpdater ObjectVersionCpp = new VersionFileUpdater(VerFile);
-					ObjectVersionCpp.ReplaceLine("#define	ENGINE_VERSION	", P4Env.ChangelistString);
+					ObjectVersionCpp.ReplaceLine("#define	ENGINE_VERSION	", ChangelistNumber.ToString());
                     ObjectVersionCpp.Commit();
                 }
 				else
@@ -468,21 +479,23 @@ namespace AutomationTool
             string EngineVersionFile = CombinePaths(CmdEnv.LocalRoot, "Engine", "Source", "Runtime", "Launch", "Resources", "Version.h");
             {
                 string VerFile = EngineVersionFile;
-                if (CommandUtils.P4Enabled && ActuallyUpdateVersionFiles)
+				if (bDoUpdateVersionFiles)
 				{
 					Log("Updating {0} with:", VerFile);
-					Log(" #define	ENGINE_VERSION  {0}", P4Env.ChangelistString);
-					Log(" #define	ENGINE_VERSION_HIWORD  {0}", (P4Env.Changelist >> 16).ToString());
-					Log(" #define	ENGINE_VERSION_LOWORD  {0}", (P4Env.Changelist & 0xFFFF).ToString());
+					Log(" #define	ENGINE_VERSION  {0}", ChangelistNumber.ToString());
+					Log(" #define	ENGINE_VERSION_HIWORD  {0}", (ChangelistNumber >> 16).ToString());
+					Log(" #define	ENGINE_VERSION_LOWORD  {0}", (ChangelistNumber & 0xFFFF).ToString());					
 					Log(" #define	BRANCH_NAME  {0}", Branch);
-					Log(" #define	BUILT_FROM_CHANGELIST  {0}", P4Env.ChangelistString);
+					Log(" #define	BUILT_FROM_CHANGELIST  {0}", ChangelistString);
+					Log(" #define   ENGINE_IS_LICENSEE_VERSION  {0}", bIsLicenseeVersion ? "1" : "0");
 
 					VersionFileUpdater VersionH = new VersionFileUpdater(VerFile);
-					VersionH.ReplaceLine("#define ENGINE_VERSION ", P4Env.ChangelistString);
-					VersionH.ReplaceLine("#define ENGINE_VERSION_HIWORD ", (P4Env.Changelist >> 16).ToString());
-					VersionH.ReplaceLine("#define ENGINE_VERSION_LOWORD ", (P4Env.Changelist & 0xFFFF).ToString());
+					VersionH.ReplaceLine("#define ENGINE_VERSION ", ChangelistNumber.ToString());
+					VersionH.ReplaceLine("#define ENGINE_VERSION_HIWORD ", (ChangelistNumber >> 16).ToString());
+					VersionH.ReplaceLine("#define ENGINE_VERSION_LOWORD ", (ChangelistNumber & 0xFFFF).ToString());					
 					VersionH.ReplaceLine("#define BRANCH_NAME ", "\"" + Branch + "\"");
-					VersionH.ReplaceLine("#define BUILT_FROM_CHANGELIST ", P4Env.ChangelistString);
+					VersionH.ReplaceLine("#define BUILT_FROM_CHANGELIST ", ChangelistString);
+					VersionH.ReplaceOrAddLine("#define ENGINE_IS_LICENSEE_VERSION ", bIsLicenseeVersion ? "1" : "0");
 
                     VersionH.Commit();
                 }
@@ -497,10 +510,10 @@ namespace AutomationTool
                 // Use Version.h data to update MetaData.cs so the assemblies match the engine version.
                 string VerFile = CombinePaths(CmdEnv.LocalRoot, "Engine", "Source", "Programs", "DotNETCommon", "MetaData.cs");
 
-                if (CommandUtils.P4Enabled && ActuallyUpdateVersionFiles)
+				if (bDoUpdateVersionFiles)
                 {
                     // Get the MAJOR/MINOR/PATCH from the Engine Version file, as it is authoritative. The rest we get from the P4Env.
-                    string NewInformationalVersion = FEngineVersionSupport.FromVersionFile(EngineVersionFile).ToString();
+                    string NewInformationalVersion = FEngineVersionSupport.FromVersionFile(EngineVersionFile, ChangelistNumber).ToString();
 
                     Log("Updating {0} with AssemblyInformationalVersion: {1}", VerFile, NewInformationalVersion);
 
diff --git a/Engine/Source/Programs/AutomationTool/Utils.cs b/Engine/Source/Programs/AutomationTool/Utils.cs
--- a/Engine/Source/Programs/AutomationTool/Utils.cs
+++ b/Engine/Source/Programs/AutomationTool/Utils.cs
@@ -718,25 +718,28 @@ namespace AutomationTool
         /// <summary>
         /// Ctor that takes a pre-determined Version. Gets the Changelist and BranchName from the current <see cref="CommandUtils.P4Env"/>.
         /// </summary>
-        /// <param name="version">Predetermined version.</param>
-        public FEngineVersionSupport(Version version)
+		/// <param name="InVersion">Predetermined version.</param>
+		/// <param name="InChangelist">Predetermined changelist (optional)</param>
+		/// <param name="InBranchName">Predetermined branch name (optional)</param>
+		public FEngineVersionSupport(Version InVersion, int InChangelist = -1, string InBranchName = null)
         {
-            Version = version;
-            Changelist = CommandUtils.P4Enabled ? CommandUtils.P4Env.Changelist : 0;
-            BranchName = CommandUtils.P4Enabled ? CommandUtils.P4Env.BuildRootEscaped : "UnknownBranch";
-        }
-
-        /// <summary>
-        /// Ctor that takes a pre-determined Version.
-        /// </summary>
-        /// <param name="version">Predetermined version.</param>
-        /// <param name="changelist">Predetermined changelist</param>
-        /// <param name="changelist">Predetermined branch name</param>
-        public FEngineVersionSupport(Version version, int changelist, string branchName)
-        {
-            Version = version;
-            Changelist = changelist;
-            BranchName = branchName;
+			Version = InVersion;
+			if (InChangelist <= 0)
+			{
+				Changelist = CommandUtils.P4Enabled ? CommandUtils.P4Env.Changelist : 0;
+			}
+			else
+			{
+				Changelist = InChangelist;
+			}
+			if (String.IsNullOrEmpty(InBranchName))
+			{
+				BranchName = CommandUtils.P4Enabled ? CommandUtils.P4Env.BuildRootEscaped : "UnknownBranch";
+			}
+			else
+			{
+				BranchName = InBranchName;
+			}
         }
 
         /// <summary>
@@ -760,9 +763,11 @@ namespace AutomationTool
         /// Ctor initializes with the values from the supplied Version file. The BranchName and CL are also taken from the current <see cref="CommandUtils.P4Env"/>.
         /// </summary>
         /// <param name="Filename">Full path to the file with the version info.</param>
-        public static FEngineVersionSupport FromVersionFile(string Filename)
+		/// <param name="InChangelist">Predetermined changelist (optional)</param>
+		/// <param name="InBranchName">Predetermined branch name (optional)</param>
+        public static FEngineVersionSupport FromVersionFile(string Filename, int InChangelist = -1, string InBranchName = null)
         {
-            return new FEngineVersionSupport(ReadVersionFromFile(Filename));
+			return new FEngineVersionSupport(ReadVersionFromFile(Filename), InChangelist, InBranchName);
         }
 
         /// <summary>
@@ -831,8 +836,8 @@ namespace AutomationTool
 		public VersionFileUpdater(string Filename)
 		{
 			MyFile = new FileInfo(Filename);
-			Lines = InternalUtils.SafeReadAllLines(Filename);
-            if (Lines == null || Lines.Length == 0)
+			Lines = new List<string>(InternalUtils.SafeReadAllLines(Filename));
+            if (CommandUtils.IsNullOrEmpty(Lines))
             {
                 throw new AutomationException("Version file {0} was empty or not found!", Filename);
             }
@@ -843,7 +848,7 @@ namespace AutomationTool
 		/// </summary>
 		public void ReplaceLine(string StartOfLine, string ReplacementRHS)
 		{
-			for (int Index = 0; Index < Lines.Length; ++Index)
+			for (int Index = 0; Index < Lines.Count; ++Index)
 			{
 				if (Lines[Index].StartsWith(StartOfLine))
 				{
@@ -854,6 +859,31 @@ namespace AutomationTool
             throw new AutomationException("Unable to find line {0} in {1}", StartOfLine, MyFile.FullName);
 		}
 
+		/// <summary>
+		/// Doc
+		/// </summary>
+		public void ReplaceOrAddLine(string StartOfLine, string ReplacementRHS)
+		{
+			if (Contains(StartOfLine))
+			{
+				ReplaceLine(StartOfLine, ReplacementRHS);
+			}
+			else
+			{
+				AddLine("");
+				AddLine(StartOfLine + ReplacementRHS);
+			}
+		}
+
+		/// <summary>
+		/// Adds a new line to the version file
+		/// </summary>
+		/// <param name="Line"></param>
+		public void AddLine(string Line)
+		{
+			Lines.Add(Line);
+		}
+
 		/// <summary>
 		/// Doc
 		/// </summary>
@@ -862,7 +892,7 @@ namespace AutomationTool
             // This searches for the AssemblyInformationalVersion string. Most the mess is to allow whitespace in places that are possible.
             // Captures the string into a group called "Ver" for replacement.
             var regex = new Regex(@"\[assembly:\s+AssemblyInformationalVersion\s*\(\s*""(?<Ver>.*)""\s*\)\s*]", RegexOptions.IgnoreCase | RegexOptions.ExplicitCapture);
-            foreach (var Index in Enumerable.Range(0, Lines.Length))
+            foreach (var Index in Enumerable.Range(0, Lines.Count))
             {
                 var line = Lines[Index];
                 var match = regex.Match(line);
@@ -885,7 +915,7 @@ namespace AutomationTool
 		public void Commit()
 		{
 			MyFile.IsReadOnly = false;
-			if (!InternalUtils.SafeWriteAllLines(MyFile.FullName, Lines))
+			if (!InternalUtils.SafeWriteAllLines(MyFile.FullName, Lines.ToArray()))
 			{
 				throw new AutomationException("Unable to update version info in {0}", MyFile.FullName);
 			}
@@ -915,7 +945,7 @@ namespace AutomationTool
 		/// <summary>
 		/// Doc
 		/// </summary>
-		protected string[] Lines;
+		protected List<string> Lines;
 	}
 
 	#endregion
diff --git a/Engine/Source/Runtime/Core/Private/Misc/EngineVersion.cpp b/Engine/Source/Runtime/Core/Private/Misc/EngineVersion.cpp
--- a/Engine/Source/Runtime/Core/Private/Misc/EngineVersion.cpp
+++ b/Engine/Source/Runtime/Core/Private/Misc/EngineVersion.cpp
@@ -3,6 +3,7 @@
 #include "CorePrivate.h"
 #include "EngineVersion.h"
 #include "Runtime/Launch/Resources/Version.h"
+#include "EngineBuildSettings.h"
 
 FEngineVersion::FEngineVersion()
 {
@@ -25,7 +26,14 @@ void FEngineVersion::Set(uint16 InMajor, uint16 InMinor, uint16 InPatch, uint32
 
 uint32 FEngineVersion::GetChangelist() const
 {
-	return Changelist;
+	// Mask to ignore licensee bit
+	return Changelist & (uint32)0x7fffffffU;
+}
+
+bool FEngineVersion::IsLicenseeVersion() const
+{
+	// Check for licensee bit
+	return (Changelist & (uint32)0x80000000U) != 0;
 }
 
 void FEngineVersion::Empty()
@@ -40,7 +48,7 @@ bool FEngineVersion::IsEmpty() const
 
 bool FEngineVersion::IsPromotedBuild() const
 {
-	return Changelist != 0;
+	return GetChangelist() != 0;
 }
 
 bool FEngineVersion::IsCompatibleWith(const FEngineVersion &Other) const
@@ -83,7 +91,7 @@ FString FEngineVersion::ToString(EVersionComponent::Type LastComponent) const
 			Result += FString::Printf(TEXT(".%u"), Patch);
 			if(LastComponent >= EVersionComponent::Changelist)
 			{
-				Result += FString::Printf(TEXT("-%u"), Changelist);
+				Result += FString::Printf(TEXT("-%u"), GetChangelist());
 				if(LastComponent >= EVersionComponent::Branch && Branch.Len() > 0)
 				{
 					Result += FString::Printf(TEXT("+%s"), *Branch);
@@ -154,11 +162,11 @@ EVersionComparison::Type FEngineVersion::GetNewest(const FEngineVersion &First,
 		return (First.Patch > Second.Patch)? EVersionComparison::First : EVersionComparison::Second;
 	}
 
-	// Compare changelist versions
-	if(First.Changelist != Second.Changelist)
+	// Compare changelists (only if they're both from the same vendor)
+	if(First.IsLicenseeVersion() == Second.IsLicenseeVersion() && First.GetChangelist() != Second.GetChangelist())
 	{
 		*OutComponent = EVersionComponent::Changelist;
-		return (First.Changelist > Second.Changelist)? EVersionComparison::First : EVersionComparison::Second;
+		return (First.GetChangelist() > Second.GetChangelist()) ? EVersionComparison::First : EVersionComparison::Second;
 	}
 
 	// Otherwise they're the same
@@ -174,5 +182,12 @@ void operator<<(FArchive &Ar, FEngineVersion &Version)
 	Ar << Version.Branch;
 }
 
+// Make sure this is defined
+#ifndef ENGINE_IS_LICENSEE_VERSION
+	#define ENGINE_IS_LICENSEE_VERSION 0
+#endif
+
+// Licensee changelists part of the engine version has the top bit set to 1
+#define ENGINE_VERSION_INTERNAL_OR_LICENSEE (BUILT_FROM_CHANGELIST | (ENGINE_IS_LICENSEE_VERSION << 31))
 // Global instance of the current engine version
-const FEngineVersion GEngineVersion(ENGINE_MAJOR_VERSION, ENGINE_MINOR_VERSION, ENGINE_PATCH_VERSION, BUILT_FROM_CHANGELIST, BRANCH_NAME);
+const FEngineVersion GEngineVersion(ENGINE_MAJOR_VERSION, ENGINE_MINOR_VERSION, ENGINE_PATCH_VERSION, ENGINE_VERSION_INTERNAL_OR_LICENSEE, BRANCH_NAME);
diff --git a/Engine/Source/Runtime/Core/Public/Misc/EngineVersion.h b/Engine/Source/Runtime/Core/Public/Misc/EngineVersion.h
--- a/Engine/Source/Runtime/Core/Public/Misc/EngineVersion.h
+++ b/Engine/Source/Runtime/Core/Public/Misc/EngineVersion.h
@@ -44,7 +44,10 @@ public:
 	void Set(uint16 InMajor, uint16 InMinor, uint16 InPatch, uint32 InChangelist, const FString &InBranch);
 
 	/** Returns the changelist number corresponding to this version. */
-	uint32 GetChangelist() const;
+	uint32 GetChangelist() const;	
+
+	/** Checks if the changelist number represents licensee changelist number. */
+	bool IsLicenseeVersion() const;
 
 	/** Clears the object. */
 	void Empty();
@@ -80,7 +83,7 @@ private:
 	/** Patch version number. */
 	uint16 Patch;
 
-	/** Changelist number. This is used to arbitrate when Major/Minor/Patch version numbers match. */
+	/** Changelist number. This is used to arbitrate when Major/Minor/Patch version numbers match. Use GetChangelist() instead of using this member directly. */
 	uint32 Changelist;
 
 	/** Branch name. */
diff --git a/Engine/Source/Runtime/Launch/Private/LaunchEngineLoop.cpp b/Engine/Source/Runtime/Launch/Private/LaunchEngineLoop.cpp
--- a/Engine/Source/Runtime/Launch/Private/LaunchEngineLoop.cpp
+++ b/Engine/Source/Runtime/Launch/Private/LaunchEngineLoop.cpp
@@ -2438,7 +2438,7 @@ void FEngineLoop::AppInit( )
 	check(sizeof(bool) == 1);
 	check(sizeof(float) == 4);
 	check(sizeof(double) == 8);
-	check(GEngineNetVersion == 0 || GEngineNetVersion >= GEngineMinNetVersion);
+	check(GEngineNetVersion == 0 || GEngineNetVersion >= GEngineMinNetVersion || GEngineVersion.IsLicenseeVersion());
 
 	// Culture.
 	TCHAR CookerCulture[8];
