build script, add a note to the failed emails

diff --git a/Engine/Source/Programs/AutomationTool/Scripts/GUBP.Automation.cs b/Engine/Source/Programs/AutomationTool/Scripts/GUBP.Automation.cs
--- a/Engine/Source/Programs/AutomationTool/Scripts/GUBP.Automation.cs
+++ b/Engine/Source/Programs/AutomationTool/Scripts/GUBP.Automation.cs
@@ -3611,6 +3611,7 @@ public class GUBP : BuildCommand
         var ECProps = new List<string>();
         EMails = "";
         string FailCauserEMails = "";
+        string EMailNote = "";
         bool SendSuccessForGreenAfterRed = false;
         if (GUBPNodesHistory.ContainsKey(NodeToDo))
         {
@@ -3631,12 +3632,23 @@ public class GUBP : BuildCommand
                     if (++NumChanges > 50)
                     {
                         FailCauserEMails = "";
+                        EMailNote = String.Format("This step has been broken for more than 50 changes. It last succeeded at CL {0}. ", History.LastSucceeded);
                         break;
                     }
                 }
             }
-            ECProps.Add(string.Format("FailCausers/{0}={1}", NodeToDo, FailCauserEMails));
-
+            else if (History.LastSucceeded <= 0)
+            {
+                EMailNote = String.Format("This step has been broken for more than a few days, so there is no record of it ever succeeding. ");
+            }
+            if (EMailNote != "" && !String.IsNullOrEmpty(History.FailedString))
+            {
+                EMailNote += String.Format("It has failed at CLs {0}. ", History.FailedString);
+            }
+            if (EMailNote != "" && !String.IsNullOrEmpty(History.InProgressString))
+            {
+                EMailNote += String.Format("These CLs are being built right now {0}. ", History.InProgressString);
+            }
             if (History.LastSucceeded > 0 && History.LastSucceeded < P4Env.Changelist && History.LastFailed > History.LastSucceeded && History.LastFailed < P4Env.Changelist)
             {
                 SendSuccessForGreenAfterRed = ParseParam("CIS");
@@ -3648,6 +3660,10 @@ public class GUBP : BuildCommand
             ECProps.Add(string.Format("RedsSince/{0}=", NodeToDo));
             ECProps.Add(string.Format("InProgress/{0}=", NodeToDo));
         }
+
+        ECProps.Add(string.Format("FailCausers/{0}={1}", NodeToDo, FailCauserEMails));
+        ECProps.Add(string.Format("EmailNotes/{0}={1}", NodeToDo, EMailNote));
+
         {
             var AdditonalEmails = "";
             if (ParseParam("CIS") && !GUBPNodes[NodeToDo].SendSuccessEmail())
