Build/cook/run fix when "FrameworkAssets" folder doesn't exist

diff --git a/Engine/Source/Programs/AutomationTool/IOS/IOSPlatform.Automation.cs b/Engine/Source/Programs/AutomationTool/IOS/IOSPlatform.Automation.cs
--- a/Engine/Source/Programs/AutomationTool/IOS/IOSPlatform.Automation.cs
+++ b/Engine/Source/Programs/AutomationTool/IOS/IOSPlatform.Automation.cs
@@ -478,7 +478,10 @@ public class IOSPlatform : Platform
 			// copy any additional framework assets that will be needed at runtime
 			{
 				string SourcePath = CombinePaths( SC.LocalRoot, "Engine", "Build", "IOS", "Resources", "FrameworkAssets" );
-				SC.StageFiles( StagedFileType.NonUFS, SourcePath, "*.*", true, null, "", true, false );
+				if ( Directory.Exists( SourcePath ) )
+				{
+					SC.StageFiles( StagedFileType.NonUFS, SourcePath, "*.*", true, null, "", true, false );
+				}
 			}
 
 			// copy the icons/launch screens from the game (may stomp the engine copies)
diff --git a/Engine/Source/Programs/UnrealBuildTool/IOS/UEDeployIOS.cs b/Engine/Source/Programs/UnrealBuildTool/IOS/UEDeployIOS.cs
--- a/Engine/Source/Programs/UnrealBuildTool/IOS/UEDeployIOS.cs
+++ b/Engine/Source/Programs/UnrealBuildTool/IOS/UEDeployIOS.cs
@@ -169,7 +169,10 @@ namespace UnrealBuildTool.IOS
 			}
 			
 			// copy additional engine framework assets in
-			CopyFiles( InEngineDir + "/Build/IOS/Resources/FrameworkAssets", AppDirectory, "*.*", true );
+			if ( Directory.Exists( InEngineDir + "/Build/IOS/Resources/FrameworkAssets" ) )
+			{
+				CopyFiles( InEngineDir + "/Build/IOS/Resources/FrameworkAssets", AppDirectory, "*.*", true );
+			}
 
 			//CopyFiles(BuildDirectory, PayloadDirectory, null, "iTunesArtwork", null);
 			CopyFolder(InEngineDir + "/Content/Stats", AppDirectory + "/cookeddata/engine/content/stats", true);
