build script, add formal builds for code projects

diff --git a/Engine/Source/Programs/AutomationTool/Scripts/GUBP.Automation.cs b/Engine/Source/Programs/AutomationTool/Scripts/GUBP.Automation.cs
--- a/Engine/Source/Programs/AutomationTool/Scripts/GUBP.Automation.cs
+++ b/Engine/Source/Programs/AutomationTool/Scripts/GUBP.Automation.cs
@@ -2083,35 +2083,6 @@ public class GUBP : BuildCommand
         }
     }
 
-
-    //@todo copy-paste from temp storage, unify
-    public static string ResolveSharedBuildDirectory(string GameFolder)
-    {
-        string Root = RootSharedTempStorageDirectory();
-        string Result = CombinePaths(Root, GameFolder);
-        if (!DirectoryExists_NoExceptions(Result))
-        {
-            string GameStr = "Game";
-            bool HadGame = false;
-            if (GameFolder.EndsWith(GameStr, StringComparison.InvariantCultureIgnoreCase))
-            {
-                string ShortFolder = GameFolder.Substring(0, GameFolder.Length - GameStr.Length);
-                Result = CombinePaths(Root, ShortFolder);
-                HadGame = true;
-            }
-            if (!HadGame || !DirectoryExists_NoExceptions(Result))
-            {
-                Result = CombinePaths(Root, "UE4");
-                if (!DirectoryExists_NoExceptions(Result))
-                {
-                    throw new AutomationException("Could not find an appropriate shared temp folder {0}", Result);
-                }
-            }
-        }
-        return Result;
-    }
-
-
     public class GamePlatformCookedAndCompiledNode : HostPlatformAggregateNode
     {
         BranchInfo.BranchUProject GameProj;
@@ -2197,14 +2168,17 @@ public class GUBP : BuildCommand
         List<UnrealTargetPlatform> ServerTargetPlatforms;
         List<UnrealTargetConfiguration> ClientConfigs;
         List<UnrealTargetConfiguration> ServerConfigs;
+        bool ClientNotGame;
+        UnrealBuildTool.TargetRules.TargetType GameOrClient;
 
         public FormalBuildNode(GUBP bp,
             BranchInfo.BranchUProject InGameProj,
             UnrealTargetPlatform InHostPlatform,
-            List<UnrealTargetPlatform> InClientTargetPlatforms,
+            List<UnrealTargetPlatform> InClientTargetPlatforms = null,
             List<UnrealTargetConfiguration> InClientConfigs = null,
             List<UnrealTargetPlatform> InServerTargetPlatforms = null,
-            List<UnrealTargetConfiguration> InServerConfigs = null
+            List<UnrealTargetConfiguration> InServerConfigs = null,
+            bool InClientNotGame = false
 
             )
             : base(InHostPlatform)
@@ -2214,33 +2188,15 @@ public class GUBP : BuildCommand
             ServerTargetPlatforms = InServerTargetPlatforms;
             ClientConfigs = InClientConfigs;
             ServerConfigs = InServerConfigs;
+            ClientNotGame = InClientNotGame;
 
-            if (ClientConfigs == null)
-            {
-                ClientConfigs = new List<UnrealTargetConfiguration>() { UnrealTargetConfiguration.Development };
-            }
-
-            var AllTargetPlatforms = new List<UnrealTargetPlatform>(ClientTargetPlatforms);
-            if (ServerTargetPlatforms != null)
-            {
-                foreach (var Plat in ServerTargetPlatforms)
-                {
-                    if (!AllTargetPlatforms.Contains(Plat))
-                    {
-                        AllTargetPlatforms.Add(Plat);
-                    }
-                }
-                if (ServerConfigs == null)
-                {
-                    ServerConfigs = new List<UnrealTargetConfiguration>() { UnrealTargetConfiguration.Development };
-                }
-            }
+            GameOrClient = TargetRules.TargetType.Game;
 
-            // add dependencies for cooked and compiled
-            foreach (var Plat in AllTargetPlatforms)
+            if (ClientNotGame)
             {
-                AddDependency(GamePlatformCookedAndCompiledNode.StaticGetFullName(HostPlatform, GameProj, Plat));
+                GameOrClient = TargetRules.TargetType.Client;
             }
+            
 
             // verify we actually built these
             var WorkingGameProject = InGameProj;
@@ -2250,27 +2206,57 @@ public class GUBP : BuildCommand
                 WorkingGameProject = bp.Branch.BaseEngineProject;
             }
 
-            if (!WorkingGameProject.Properties.Targets.ContainsKey(TargetRules.TargetType.Game))
+            var AllTargetPlatforms = new List<UnrealTargetPlatform>();
+            if (ClientTargetPlatforms != null)
             {
-                throw new AutomationException("Can't make a game build for {0} because it doesn't have a game target.", WorkingGameProject.GameName);
-            }
-            //@todo this doesn't support client only 
-            foreach (var Plat in ClientTargetPlatforms)
-            {
-                if (!WorkingGameProject.Properties.Targets[TargetRules.TargetType.Game].Rules.GUBP_GetPlatforms_MonolithicOnly(HostPlatform).Contains(Plat))
+                if (!WorkingGameProject.Properties.Targets.ContainsKey(GameOrClient))
                 {
-                    throw new AutomationException("Can't make a game build for {0} because we didn't build platform {1}.", WorkingGameProject.GameName, Plat.ToString());
+                    throw new AutomationException("Can't make a game build for {0} because it doesn't have a {1} target.", WorkingGameProject.GameName, GameOrClient.ToString());
                 }
-                foreach (var Config in ClientConfigs)
+
+                foreach (var Plat in ClientTargetPlatforms)
+                {
+                    if (!AllTargetPlatforms.Contains(Plat))
+                    {
+                        AllTargetPlatforms.Add(Plat);
+                    }
+                }
+                if (ClientConfigs == null)
                 {
-                    if (!WorkingGameProject.Properties.Targets[TargetRules.TargetType.Game].Rules.GUBP_GetConfigs_MonolithicOnly(HostPlatform, Plat).Contains(Config))
+                    ClientConfigs = new List<UnrealTargetConfiguration>() { UnrealTargetConfiguration.Development };
+                }
+                foreach (var Plat in ClientTargetPlatforms)
+                {
+                    if (!WorkingGameProject.Properties.Targets[GameOrClient].Rules.GUBP_GetPlatforms_MonolithicOnly(HostPlatform).Contains(Plat))
                     {
-                        throw new AutomationException("Can't make a game build for {0} because we didn't build platform {1} config {2}.", WorkingGameProject.GameName, Plat.ToString(), Config.ToString());
+                        throw new AutomationException("Can't make a game/client build for {0} because we didn't build platform {1}.", WorkingGameProject.GameName, Plat.ToString());
+                    }
+                    foreach (var Config in ClientConfigs)
+                    {
+                        if (!WorkingGameProject.Properties.Targets[GameOrClient].Rules.GUBP_GetConfigs_MonolithicOnly(HostPlatform, Plat).Contains(Config))
+                        {
+                            throw new AutomationException("Can't make a game/client build for {0} because we didn't build platform {1} config {2}.", WorkingGameProject.GameName, Plat.ToString(), Config.ToString());
+                        }
                     }
                 }
             }
             if (ServerTargetPlatforms != null)
             {
+                if (!WorkingGameProject.Properties.Targets.ContainsKey(TargetRules.TargetType.Server) && ServerTargetPlatforms != null)
+                {
+                    throw new AutomationException("Can't make a server build for {0} because it doesn't have a server target.", WorkingGameProject.GameName);
+                }
+                foreach (var Plat in ServerTargetPlatforms)
+                {
+                    if (!AllTargetPlatforms.Contains(Plat))
+                    {
+                        AllTargetPlatforms.Add(Plat);
+                    }
+                }
+                if (ServerConfigs == null)
+                {
+                    ServerConfigs = new List<UnrealTargetConfiguration>() { UnrealTargetConfiguration.Development };
+                }
                 foreach (var Plat in ServerTargetPlatforms)
                 {
                     if (!WorkingGameProject.Properties.Targets[TargetRules.TargetType.Server].Rules.GUBP_GetPlatforms_MonolithicOnly(HostPlatform).Contains(Plat))
@@ -2286,24 +2272,52 @@ public class GUBP : BuildCommand
                     }
                 }
             }
+
+            // add dependencies for cooked and compiled
+            foreach (var Plat in AllTargetPlatforms)
+            {
+                AddDependency(GamePlatformCookedAndCompiledNode.StaticGetFullName(HostPlatform, GameProj, Plat));
+            }
         }
 
-        public static string StaticGetFullName(BranchInfo.BranchUProject InGameProj, UnrealTargetPlatform InHostPlatform, List<UnrealTargetPlatform> InClientTargetPlatforms, List<UnrealTargetConfiguration> InClientConfigs)
+        public static string StaticGetFullName(BranchInfo.BranchUProject InGameProj, UnrealTargetPlatform InHostPlatform, List<UnrealTargetPlatform> InClientTargetPlatforms = null, List<UnrealTargetConfiguration> InClientConfigs = null, List<UnrealTargetPlatform> InServerTargetPlatforms = null, List<UnrealTargetConfiguration> InServerConfigs = null, bool InClientNotGame = false)
         {
             string Infix = "";
-            if (InClientTargetPlatforms.Count == 1)
+            if (InClientNotGame)
+            {
+                if (InClientTargetPlatforms != null && InClientTargetPlatforms.Count == 1)
+                {
+                    Infix = "_Client_" + InClientTargetPlatforms[0].ToString();
+                }
+                if (InClientConfigs != null && InClientConfigs.Count == 1)
+                {
+                    Infix += "_Client_" + InClientConfigs[0].ToString();
+                }
+            }
+            else
+            {
+                if (InClientTargetPlatforms != null && InClientTargetPlatforms.Count == 1)
+                {
+                    Infix = "_" + InClientTargetPlatforms[0].ToString();
+                }
+                if (InClientConfigs != null && InClientConfigs.Count == 1)
+                {
+                    Infix += "_" + InClientConfigs[0].ToString();
+                }
+            }
+            if (InServerTargetPlatforms != null && InServerTargetPlatforms.Count == 1)
             {
-                Infix = "_" + InClientTargetPlatforms[0].ToString();
+                Infix = "_Serv_" + InServerTargetPlatforms[0].ToString();
             }
-            if (InClientConfigs.Count == 1)
+            if (InServerConfigs != null && InServerConfigs.Count == 1)
             {
-                Infix += "_" + InClientConfigs[0].ToString();
+                Infix += "_Serv_" + InServerConfigs[0].ToString();
             }
             return InGameProj.GameName + Infix + "_MakeBuild" + HostPlatformNode.StaticGetHostPlatformSuffix(InHostPlatform);
         }
         public override string GetFullName()
         {
-            return StaticGetFullName(GameProj, HostPlatform, ClientTargetPlatforms, ClientConfigs);
+            return StaticGetFullName(GameProj, HostPlatform, ClientTargetPlatforms, ClientConfigs, ServerTargetPlatforms, ServerConfigs, ClientNotGame);
         }
         public override string GameNameIfAnyForTempStorage()
         {
@@ -2332,32 +2346,43 @@ public class GUBP : BuildCommand
             }
             string Args = String.Format("BuildCookRun{0} -SkipBuild -SkipCook -Stage -Pak -Package -NoSubmit", ProjectArg);
 
-            bool bFirstClient = true;
-            foreach (var Plat in ClientTargetPlatforms)
+            if (ClientTargetPlatforms != null)
             {
-                if (bFirstClient)
+                bool bFirstClient = true;
+                foreach (var Plat in ClientTargetPlatforms)
                 {
-                    bFirstClient = false;
-                    Args += String.Format(" -platform={0}", Plat.ToString());
-                }
-                else
-                {
-                    Args += String.Format("+{0}", Plat.ToString());
+                    if (bFirstClient)
+                    {
+                        bFirstClient = false;
+                        Args += String.Format(" -platform={0}", Plat.ToString());
+                    }
+                    else
+                    {
+                        Args += String.Format("+{0}", Plat.ToString());
+                    }
                 }
-            }
-            bool bFirstClientConfig = true;
-            foreach (var Config in ClientConfigs)
-            {
-                if (bFirstClientConfig)
+                bool bFirstClientConfig = true;
+                foreach (var Config in ClientConfigs)
                 {
-                    bFirstClientConfig = false;
-                    Args += String.Format(" -clientconfig={0}", Config.ToString());
+                    if (bFirstClientConfig)
+                    {
+                        bFirstClientConfig = false;
+                        Args += String.Format(" -clientconfig={0}", Config.ToString());
+                    }
+                    else
+                    {
+                        Args += String.Format("+{0}", Config.ToString());
+                    }
                 }
-                else
+                if (ClientNotGame)
                 {
-                    Args += String.Format("+{0}", Config.ToString());
+                    Args += " -client";
                 }
             }
+            else
+            {
+                Args += " -noclient";
+            }
             if (ServerTargetPlatforms != null)
             {
                 Args += " -server";
@@ -2404,7 +2429,6 @@ public class GUBP : BuildCommand
                 Args += String.Format(" -Archive -archivedirectory={0}", ArchiveDirectory);
             }
 
-
             string LogFile = CommandUtils.RunUAT(CommandUtils.CmdEnv, Args);
             SaveRecordOfSuccessAndAddToBuildProducts(CommandUtils.ReadAllText(LogFile));
         }
@@ -3532,7 +3556,35 @@ public class GUBP : BuildCommand
                 Log("Cycle in GUBP, could not resolve:");
                 foreach (var NodeToDo in NodesToDo)
                 {
-                    Log("  {0}", NodeToDo);
+                    string Deps = "";
+                    if (!SubSort && GUBPNodes[NodeToDo].AgentSharingGroup != "")
+                    {
+                        foreach (var ChainNode in SortedAgentGroupChains[GUBPNodes[NodeToDo].AgentSharingGroup])
+                        {
+                            foreach (var Dep in GUBPNodes[ChainNode].FullNamesOfDependencies)
+                            {
+                                if (!SortedAgentGroupChains[GUBPNodes[NodeToDo].AgentSharingGroup].Contains(Dep) && NodesToDo.Contains(Dep))
+                                {
+                                    Deps = Deps + Dep + "[" + ChainNode + "->" + GUBPNodes[NodeToDo].AgentSharingGroup + "]" + " ";
+                                }
+                            }
+                        }
+                    }
+                    foreach (var Dep in GUBPNodes[NodeToDo].FullNamesOfDependencies)
+                    {
+                        if (NodesToDo.Contains(Dep))
+                        {
+                            Deps = Deps + Dep + " ";
+                        }
+                    }
+                    foreach (var Dep in GUBPNodes[NodeToDo].FullNamesOfPseudosependencies)
+                    {
+                        if (NodesToDo.Contains(Dep))
+                        {
+                            Deps = Deps + Dep + " ";
+                        }
+                    }
+                    Log("  {0}    deps: {1}", NodeToDo, Deps);
                 }
                 throw new AutomationException("Cycle in GUBP");
             }
@@ -4340,7 +4392,7 @@ public class GUBP : BuildCommand
                                             {
                                                 if (PlatPair.Key == Plat)
                                                 {
-                                                    AddNode(new FormalBuildNode(this, NonCodeProject, HostPlatform, new List<UnrealTargetPlatform>() { Plat }, new List<UnrealTargetConfiguration>() {PlatPair.Value}));
+                                                    var NodeName = AddNode(new FormalBuildNode(this, NonCodeProject, HostPlatform, new List<UnrealTargetPlatform>() { Plat }, new List<UnrealTargetConfiguration>() {PlatPair.Value}));
                                                     // we don't want this delayed
                                                     // this would normally wait for the testing phase, we just want to build it right away
                                                     RemovePseudodependencyFromNode(
@@ -4350,14 +4402,14 @@ public class GUBP : BuildCommand
                                                         CookNode.StaticGetFullName(HostPlatform, NonCodeProject, CookedPlatform),
                                                         CookNode.StaticGetFullName(HostPlatform, Branch.BaseEngineProject, CookedPlatform));
 
-                                                    string BuildAgentSharingGroup = NonCodeProject.GameName + "_MakeFormalBuild_" + Plat.ToString();
+                                                    string BuildAgentSharingGroup = NonCodeProject.GameName + "_MakeFormalBuild_" + Plat.ToString() + HostPlatformNode.StaticGetHostPlatformSuffix(HostPlatform);
                                                     if (Plat == UnrealTargetPlatform.IOS) // Mac-IOS trashes build products, so we need to use different agents
                                                     {
                                                         BuildAgentSharingGroup = "";
                                                     }
                                                     GUBPNodes[CookNode.StaticGetFullName(HostPlatform, NonCodeProject, CookedPlatform)].AgentSharingGroup = BuildAgentSharingGroup;
                                                     GUBPNodes[GamePlatformCookedAndCompiledNode.StaticGetFullName(HostPlatform, NonCodeProject, Plat)].AgentSharingGroup = BuildAgentSharingGroup;
-                                                    GUBPNodes[FormalBuildNode.StaticGetFullName(NonCodeProject, HostPlatform, new List<UnrealTargetPlatform>() { Plat }, new List<UnrealTargetConfiguration>() { PlatPair.Value })].AgentSharingGroup = BuildAgentSharingGroup;
+                                                    GUBPNodes[NodeName].AgentSharingGroup = BuildAgentSharingGroup;
 
                                                 }
                                             }
@@ -4497,6 +4549,44 @@ public class GUBP : BuildCommand
                                 {
                                     AddNode(new GamePlatformCookedAndCompiledNode(this, HostPlatform, CodeProj, Plat, true));
                                 }
+#if true
+                                var FormalBuildConfigs = Target.Rules.GUBP_GetConfigsForFormalBuilds_MonolithicOnly(HostPlatform, Plat);
+
+                                foreach (var Config in FormalBuildConfigs)
+                                {
+                                    string FormalNodeName = null;
+                                    if (Kind == TargetRules.TargetType.Client)
+                                    {
+                                        FormalNodeName = AddNode(new FormalBuildNode(this, CodeProj, HostPlatform, InClientTargetPlatforms: new List<UnrealTargetPlatform>() { Plat }, InClientConfigs: new List<UnrealTargetConfiguration>() { Config }, InClientNotGame : true));
+                                    }
+                                    else if (Kind == TargetRules.TargetType.Server)
+                                    {
+                                        FormalNodeName = AddNode(new FormalBuildNode(this, CodeProj, HostPlatform, InServerTargetPlatforms : new List<UnrealTargetPlatform>() { Plat }, InServerConfigs : new List<UnrealTargetConfiguration>() { Config }));
+                                    }
+                                    else if (Kind == TargetRules.TargetType.Game)
+                                    {
+                                        FormalNodeName = AddNode(new FormalBuildNode(this, CodeProj, HostPlatform, InClientTargetPlatforms: new List<UnrealTargetPlatform>() { Plat }, InClientConfigs: new List<UnrealTargetConfiguration>() { Config }));
+                                    }
+                                    // we don't want this delayed
+                                    // this would normally wait for the testing phase, we just want to build it right away
+                                    RemovePseudodependencyFromNode(
+                                        CookNode.StaticGetFullName(HostPlatform, CodeProj, CookedPlatform),
+                                        WaitForTestShared.StaticGetFullName());
+                                    RemovePseudodependencyFromNode(
+                                        CookNode.StaticGetFullName(HostPlatform, CodeProj, CookedPlatform),
+                                        CookNode.StaticGetFullName(HostPlatform, Branch.BaseEngineProject, CookedPlatform));
+
+                                    string BuildAgentSharingGroup = CodeProj.GameName + "_MakeFormalBuild_" + Plat.ToString() + HostPlatformNode.StaticGetHostPlatformSuffix(HostPlatform);
+                                    if (Plat == UnrealTargetPlatform.IOS) // Mac-IOS trashes build products, so we need to use different agents
+                                    {
+                                        BuildAgentSharingGroup = "";
+                                    }
+                                    GUBPNodes[CookNode.StaticGetFullName(HostPlatform, CodeProj, CookedPlatform)].AgentSharingGroup = BuildAgentSharingGroup;
+                                    GUBPNodes[GamePlatformCookedAndCompiledNode.StaticGetFullName(HostPlatform, CodeProj, Plat)].AgentSharingGroup = BuildAgentSharingGroup;
+                                    GUBPNodes[FormalNodeName].AgentSharingGroup = BuildAgentSharingGroup;
+
+                                }
+#endif
                                 if (bAutomatedTesting)
                                 {
 if (HostPlatform == UnrealTargetPlatform.Mac) continue; //temp hack till mac automated testing works
diff --git a/Engine/Source/Programs/AutomationTool/TempStorage.cs b/Engine/Source/Programs/AutomationTool/TempStorage.cs
--- a/Engine/Source/Programs/AutomationTool/TempStorage.cs
+++ b/Engine/Source/Programs/AutomationTool/TempStorage.cs
@@ -545,28 +545,27 @@ namespace AutomationTool
         }
 
         static Dictionary<string, string> ResolveCache = new Dictionary<string, string>();
-        public static string ResolveSharedTempStorageDirectory(string GameFolder)
+        public static string ResolveSharedBuildDirectory(string GameFolder)
         {
             if (ResolveCache.ContainsKey(GameFolder))
             {
                 return ResolveCache[GameFolder];
             }
-            string SharedSubdir = SharedTempStorageDirectory();
             string Root = RootSharedTempStorageDirectory();
-            string Result = CombinePaths(Root, GameFolder, SharedSubdir);
-            if (!DirectoryExists_NoExceptions(Result))
+            string Result = CombinePaths(Root, GameFolder);
+            if (String.IsNullOrEmpty(GameFolder) || !DirectoryExists_NoExceptions(Result))
             {
                 string GameStr = "Game";
                 bool HadGame = false;
                 if (GameFolder.EndsWith(GameStr, StringComparison.InvariantCultureIgnoreCase))
                 {
                     string ShortFolder = GameFolder.Substring(0, GameFolder.Length - GameStr.Length);
-                    Result = CombinePaths(Root, ShortFolder, SharedSubdir);
+                    Result = CombinePaths(Root, ShortFolder);
                     HadGame = true;
                 }
                 if (!HadGame || !DirectoryExists_NoExceptions(Result))
                 {
-                    Result = CombinePaths(Root, "UE4", SharedSubdir);
+                    Result = CombinePaths(Root, "UE4");
                     if (!DirectoryExists_NoExceptions(Result))
                     {
                         throw new AutomationException("Could not find an appropriate shared temp folder {0}", Result);
@@ -577,6 +576,21 @@ namespace AutomationTool
             return Result;
         }
 
+        public static string ResolveSharedTempStorageDirectory(string GameFolder)
+        {
+            string SharedSubdir = SharedTempStorageDirectory();
+            string Result = CombinePaths(ResolveSharedBuildDirectory(GameFolder), SharedSubdir);
+            if (!DirectoryExists_NoExceptions(Result))
+            {
+                CreateDirectory_NoExceptions(Result);
+            }
+            if (!DirectoryExists_NoExceptions(Result))
+            {
+                throw new AutomationException("Could not create an appropriate shared temp folder {0}", Result);
+            }
+            return Result;
+        }
+
         static HashSet<string> TestedForClean = new HashSet<string>();
         public static void CleanSharedTempStorageDirectory(string GameFolder)
         {
diff --git a/Engine/Source/Programs/UnrealBuildTool/System/RulesCompiler.cs b/Engine/Source/Programs/UnrealBuildTool/System/RulesCompiler.cs
--- a/Engine/Source/Programs/UnrealBuildTool/System/RulesCompiler.cs
+++ b/Engine/Source/Programs/UnrealBuildTool/System/RulesCompiler.cs
@@ -607,6 +607,14 @@ namespace UnrealBuildTool
         {
             return new List<UnrealTargetConfiguration>{UnrealTargetConfiguration.Development};
         }
+        /// <summary>
+        /// Return a list of configs for target platforms for formal builds
+        /// </summary>
+        /// <returns>a list of configs for a target platforms for the monolithic</returns>        
+        public virtual List<UnrealTargetConfiguration> GUBP_GetConfigsForFormalBuilds_MonolithicOnly(UnrealTargetPlatform HostPlatform, UnrealTargetPlatform Platform)
+        {
+            return new List<UnrealTargetConfiguration>();
+        }
 
         /// <summary>
         /// Return true if this target should be included in a promotion and indicate shared or not
