build script, fix build cook run default project

diff --git a/Engine/Source/Programs/AutomationTool/Scripts/BuildCookRun.Automation.cs b/Engine/Source/Programs/AutomationTool/Scripts/BuildCookRun.Automation.cs
--- a/Engine/Source/Programs/AutomationTool/Scripts/BuildCookRun.Automation.cs
+++ b/Engine/Source/Programs/AutomationTool/Scripts/BuildCookRun.Automation.cs
@@ -323,7 +323,7 @@ public class BuildCookRun : BuildCommand
 				}
 				else
 				{
-					var ProjectName = ParseParamValue("project", CombinePaths("Samples", "SampleGames", "HoverShip", "HoverShip.uproject"));
+                    var ProjectName = ParseParamValue("project", CombinePaths("Samples", "Sandbox", "BlankProject", "BlankProject.uproject"));
 					ProjectName = ProjectName.Trim (new char[]{ '\"' });
                     if (ProjectName.IndexOfAny(new char[] { '\\', '/' }) < 0)
                     {
