build script, fix case sensitivity issues

diff --git a/Engine/Source/Programs/AutomationTool/ScriptCompiler.cs b/Engine/Source/Programs/AutomationTool/ScriptCompiler.cs
--- a/Engine/Source/Programs/AutomationTool/ScriptCompiler.cs
+++ b/Engine/Source/Programs/AutomationTool/ScriptCompiler.cs
@@ -281,7 +281,7 @@ namespace AutomationTool
 
 		private static string GetScriptAssemblyFolder()
 		{
-			return CommandUtils.CombinePaths(CommandUtils.CmdEnv.LocalRoot, "Engine", "Binaries", "DotNet", "AutomationScripts");
+			return CommandUtils.CombinePaths(CommandUtils.CmdEnv.LocalRoot, "Engine", "Binaries", "DotNET", "AutomationScripts");
 		}
 
 		#endregion
diff --git a/Engine/Source/Programs/AutomationTool/Scripts/Tests.Automation.cs b/Engine/Source/Programs/AutomationTool/Scripts/Tests.Automation.cs
--- a/Engine/Source/Programs/AutomationTool/Scripts/Tests.Automation.cs
+++ b/Engine/Source/Programs/AutomationTool/Scripts/Tests.Automation.cs
@@ -1628,7 +1628,7 @@ public class GenerateAutomationProject : BuildCommand
 
 			{
 				const string OutputPathTag = "<OutputPath>";
-				var OutputPath = CombinePaths(CmdEnv.LocalRoot, "Engine", "Binaries", "DotNet", "AutomationScripts");
+                var OutputPath = CombinePaths(CmdEnv.LocalRoot, "Engine", "Binaries", "DotNET", "AutomationScripts");
 				int PathStart = CSProjFileTemplate.IndexOf(OutputPathTag) + OutputPathTag.Length;
 				int PathEnd = CSProjFileTemplate.IndexOf("</OutputPath>", PathStart);
 				string OldOutputPath = CSProjFileTemplate.Substring(PathStart, PathEnd - PathStart);
diff --git a/Engine/Source/Programs/AutomationTool/TempStorage.cs b/Engine/Source/Programs/AutomationTool/TempStorage.cs
--- a/Engine/Source/Programs/AutomationTool/TempStorage.cs
+++ b/Engine/Source/Programs/AutomationTool/TempStorage.cs
@@ -137,6 +137,10 @@ namespace AutomationTool
                     var Filename = CombinePaths(InFilename);
                     Robust_FileExists_NoExceptions(true, Filename, "Could not add {0} to manifest because it does not exist");
 
+                    FileInfo Info = new FileInfo(Filename);
+                    Filename = Info.FullName;
+                    Robust_FileExists_NoExceptions(true, Filename, "Could not add {0} to manifest because it does not exist2");
+
                     if (!Filename.StartsWith(BaseFolder, StringComparison.InvariantCultureIgnoreCase))
                     {
                         throw new AutomationException("Could not add {0} to manifest because it does not start with the base folder {1}", Filename, BaseFolder);
diff --git a/Engine/Source/Programs/AutomationTool/UE4Build.cs b/Engine/Source/Programs/AutomationTool/UE4Build.cs
--- a/Engine/Source/Programs/AutomationTool/UE4Build.cs
+++ b/Engine/Source/Programs/AutomationTool/UE4Build.cs
@@ -1377,7 +1377,7 @@ namespace AutomationTool
 		public void AddUATFilesToBuildProducts()
 		{
 			const string ScriptsPostfix = ".Automation.dll";
-			var DotNetOutputLocation = CombinePaths(CmdEnv.LocalRoot, "Engine", "Binaries", "DotNet");
+            var DotNetOutputLocation = CombinePaths(CmdEnv.LocalRoot, "Engine", "Binaries", "DotNET");
 			var UATScriptsLocation = CombinePaths(DotNetOutputLocation, "AutomationScripts");
 			var UATScripts = Directory.GetFiles(UATScriptsLocation, "*" + ScriptsPostfix, SearchOption.AllDirectories);
 			if (IsNullOrEmpty(UATScripts))
