build script, clean up warnings

diff --git a/Engine/Source/Programs/AutomationTool/TempStorage.cs b/Engine/Source/Programs/AutomationTool/TempStorage.cs
--- a/Engine/Source/Programs/AutomationTool/TempStorage.cs
+++ b/Engine/Source/Programs/AutomationTool/TempStorage.cs
@@ -59,6 +59,8 @@ namespace AutomationTool
                     // RPCUtility is build by IPP and maybe other things
                     bSilentOkToBeDifferent = bSilentOkToBeDifferent || (Name == "Engine/Binaries/DotNET/RPCUtility.exe");
                     bSilentOkToBeDifferent = bSilentOkToBeDifferent || (Name == "Engine/Binaries/DotNET/RPCUtility.pdb");
+                    bSilentOkToBeDifferent = bSilentOkToBeDifferent || (Name == "Engine/Binaries/DotNET/AutomationTool.exe");
+                    bSilentOkToBeDifferent = bSilentOkToBeDifferent || (Name == "Engine/Binaries/DotNET/UnrealBuildTool.exe");
 
                     System.Diagnostics.TraceEventType LogType = bOkToBeDifferent ? System.Diagnostics.TraceEventType.Warning : System.Diagnostics.TraceEventType.Error;
                     if (bSilentOkToBeDifferent && bOkToBeDifferent)
