build script, warning cleanup

diff --git a/Engine/Source/Programs/AutomationTool/TempStorage.cs b/Engine/Source/Programs/AutomationTool/TempStorage.cs
--- a/Engine/Source/Programs/AutomationTool/TempStorage.cs
+++ b/Engine/Source/Programs/AutomationTool/TempStorage.cs
@@ -61,6 +61,7 @@ namespace AutomationTool
                     bSilentOkToBeDifferent = bSilentOkToBeDifferent || (Name == "Engine/Binaries/DotNET/RPCUtility.pdb");
                     bSilentOkToBeDifferent = bSilentOkToBeDifferent || (Name == "Engine/Binaries/DotNET/AutomationTool.exe");
                     bSilentOkToBeDifferent = bSilentOkToBeDifferent || (Name == "Engine/Binaries/DotNET/UnrealBuildTool.exe");
+                    bSilentOkToBeDifferent = bSilentOkToBeDifferent || (Name == "Engine/Binaries/DotNET/UnrealBuildTool.exe.config");
 
                     System.Diagnostics.TraceEventType LogType = bOkToBeDifferent ? System.Diagnostics.TraceEventType.Warning : System.Diagnostics.TraceEventType.Error;
                     if (bSilentOkToBeDifferent && bOkToBeDifferent)
