build script, changes test

diff --git a/Engine/Source/Programs/AutomationTool/Scripts/Tests.Automation.cs b/Engine/Source/Programs/AutomationTool/Scripts/Tests.Automation.cs
--- a/Engine/Source/Programs/AutomationTool/Scripts/Tests.Automation.cs
+++ b/Engine/Source/Programs/AutomationTool/Scripts/Tests.Automation.cs
@@ -1736,6 +1736,28 @@ class TestBlame : BuildCommand
 	}
 }
 
+[Help("Test P4 changes.")]
+[RequireP4]
+[DoesNotNeedP4CL]
+class TestChanges : BuildCommand
+{
+    public override void ExecuteBuild()
+    {
+        var CommandParam = ParseParamValue("CommandParam", "//depot/UE4-LauncherReleases/*/Source/...@2061085,2061287 //depot/UE4-LauncherReleases/*/Build/...@2061085,2061287");
+        var OutFilename = ParseParamValue("Out");
+        
+        List<P4Connection.ChangeRecord> ChangeRecords;
+        if (!P4.Changes(out ChangeRecords, CommandParam, true, true, LongComment: true))
+        {
+            throw new AutomationException("failed");
+        }
+        foreach (var Record in ChangeRecords)
+        {
+            Log("{0} {1} {2}", Record.CL, Record.UserEmail, Record.Summary);
+        }
+    }
+}
+
 [Help("Spawns a process to test if UAT kills it automatically.")]
 class TestKillAll : BuildCommand
 {
