build script, fix P4 changes for launcher branch

diff --git a/Engine/Source/Programs/AutomationTool/P4Utils.cs b/Engine/Source/Programs/AutomationTool/P4Utils.cs
--- a/Engine/Source/Programs/AutomationTool/P4Utils.cs
+++ b/Engine/Source/Programs/AutomationTool/P4Utils.cs
@@ -598,7 +598,7 @@ namespace AutomationTool
         public bool Changes(out List<ChangeRecord> ChangeRecords, string CommandLine, bool AllowSpew = true, bool UseCaching = false, bool LongComment = false)
         {
             // If the user specified '-l' or '-L', the summary will appear on subsequent lines (no quotes) instead of the same line (surrounded by single quotes)
-            bool bSummaryIsOnSameLine = CommandLine.IndexOf("-L", StringComparison.InvariantCultureIgnoreCase) == -1;
+            bool bSummaryIsOnSameLine = CommandLine.IndexOf(" -L ", StringComparison.InvariantCultureIgnoreCase) == -1 && !CommandLine.StartsWith("-L ",StringComparison.InvariantCultureIgnoreCase);
             if (bSummaryIsOnSameLine && LongComment)
             {
                 CommandLine = "-L " + CommandLine;
