build script, fix mac error

diff --git a/Engine/Source/Programs/AutomationTool/Scripts/Tests.Automation.cs b/Engine/Source/Programs/AutomationTool/Scripts/Tests.Automation.cs
--- a/Engine/Source/Programs/AutomationTool/Scripts/Tests.Automation.cs
+++ b/Engine/Source/Programs/AutomationTool/Scripts/Tests.Automation.cs
@@ -1744,7 +1744,6 @@ class TestChanges : BuildCommand
     public override void ExecuteBuild()
     {
         var CommandParam = ParseParamValue("CommandParam", "//depot/UE4-LauncherReleases/*/Source/...@2061085,2061287 //depot/UE4-LauncherReleases/*/Build/...@2061085,2061287");
-        var OutFilename = ParseParamValue("Out");
         
         List<P4Connection.ChangeRecord> ChangeRecords;
         if (!P4.Changes(out ChangeRecords, CommandParam, true, true, LongComment: true))
