fix for code vs codeless check when the source dir is missing
#ue4
#uat

diff --git a/Engine/Source/Programs/AutomationTool/ProjectParams.cs b/Engine/Source/Programs/AutomationTool/ProjectParams.cs
--- a/Engine/Source/Programs/AutomationTool/ProjectParams.cs
+++ b/Engine/Source/Programs/AutomationTool/ProjectParams.cs
@@ -1137,6 +1137,13 @@ namespace AutomationTool
 					ProgramTarget = TargetData.TargetName;
 				}
 			}
+			else if (!this.Build)
+			{
+				var ShortName = ProjectUtils.GetShortProjectName(RawProjectPath);
+				GameTarget = ShortName;
+				EditorTarget = ShortName + "Editor";
+				ServerTarget = ShortName + "Server";
+			}
 			else
 			{
 				throw new AutomationException("{0} does not look like uproject file but no targets have been found!", RawProjectPath);
diff --git a/Engine/Source/Programs/AutomationTool/ProjectUtils.cs b/Engine/Source/Programs/AutomationTool/ProjectUtils.cs
--- a/Engine/Source/Programs/AutomationTool/ProjectUtils.cs
+++ b/Engine/Source/Programs/AutomationTool/ProjectUtils.cs
@@ -168,6 +168,13 @@ namespace AutomationTool
 				}
 			}
 
+			// check to see if the uproject loads modules, only if we haven't already determined it is a code based project
+			if (!Properties.bIsCodeBasedProject)
+			{
+				string uprojectStr = File.ReadAllText(RawProjectPath);
+				Properties.bIsCodeBasedProject = uprojectStr.Contains("\"Modules\"");
+			}
+
 			return Properties;
 		}
 
