build script, add a few automated tests and add a list of all nodes

diff --git a/Engine/Source/Programs/AutomationTool/Scripts/GUBP.Automation.cs b/Engine/Source/Programs/AutomationTool/Scripts/GUBP.Automation.cs
--- a/Engine/Source/Programs/AutomationTool/Scripts/GUBP.Automation.cs
+++ b/Engine/Source/Programs/AutomationTool/Scripts/GUBP.Automation.cs
@@ -2639,7 +2639,7 @@ public class GUBP : BuildCommand
             {
                 ProjectArg = " -project=\"" + GameProj.FilePath + "\"";
             }
-            string WorkingCommandline = UATCommandLine + ProjectArg + " -NoSubmit";
+            string WorkingCommandline = UATCommandLine + ProjectArg + " -NoSubmit -addcmdline=\"-DisablePS4TMAPI\"";
             if (WorkingCommandline.Contains("-project=\"\""))
             {
                 throw new AutomationException("Command line {0} contains -project=\"\" which is doomed to fail", WorkingCommandline);
@@ -2908,6 +2908,24 @@ public class GUBP : BuildCommand
         return Result;
     }
 
+    public string CISFrequencyQuantumShiftString(string NodeToDo)
+    {
+        string FrequencyString = "";
+        int Quantum = GUBPNodes[NodeToDo].DependentCISFrequencyQuantumShift();
+        if (Quantum > 0)
+        {
+            int Minutes = 20 * (1 << Quantum);
+            if (Minutes < 60)
+            {
+                FrequencyString = string.Format(" ({0}m)", Minutes);
+            }
+            else
+            {
+                FrequencyString = string.Format(" ({0}h{1}m)", Minutes / 60, Minutes % 60);
+            }
+        }
+        return FrequencyString;
+    }
 
     public int ComputeDependentCISFrequencyQuantumShift(string NodeToDo)
     {
@@ -3144,20 +3162,8 @@ public class GUBP : BuildCommand
             {
                 Agent = "[" + Agent + "]";
             }
-            string FrequencyString = "";
-            int Quantum = GUBPNodes[NodeToDo].DependentCISFrequencyQuantumShift();
-            if (Quantum > 0)
-            {
-                int Minutes = 20 * (1 << Quantum);
-                if (Minutes < 60)
-                {
-                    FrequencyString = string.Format(" ({0}m)", Minutes);
-                }
-                else
-                {
-                    FrequencyString = string.Format(" ({0}h{1}m)", Minutes / 60, Minutes % 60);
-                }
-            }
+            string FrequencyString = CISFrequencyQuantumShiftString(NodeToDo);
+
             Log("      {0}{1}{2}{3}{4}{5}{6} {7}  {8}",
                 (LastAgentGroup != "" ? "  " : ""),
                 NodeToDo,
@@ -3361,7 +3367,7 @@ public class GUBP : BuildCommand
         return Result;
     }
 
-    List<string> TopologicalSort(HashSet<string> NodesToDo, string ExplicitTrigger, bool LocalOnly, bool SubSort = false)
+    List<string> TopologicalSort(HashSet<string> NodesToDo, string ExplicitTrigger = "", bool LocalOnly = false, bool SubSort = false, bool DoNotConsiderCompletion = false)
     {
         var StartTime = DateTime.UtcNow;
 
@@ -3388,7 +3394,7 @@ public class GUBP : BuildCommand
             }
             foreach (var Chain in AgentGroupChains)
             {
-                SortedAgentGroupChains.Add(Chain.Key, TopologicalSort(new HashSet<string>(Chain.Value), ExplicitTrigger, LocalOnly, true));
+                SortedAgentGroupChains.Add(Chain.Key, TopologicalSort(new HashSet<string>(Chain.Value), ExplicitTrigger, LocalOnly, true, DoNotConsiderCompletion));
             }
             Log("***************Done with recursion");
         }
@@ -3501,8 +3507,8 @@ public class GUBP : BuildCommand
                             }
                             else if (BestPseudoReady == bPseudoReady)
                             {
-                                bool IamLateTrigger = GUBPNodes[NodeToDo].TriggerNode() && NodeToDo != ExplicitTrigger && !NodeIsAlreadyComplete(NodeToDo, LocalOnly);
-                                bool BestIsLateTrigger = GUBPNodes[BestNode].TriggerNode() && BestNode != ExplicitTrigger && !NodeIsAlreadyComplete(BestNode, LocalOnly);
+                                bool IamLateTrigger = !DoNotConsiderCompletion && GUBPNodes[NodeToDo].TriggerNode() && NodeToDo != ExplicitTrigger && !NodeIsAlreadyComplete(NodeToDo, LocalOnly);
+                                bool BestIsLateTrigger = !DoNotConsiderCompletion && GUBPNodes[BestNode].TriggerNode() && BestNode != ExplicitTrigger && !NodeIsAlreadyComplete(BestNode, LocalOnly);
                                 if (BestIsLateTrigger && !IamLateTrigger)
                                 {
                                     bReady = false;
@@ -4248,22 +4254,25 @@ public class GUBP : BuildCommand
                 }
 
 
-
-                if (bAutomatedTesting)
+                if (HostPlatform != UnrealTargetPlatform.Mac) //temp hack till mac automated testing works
                 {
+
                     var EditorTests = Branch.BaseEngineProject.Properties.Targets[TargetRules.TargetType.Editor].Rules.GUBP_GetEditorTests_EditorTypeOnly(HostPlatform);
                     var EditorTestNodes = new List<string>();
                     foreach (var Test in EditorTests)
                     {
                         EditorTestNodes.Add(AddNode(new UATTestNode(this, HostPlatform, Branch.BaseEngineProject, Test.Key, Test.Value, AgentSharingGroup)));
-                        foreach (var NonCodeProject in Branch.NonCodeProjects)
+                        if (bAutomatedTesting)
                         {
-                            if (!NonCodeProjectNames.ContainsKey(NonCodeProject.GameName))
+                            foreach (var NonCodeProject in Branch.NonCodeProjects)
                             {
-                                continue;
+                                if (!NonCodeProjectNames.ContainsKey(NonCodeProject.GameName))
+                                {
+                                    continue;
+                                }
+                                if (HostPlatform == UnrealTargetPlatform.Mac) continue; //temp hack till mac automated testing works
+                                EditorTestNodes.Add(AddNode(new UATTestNode(this, HostPlatform, NonCodeProject, Test.Key, Test.Value, AgentSharingGroup)));
                             }
-                            if (HostPlatform == UnrealTargetPlatform.Mac) continue; //temp hack till mac automated testing works
-                            EditorTestNodes.Add(AddNode(new UATTestNode(this, HostPlatform, NonCodeProject, Test.Key, Test.Value, AgentSharingGroup)));
                         }
                     }
                     if (EditorTestNodes.Count > 0)
@@ -4272,6 +4281,7 @@ public class GUBP : BuildCommand
                     }
                 }
 
+
                 var ServerPlatforms = new List<UnrealTargetPlatform>();
                 var GamePlatforms = new List<UnrealTargetPlatform>();
 
@@ -4724,20 +4734,38 @@ if (HostPlatform == UnrealTargetPlatform.Mac) continue; //temp hack till mac aut
             ComputeDependentCISFrequencyQuantumShift(NodeToDo.Key);
         }
 
-        var NodesToDo = new HashSet<string>();
-
         if (bCleanLocalTempStorage)  // shared temp storage can never be wiped
         {
             DeleteLocalTempStorageManifests(CmdEnv);
         }
-        Log("******* {0} GUBP Nodes", GUBPNodes.Count);
-        foreach (var Node in GUBPNodes)
+
+        GUBPNodesControllingTrigger = new Dictionary<string, string>();
+        GUBPNodesControllingTriggerDotName = new Dictionary<string, string>();
+
+        var FullNodeList = new Dictionary<string, string>();
         {
-            Log("  {0}", Node.Key); 
+            Log("******* {0} GUBP Nodes", GUBPNodes.Count);
+            var SortedNodes = TopologicalSort(new HashSet<string>(GUBPNodes.Keys), LocalOnly: true, DoNotConsiderCompletion: true);
+            foreach (var Node in SortedNodes)
+            {
+                string Note = GetControllingTriggerDotName(Node);
+                if (Note == "")
+                {
+                    Note = CISFrequencyQuantumShiftString(Node);
+                }
+                if (Note == "")
+                {
+                    Note = "always";
+                }
+                Log("  {0}: {1}", Node, Note);
+                FullNodeList.Add(Node, Note);
+            }
         }
 
         bool bOnlyNode = false;
         bool bRelatedToNode = false;
+        var NodesToDo = new HashSet<string>();
+
         {
             string NodeSpec = ParseParamValue("Node");
             if (String.IsNullOrEmpty(NodeSpec))
@@ -5006,8 +5034,6 @@ if (HostPlatform == UnrealTargetPlatform.Mac) continue; //temp hack till mac aut
         }
 
         GUBPNodesCompleted = new Dictionary<string, bool>();
-        GUBPNodesControllingTrigger = new Dictionary<string, string>();
-        GUBPNodesControllingTriggerDotName = new Dictionary<string, string>();
         GUBPNodesHistory = new Dictionary<string, NodeHistory>();
 
         Log("******* Caching completion");
@@ -5094,6 +5120,10 @@ if (HostPlatform == UnrealTargetPlatform.Mac) continue; //temp hack till mac aut
             }
             var ECProps = new List<string>();
             ECProps.Add(String.Format("TimeIndex={0}", TimeIndex));
+            foreach (var NodePair in FullNodeList)
+            {
+                ECProps.Add(string.Format("AllNodes/{0}={1}", NodePair.Key, NodePair.Value));
+            }
 
             var ECJobProps = new List<string>();
             if (ExplicitTrigger != "")
