build script, do not add props for nodes that don't run in EC

diff --git a/Engine/Source/Programs/AutomationTool/Scripts/GUBP.Automation.cs b/Engine/Source/Programs/AutomationTool/Scripts/GUBP.Automation.cs
--- a/Engine/Source/Programs/AutomationTool/Scripts/GUBP.Automation.cs
+++ b/Engine/Source/Programs/AutomationTool/Scripts/GUBP.Automation.cs
@@ -4757,8 +4757,15 @@ if (HostPlatform == UnrealTargetPlatform.Mac) continue; //temp hack till mac aut
                 {
                     Note = "always";
                 }
-                Log("  {0}: {1}", Node, Note);
-                FullNodeList.Add(Node, Note);
+                if (GUBPNodes[Node].RunInEC())
+                {
+                    Log("  {0}: {1}", Node, Note);
+                    FullNodeList.Add(Node, Note);
+                }
+                else
+                {
+                    Log("  {0}: {1} [Aggregate]", Node, Note);
+                }
             }
         }
 
