build script, add dependency lists

diff --git a/Engine/Source/Programs/AutomationTool/Scripts/GUBP.Automation.cs b/Engine/Source/Programs/AutomationTool/Scripts/GUBP.Automation.cs
--- a/Engine/Source/Programs/AutomationTool/Scripts/GUBP.Automation.cs
+++ b/Engine/Source/Programs/AutomationTool/Scripts/GUBP.Automation.cs
@@ -4743,6 +4743,7 @@ if (HostPlatform == UnrealTargetPlatform.Mac) continue; //temp hack till mac aut
         GUBPNodesControllingTriggerDotName = new Dictionary<string, string>();
 
         var FullNodeList = new Dictionary<string, string>();
+        var FullNodeDirectDependencies = new Dictionary<string, string>();
         {
             Log("******* {0} GUBP Nodes", GUBPNodes.Count);
             var SortedNodes = TopologicalSort(new HashSet<string>(GUBPNodes.Keys), LocalOnly: true, DoNotConsiderCompletion: true);
@@ -4759,8 +4760,19 @@ if (HostPlatform == UnrealTargetPlatform.Mac) continue; //temp hack till mac aut
                 }
                 if (GUBPNodes[Node].RunInEC())
                 {
-                    Log("  {0}: {1}", Node, Note);
+                    var Deps = GetECDependencies(Node);
+                    string All = "";
+                    foreach (var Dep in Deps)
+                    {
+                        if (All != "")
+                        {
+                            All += " ";
+                        }
+                        All += Dep;
+                    }
+                    Log("  {0}: {1}      {2}", Node, Note, All);
                     FullNodeList.Add(Node, Note);
+                    FullNodeDirectDependencies.Add(Node, All);
                 }
                 else
                 {
@@ -5131,6 +5143,10 @@ if (HostPlatform == UnrealTargetPlatform.Mac) continue; //temp hack till mac aut
             {
                 ECProps.Add(string.Format("AllNodes/{0}={1}", NodePair.Key, NodePair.Value));
             }
+            foreach (var NodePair in FullNodeDirectDependencies)
+            {
+                ECProps.Add(string.Format("DirectDependencies/{0}={1}", NodePair.Key, NodePair.Value));
+            }
 
             var ECJobProps = new List<string>();
             if (ExplicitTrigger != "")
