Add conhost to systemservices list

diff --git a/Engine/Source/Programs/AutomationTool/WindowsHostPlatform.cs b/Engine/Source/Programs/AutomationTool/WindowsHostPlatform.cs
--- a/Engine/Source/Programs/AutomationTool/WindowsHostPlatform.cs
+++ b/Engine/Source/Programs/AutomationTool/WindowsHostPlatform.cs
@@ -216,7 +216,8 @@ namespace AutomationTool
 			"smss",
 			"svchost",
 			"services",
-			"lsass"
+			"lsass",
+            "conhost"
 		};
 		public override string[] DontKillProcessList
 		{
