build script, timing for shared temp storage, spew reduction, 50 people, not 50 changes for spam reduction

diff --git a/Engine/Source/Programs/AutomationTool/ProcessUtils.cs b/Engine/Source/Programs/AutomationTool/ProcessUtils.cs
--- a/Engine/Source/Programs/AutomationTool/ProcessUtils.cs
+++ b/Engine/Source/Programs/AutomationTool/ProcessUtils.cs
@@ -687,8 +687,6 @@ namespace AutomationTool
 		/// <param name="Options">Defines the options how to run. See ERunOptions.</param>
         public static void RunAndLog(string App, string CommandLine, string Logfile = null, int MaxSuccessCode = 0, string Input = null, ERunOptions Options = ERunOptions.Default)
         {
-            Log("Running {0} {1}", App, CommandLine);
-
             ProcessResult Result = Run(App, CommandLine, Input, Options);
             if (Result.Output.Length > 0 && Logfile != null)
             {
@@ -719,8 +717,6 @@ namespace AutomationTool
         /// <returns>Whether the program executed successfully or not.</returns>
         public static string RunAndLog(string App, string CommandLine, out int SuccessCode, string Logfile = null)
         {
-            Log("Running {0} {1}", App, CommandLine);
-
             ProcessResult Result = Run(App, CommandLine);
             SuccessCode = Result.ExitCode;
             if (Result.Output.Length > 0 && Logfile != null)
diff --git a/Engine/Source/Programs/AutomationTool/Scripts/GUBP.Automation.cs b/Engine/Source/Programs/AutomationTool/Scripts/GUBP.Automation.cs
--- a/Engine/Source/Programs/AutomationTool/Scripts/GUBP.Automation.cs
+++ b/Engine/Source/Programs/AutomationTool/Scripts/GUBP.Automation.cs
@@ -2989,7 +2989,7 @@ public class GUBP : BuildCommand
             ERunOptions Opts = ERunOptions.Default;
             if (bQuiet)
             {
-                Opts = Opts & ~ERunOptions.AllowSpew;
+                Opts = (Opts & ~ERunOptions.AllowSpew) | ERunOptions.NoLoggingOfRunCommand;
             }
             RunAndLog("ectool", Args, Options: Opts);
         }
@@ -3683,11 +3683,18 @@ public class GUBP : BuildCommand
             if (History.LastSucceeded > 0 && History.LastSucceeded < P4Env.Changelist)
             {
                 var ChangeRecords = GetChanges(History.LastSucceeded, P4Env.Changelist, History.LastSucceeded);
-                int NumChanges = 0;
                 foreach (var Record in ChangeRecords)
                 {
                     FailCauserEMails = GUBPNode.MergeSpaceStrings(FailCauserEMails, Record.UserEmail);
-                    if (++NumChanges > 50)
+                    int NumPeople = 1;
+                    foreach (var AChar in FailCauserEMails.ToCharArray())
+                    {
+                        if (AChar == ' ')
+                        {
+                            NumPeople++;
+                        }
+                    }
+                    if (NumPeople > 50)
                     {
                         FailCauserEMails = "";
                         EMailNote = String.Format("This step has been broken for more than 50 changes. It last succeeded at CL {0}. ", History.LastSucceeded);
diff --git a/Engine/Source/Programs/AutomationTool/TempStorage.cs b/Engine/Source/Programs/AutomationTool/TempStorage.cs
--- a/Engine/Source/Programs/AutomationTool/TempStorage.cs
+++ b/Engine/Source/Programs/AutomationTool/TempStorage.cs
@@ -632,7 +632,7 @@ namespace AutomationTool
                 throw new AutomationException("Saver manifest differs has wrong number of files {0} != {1}", Saver.GetFileCount(), Files.Count);
             }
             var TempFilename = FinalFilename + ".temp";
-            if (FileExists_NoExceptions(TempFilename))
+            if (FileExists_NoExceptions(true, TempFilename))
             {
                 throw new AutomationException("Temp manifest file already exists {0}", TempFilename);
             }
@@ -724,6 +724,8 @@ namespace AutomationTool
             var Local = SaveLocalTempStorageManifest(Env, BaseFolder, StorageBlockName, Files); 
             if (!bLocalOnly)
             {
+                var StartTime = DateTime.UtcNow;
+
                 var BlockPath = SharedTempStorageDirectory(StorageBlockName, GameFolder);
                 Log("Storing to {0}", BlockPath);
                 if (DirectoryExists_NoExceptions(BlockPath))
@@ -763,7 +765,12 @@ namespace AutomationTool
                     RenameFile_NoExceptions(SharedTempStorageManifestFilename(Env, StorageBlockName, GameFolder), SharedTempStorageManifestFilename(Env, StorageBlockName, GameFolder) + ".broken");
                     throw new AutomationException("Shared and Local manifest mismatch.");
                 }
-
+                float BuildDuration = (float)((DateTime.UtcNow - StartTime).TotalSeconds);
+                if (BuildDuration > 60.0f && Shared.GetTotalSize() > 0)
+                {
+                    var MBSec = (((float)(Shared.GetTotalSize())) / (1024.0f * 1024.0f)) / BuildDuration;
+                    Log("Wrote to shared temp storage at {0} MB/s    {1}B {2}s", MBSec, Shared.GetTotalSize(), BuildDuration);
+                }
             }
 
         }
@@ -799,6 +806,7 @@ namespace AutomationTool
                 return Files;
             }
             WasLocal = false;
+            var StartTime = DateTime.UtcNow;
 
             var BlockPath = CombinePaths(SharedTempStorageDirectory(StorageBlockName, GameFolder), "/");
             if (!BlockPath.EndsWith("/") && !BlockPath.EndsWith("\\"))
@@ -854,6 +862,12 @@ namespace AutomationTool
                 RenameFile_NoExceptions(LocalManifest, LocalManifest + ".broken");
                 throw new AutomationException("Shared and Local manifest mismatch.");
             }
+            float BuildDuration = (float)((DateTime.UtcNow - StartTime).TotalSeconds);
+            if (BuildDuration > 60.0f && Shared.GetTotalSize() > 0)
+            {
+                var MBSec = (((float)(Shared.GetTotalSize())) / (1024.0f * 1024.0f)) / BuildDuration;
+                Log("Read from shared temp storage at {0} MB/s    {1}B {2}s", MBSec, Shared.GetTotalSize(), BuildDuration);
+            }
             return DestFiles;
         }
         public static List<string> FindTempStorageManifests(CommandEnvironment Env, string StorageBlockName, bool LocalOnly = false, bool SharedOnly = false, string GameFolder = "")
