build script, CIS heros

diff --git a/Engine/Source/Programs/AutomationTool/Scripts/GUBP.Automation.cs b/Engine/Source/Programs/AutomationTool/Scripts/GUBP.Automation.cs
--- a/Engine/Source/Programs/AutomationTool/Scripts/GUBP.Automation.cs
+++ b/Engine/Source/Programs/AutomationTool/Scripts/GUBP.Automation.cs
@@ -97,9 +97,9 @@ public class GUBP : BuildCommand
         {
             return new List<string>();
         }
-        public virtual string ModifyEmail(string Email, GUBP bp, string Branch, string NodeName)
+        public virtual List<string> ModifyEmail(string Email, GUBP bp, string Branch, string NodeName)
         {
-            return Email;
+            return new List<string>{Email};
         }
         public virtual List<string> FinalizeEmails(List<string> Emails, GUBP bp, string Branch, string NodeName)
         {
@@ -110,8 +110,6 @@ public class GUBP : BuildCommand
     private static List<GUBPEmailHacker> EmailHackers;
     private string HackEmails(string Emails, string Branch, string NodeName)
     {
-        var StartTime = DateTime.UtcNow;
-
         if (EmailHackers == null)
         {
             EmailHackers = new List<GUBPEmailHacker>();
@@ -143,11 +141,7 @@ public class GUBP : BuildCommand
             var NewResult = new List<string>();
             foreach (var Email in Result)
             {
-                var NewEmail = EmailHacker.ModifyEmail(Email, this, Branch, NodeName);
-                if (!String.IsNullOrEmpty(NewEmail))
-                {
-                    NewResult.Add(NewEmail);
-                }
+                NewResult.AddRange(EmailHacker.ModifyEmail(Email, this, Branch, NodeName));
             }
             Result = NewResult;
         }
@@ -162,9 +156,6 @@ public class GUBP : BuildCommand
             FinalEmails = GUBPNode.MergeSpaceStrings(FinalEmails, Email);
             Count++;
         }
-        var BuildDuration = (DateTime.UtcNow - StartTime).TotalMilliseconds;
-        Log("Took {0}s to hack {1} emails", BuildDuration / 1000, Count);
-
         return FinalEmails;
     }
     public abstract class GUBPNode
@@ -3671,6 +3662,7 @@ public class GUBP : BuildCommand
         string FailCauserEMails = "";
         string EMailNote = "";
         bool SendSuccessForGreenAfterRed = false;
+        int NumPeople = 0;
         if (GUBPNodesHistory.ContainsKey(NodeToDo))
         {
             var History = GUBPNodesHistory[NodeToDo];
@@ -3686,7 +3678,10 @@ public class GUBP : BuildCommand
                 foreach (var Record in ChangeRecords)
                 {
                     FailCauserEMails = GUBPNode.MergeSpaceStrings(FailCauserEMails, Record.UserEmail);
-                    int NumPeople = 1;
+                }
+                if (!String.IsNullOrEmpty(FailCauserEMails))
+                {
+                    NumPeople++;
                     foreach (var AChar in FailCauserEMails.ToCharArray())
                     {
                         if (AChar == ' ')
@@ -3696,9 +3691,7 @@ public class GUBP : BuildCommand
                     }
                     if (NumPeople > 50)
                     {
-                        FailCauserEMails = "";
                         EMailNote = String.Format("This step has been broken for more than 50 changes. It last succeeded at CL {0}. ", History.LastSucceeded);
-                        break;
                     }
                 }
             }
@@ -3733,6 +3726,10 @@ public class GUBP : BuildCommand
             var AdditonalEmails = "";
             if (ParseParam("CIS") && !GUBPNodes[NodeToDo].SendSuccessEmail())
             {
+                if (NumPeople > 50 || NumPeople == 0)
+                {
+                    FailCauserEMails = "[CISHeros]";
+                }
                 AdditonalEmails = FailCauserEMails;
             }
             string AddEmails = ParseParamValue("AddEmails");
