Disabling OneSky code

diff --git a/Engine/Source/Programs/AutomationTool/Scripts/AutomationScripts.Automation.csproj b/Engine/Source/Programs/AutomationTool/Scripts/AutomationScripts.Automation.csproj
--- a/Engine/Source/Programs/AutomationTool/Scripts/AutomationScripts.Automation.csproj
+++ b/Engine/Source/Programs/AutomationTool/Scripts/AutomationScripts.Automation.csproj
@@ -38,10 +38,6 @@
       <HintPath>..\..\..\..\Binaries\DotNET\Ionic.Zip.Reduced.dll</HintPath>
       <Private>False</Private>
     </Reference>
-    <Reference Include="OneSky, Version=1.0.0.0, Culture=neutral, processorArchitecture=MSIL">
-      <SpecificVersion>False</SpecificVersion>
-      <HintPath>..\..\..\ThirdParty\OneSky\bin\Release\OneSky.dll</HintPath>
-    </Reference>
     <Reference Include="System" />
     <Reference Include="System.Core" />
     <Reference Include="System.Runtime.Serialization" />
diff --git a/Engine/Source/Programs/AutomationTool/Scripts/LauncherLocalization.Automation.cs b/Engine/Source/Programs/AutomationTool/Scripts/LauncherLocalization.Automation.cs
--- a/Engine/Source/Programs/AutomationTool/Scripts/LauncherLocalization.Automation.cs
+++ b/Engine/Source/Programs/AutomationTool/Scripts/LauncherLocalization.Automation.cs
@@ -5,7 +5,7 @@ using System.Globalization;
 using System.Linq;
 using System.IO;
 using AutomationTool;
-using OneSky;
+//using OneSky;
 
 class LauncherLocalization : BuildCommand
 {
@@ -85,105 +85,105 @@ class LauncherLocalization : BuildCommand
         //UploadDirectoryToProject(GetAppProject(oneSkyService), new DirectoryInfo(CmdEnv.LocalRoot + "/Engine/Programs/NoRedist/UnrealEngineLauncher/Content/Localization/App"), "*.po");
     }
 
-    private static ProjectGroup GetLauncherGroup(OneSkyService oneSkyService)
-    {
-        var launcherGroup = oneSkyService.ProjectGroups.FirstOrDefault(g => g.Name == "Launcher");
-
-        if (launcherGroup == null)
-        {
-            launcherGroup = new ProjectGroup("Launcher", new CultureInfo("en"));
-            oneSkyService.ProjectGroups.Add(launcherGroup);
-        }
-
-        return launcherGroup;
-    }
-
-    private static OneSky.Project GetAppProject(OneSkyService oneSkyService)
-    {
-        var launcherGroup = GetLauncherGroup(oneSkyService);
-
-        OneSky.Project appProject = launcherGroup.Projects.FirstOrDefault(p => p.Name == "App");
-
-        if (appProject == null)
-        {
-            ProjectType projectType = oneSkyService.ProjectTypes.First(pt => pt.Code == "website");
-
-            appProject = new OneSky.Project("App", "The core application text that ships with the Launcher", projectType);
-            launcherGroup.Projects.Add(appProject);
-        }
-
-        return appProject;
-    }
-
-    private static void ExportFileToDirectory(UploadedFile file, DirectoryInfo destination, IEnumerable<CultureInfo> cultures)
-    {
-        foreach (var culture in cultures)
-        {
-            var cultureDirectory = new DirectoryInfo(Path.Combine(destination.FullName, culture.Name));
-            if (!cultureDirectory.Exists)
-            {
-                cultureDirectory.Create();
-            }
-
-            using (var memoryStream = new MemoryStream())
-            {
-                var exportFile = new FileInfo(Path.Combine(cultureDirectory.FullName, file.Filename));
-
-                var exportTranslationState = file.ExportTranslation(culture, memoryStream).Result;
-                if (exportTranslationState == UploadedFile.ExportTranslationState.Success)
-                {
-                    memoryStream.Position = 0;
-                    using (Stream fileStream = File.OpenWrite(exportFile.FullName))
-                    {
-                        memoryStream.CopyTo(fileStream);
-                        Console.WriteLine("[SUCCESS] Exporting: " + exportFile.FullName + " Locale: " + culture.Name);
-                    }
-                }
-                else if (exportTranslationState == UploadedFile.ExportTranslationState.NoContent)
-                {
-                    Console.WriteLine("[WARNING] Exporting: " + exportFile.FullName + " Locale: " + culture.Name + " has no translations!");
-                }
-                else
-                {
-                    Console.WriteLine("[FAILED] Exporting: " + exportFile.FullName + " Locale: " + culture.Name);
-                }
-            }
-        }
-    }
-
-    static void UploadDirectoryToProject(OneSky.Project project, DirectoryInfo directory, string fileExtension)
-    {
-        foreach (var file in Directory.GetFiles(directory.FullName, fileExtension, SearchOption.AllDirectories))
-        {
-            DirectoryInfo parentDirectory = Directory.GetParent(file);
-            string localeName = parentDirectory.Name;
-            string currentFile = file;
-
-            using (var fileStream = File.OpenRead(currentFile))
-            {
-                // Read the BOM
-                var bom = new byte[3];
-                fileStream.Read(bom, 0, 3);
-
-                //We want to ignore the utf8 BOM
-                if (bom[0] != 0xef || bom[1] != 0xbb || bom[2] != 0xbf)
-                {
-                    fileStream.Position = 0;
-                }
-
-                Console.WriteLine("Uploading: " + currentFile + " Locale: " + localeName);
-                var uploadedFile = project.Upload(Path.GetFileName(currentFile), fileStream, new CultureInfo(localeName)).Result;
-
-                if (uploadedFile == null)
-                {
-                    Console.WriteLine("[FAILED] Uploading: " + currentFile + " Locale: " + localeName);
-                }
-                else
-                {
-                    Console.WriteLine("[SUCCESS] Uploading: " + currentFile + " Locale: " + localeName);
-                }
-            }
-        }
-    }
+    //private static ProjectGroup GetLauncherGroup(OneSkyService oneSkyService)
+    //{
+    //    var launcherGroup = oneSkyService.ProjectGroups.FirstOrDefault(g => g.Name == "Launcher");
+
+    //    if (launcherGroup == null)
+    //    {
+    //        launcherGroup = new ProjectGroup("Launcher", new CultureInfo("en"));
+    //        oneSkyService.ProjectGroups.Add(launcherGroup);
+    //    }
+
+    //    return launcherGroup;
+    //}
+
+    //private static OneSky.Project GetAppProject(OneSkyService oneSkyService)
+    //{
+    //    var launcherGroup = GetLauncherGroup(oneSkyService);
+
+    //    OneSky.Project appProject = launcherGroup.Projects.FirstOrDefault(p => p.Name == "App");
+
+    //    if (appProject == null)
+    //    {
+    //        ProjectType projectType = oneSkyService.ProjectTypes.First(pt => pt.Code == "website");
+
+    //        appProject = new OneSky.Project("App", "The core application text that ships with the Launcher", projectType);
+    //        launcherGroup.Projects.Add(appProject);
+    //    }
+
+    //    return appProject;
+    //}
+
+    //private static void ExportFileToDirectory(UploadedFile file, DirectoryInfo destination, IEnumerable<CultureInfo> cultures)
+    //{
+    //    foreach (var culture in cultures)
+    //    {
+    //        var cultureDirectory = new DirectoryInfo(Path.Combine(destination.FullName, culture.Name));
+    //        if (!cultureDirectory.Exists)
+    //        {
+    //            cultureDirectory.Create();
+    //        }
+
+    //        using (var memoryStream = new MemoryStream())
+    //        {
+    //            var exportFile = new FileInfo(Path.Combine(cultureDirectory.FullName, file.Filename));
+
+    //            var exportTranslationState = file.ExportTranslation(culture, memoryStream).Result;
+    //            if (exportTranslationState == UploadedFile.ExportTranslationState.Success)
+    //            {
+    //                memoryStream.Position = 0;
+    //                using (Stream fileStream = File.OpenWrite(exportFile.FullName))
+    //                {
+    //                    memoryStream.CopyTo(fileStream);
+    //                    Console.WriteLine("[SUCCESS] Exporting: " + exportFile.FullName + " Locale: " + culture.Name);
+    //                }
+    //            }
+    //            else if (exportTranslationState == UploadedFile.ExportTranslationState.NoContent)
+    //            {
+    //                Console.WriteLine("[WARNING] Exporting: " + exportFile.FullName + " Locale: " + culture.Name + " has no translations!");
+    //            }
+    //            else
+    //            {
+    //                Console.WriteLine("[FAILED] Exporting: " + exportFile.FullName + " Locale: " + culture.Name);
+    //            }
+    //        }
+    //    }
+    //}
+
+    //static void UploadDirectoryToProject(OneSky.Project project, DirectoryInfo directory, string fileExtension)
+    //{
+    //    foreach (var file in Directory.GetFiles(directory.FullName, fileExtension, SearchOption.AllDirectories))
+    //    {
+    //        DirectoryInfo parentDirectory = Directory.GetParent(file);
+    //        string localeName = parentDirectory.Name;
+    //        string currentFile = file;
+
+    //        using (var fileStream = File.OpenRead(currentFile))
+    //        {
+    //            // Read the BOM
+    //            var bom = new byte[3];
+    //            fileStream.Read(bom, 0, 3);
+
+    //            //We want to ignore the utf8 BOM
+    //            if (bom[0] != 0xef || bom[1] != 0xbb || bom[2] != 0xbf)
+    //            {
+    //                fileStream.Position = 0;
+    //            }
+
+    //            Console.WriteLine("Uploading: " + currentFile + " Locale: " + localeName);
+    //            var uploadedFile = project.Upload(Path.GetFileName(currentFile), fileStream, new CultureInfo(localeName)).Result;
+
+    //            if (uploadedFile == null)
+    //            {
+    //                Console.WriteLine("[FAILED] Uploading: " + currentFile + " Locale: " + localeName);
+    //            }
+    //            else
+    //            {
+    //                Console.WriteLine("[SUCCESS] Uploading: " + currentFile + " Locale: " + localeName);
+    //            }
+    //        }
+    //    }
+    //}
 }
 
