build script, don't add causers or heros for trigger nodes

diff --git a/Engine/Source/Programs/AutomationTool/Scripts/GUBP.Automation.cs b/Engine/Source/Programs/AutomationTool/Scripts/GUBP.Automation.cs
--- a/Engine/Source/Programs/AutomationTool/Scripts/GUBP.Automation.cs
+++ b/Engine/Source/Programs/AutomationTool/Scripts/GUBP.Automation.cs
@@ -3724,7 +3724,7 @@ public class GUBP : BuildCommand
 
         {
             var AdditonalEmails = "";
-            if (ParseParam("CIS") && !GUBPNodes[NodeToDo].SendSuccessEmail())
+            if (ParseParam("CIS") && !GUBPNodes[NodeToDo].SendSuccessEmail() && !GUBPNodes[NodeToDo].TriggerNode())
             {
                 if (NumPeople > 50 || NumPeople == 0)
                 {
