Removed an excess comment when correctly finding an MCP config

diff --git a/Engine/Source/Programs/AutomationTool/MCPPublic.cs b/Engine/Source/Programs/AutomationTool/MCPPublic.cs
--- a/Engine/Source/Programs/AutomationTool/MCPPublic.cs
+++ b/Engine/Source/Programs/AutomationTool/MCPPublic.cs
@@ -547,7 +547,6 @@ namespace EpicGames.MCP.Config
             {
                 throw new AutomationException("Unable to find requested McpConfig: {0}", ConfigName);
             }
-            CommandUtils.Log("Found config for {0}", ConfigName);
             return LoadedConfig;
         }
     }
