Making empty-host clients valid to use with UAT

diff --git a/Engine/Source/Programs/AutomationTool/LocalP4Environment.cs b/Engine/Source/Programs/AutomationTool/LocalP4Environment.cs
--- a/Engine/Source/Programs/AutomationTool/LocalP4Environment.cs
+++ b/Engine/Source/Programs/AutomationTool/LocalP4Environment.cs
@@ -183,9 +183,9 @@ namespace AutomationTool
 			P4ClientInfo[] P4Clients = Connection.GetClientsForUser(UserName, UATLocation);
 			foreach (var Client in P4Clients)
 			{
-				if (String.IsNullOrEmpty(Client.Host) || String.Compare(Client.Host, HostName, true) != 0)
+				if (!String.IsNullOrEmpty(Client.Host) && String.Compare(Client.Host, HostName, true) != 0)
 				{
-					Log.TraceInformation("Wrong Host {0} \"{1}\" != \"{2}\"", Client.Name, Client.Host, HostName);
+					Log.TraceInformation("Rejecting client because of different Host {0} \"{1}\" != \"{2}\"", Client.Name, Client.Host, HostName);
 					continue;
 				}
 				
@@ -203,8 +203,20 @@ namespace AutomationTool
 			}
 			else
 			{
-				Log.TraceWarning("{0} clients found that match the current host and root path. The most recently accessed client will be used.", MatchingClients.Count);
-				ClientToUse = GetMostRecentClient(MatchingClients);
+				// We may have empty host clients here, so pick the first non-empty one if possible
+				foreach (var Client in MatchingClients)
+				{
+					if (!String.IsNullOrEmpty(Client.Host) && String.Compare(Client.Host, HostName, true) == 0)
+					{
+						ClientToUse = Client;
+						break;
+					}
+				}
+				if (ClientToUse == null)
+				{
+					Log.TraceWarning("{0} clients found that match the current host and root path. The most recently accessed client will be used.", MatchingClients.Count);
+					ClientToUse = GetMostRecentClient(MatchingClients);
+				}
 			}
 
 			return ClientToUse;
