UAT: Making killing processes safer (added a safe way of accessing process name)

diff --git a/Engine/Source/Programs/AutomationTool/ProcessUtils.cs b/Engine/Source/Programs/AutomationTool/ProcessUtils.cs
--- a/Engine/Source/Programs/AutomationTool/ProcessUtils.cs
+++ b/Engine/Source/Programs/AutomationTool/ProcessUtils.cs
@@ -90,9 +90,10 @@ namespace AutomationTool
 			// Remove processes that can't be killed
 			for (int ProcessIndex = ProcessesToKill.Count - 1; ProcessIndex >= 0; --ProcessIndex )
 			{
-				if (!CanBeKilled(ProcessesToKill[ProcessIndex].ProcessObject.ProcessName))
+				var ProcessName = ProcessesToKill[ProcessIndex].GetProcessName();
+				if (!String.IsNullOrEmpty(ProcessName) && !CanBeKilled(ProcessName))
 				{
-					CommandUtils.Log("Ignoring process \"{0}\" because it can't be killed.", ProcessesToKill[ProcessIndex].ProcessObject.ProcessName);
+					CommandUtils.Log("Ignoring process \"{0}\" because it can't be killed.", ProcessName);
 					ProcessesToKill.RemoveAt(ProcessIndex);
 				}
 			}
@@ -101,7 +102,7 @@ namespace AutomationTool
 			{
 				if (!Proc.HasExited)
 				{
-					CommandUtils.Log("  {0}", Proc.ProcessObject.ProcessName);
+					CommandUtils.Log("  {0}", Proc.GetProcessName());
 				}
 			}
             if (CommandUtils.IsBuildMachine)
@@ -116,7 +117,7 @@ namespace AutomationTool
                             if (!Proc.HasExited)
                             {
                                 AllDone = false;
-                                CommandUtils.Log(TraceEventType.Information, "Waiting for process: {0}", Proc.ProcessObject.ProcessName);
+                                CommandUtils.Log(TraceEventType.Information, "Waiting for process: {0}", Proc.GetProcessName());
                             }
                         }
                         catch (Exception)
@@ -152,7 +153,7 @@ namespace AutomationTool
 				{
 					if (!Proc.HasExited)
 					{
-						CommandUtils.Log(TraceEventType.Information, "Killing process: {0}", Proc.ProcessObject.ProcessName);
+						CommandUtils.Log(TraceEventType.Information, "Killing process: {0}", Proc.GetProcessName());
 						Proc.StopProcess(false);
 					}
 				}
@@ -318,6 +319,23 @@ namespace AutomationTool
 			get { return Proc; }
 		}
 
+		/// <summary>
+		/// Thread-safe way of getting the process name
+		/// </summary>
+		/// <returns>Name of the process this object represents</returns>
+		public string GetProcessName()
+		{
+			string Name = null;
+			lock (ProcSyncObject)
+			{
+				if (Proc != null)
+				{
+					Name = Proc.ProcessName;
+				}
+			}
+			return Name;
+		}
+
 		/// <summary>
 		/// Object iterface.
 		/// </summary>
