build script, add IOSOnPC test node and sortkey prop folder

diff --git a/Engine/Source/Programs/AutomationTool/Scripts/GUBP.Automation.cs b/Engine/Source/Programs/AutomationTool/Scripts/GUBP.Automation.cs
--- a/Engine/Source/Programs/AutomationTool/Scripts/GUBP.Automation.cs
+++ b/Engine/Source/Programs/AutomationTool/Scripts/GUBP.Automation.cs
@@ -2491,6 +2491,38 @@ public class GUBP : BuildCommand
         }
     }
 
+    public class IOSOnPCTestNode : TestNode
+    {
+        public IOSOnPCTestNode()
+            : base(UnrealTargetPlatform.Win64)
+        {
+        }
+        public static string StaticGetFullName()
+        {
+            return "IOSOnPCTestCompile";
+        }
+        public override string GetFullName()
+        {
+            return StaticGetFullName();
+        }
+        public override int CISFrequencyQuantumShift(GUBP bp)
+        {
+            return base.CISFrequencyQuantumShift(bp) + 3;
+        }
+        public override void DoTest(GUBP bp)
+        {
+            var Build = new UE4Build(bp);
+            var Agenda = new UE4Build.BuildAgenda();
+
+            Agenda.AddTargets(new string[] { "UE4Game" }, UnrealTargetPlatform.IOS, UnrealTargetConfiguration.Development);
+
+            Build.Build(Agenda, InDeleteBuildProducts: true, InUpdateVersionFiles: false);
+
+            UE4Build.CheckBuildProducts(Build.BuildProductFiles);
+            SaveRecordOfSuccessAndAddToBuildProducts();
+        }
+    }
+
     public class UATTestNode : TestNode
     {
         string TestName;
@@ -4247,6 +4279,10 @@ public class GUBP : BuildCommand
             if (DoASharedPromotable)
             {
                 AddNode(new NonUnityTestNode(HostPlatform));
+                if (HostPlatform == UnrealTargetPlatform.Win64)
+                {
+                    AddNode(new IOSOnPCTestNode());
+                }
 
                 var AgentSharingGroup = "Shared_EditorTests" + HostPlatformNode.StaticGetHostPlatformSuffix(HostPlatform);
 
@@ -4748,10 +4784,12 @@ if (HostPlatform == UnrealTargetPlatform.Mac) continue; //temp hack till mac aut
         GUBPNodesControllingTriggerDotName = new Dictionary<string, string>();
 
         var FullNodeList = new Dictionary<string, string>();
+        var FullNodeListSortKey = new Dictionary<string, int>();
         var FullNodeDirectDependencies = new Dictionary<string, string>();
         {
             Log("******* {0} GUBP Nodes", GUBPNodes.Count);
             var SortedNodes = TopologicalSort(new HashSet<string>(GUBPNodes.Keys), LocalOnly: true, DoNotConsiderCompletion: true);
+            int Count = 0;
             foreach (var Node in SortedNodes)
             {
                 string Note = GetControllingTriggerDotName(Node);
@@ -4778,6 +4816,8 @@ if (HostPlatform == UnrealTargetPlatform.Mac) continue; //temp hack till mac aut
                     Log("  {0}: {1}      {2}", Node, Note, All);
                     FullNodeList.Add(Node, Note);
                     FullNodeDirectDependencies.Add(Node, All);
+                    FullNodeListSortKey.Add(Node, Count);
+                    Count++;
                 }
                 else
                 {
@@ -5152,6 +5192,10 @@ if (HostPlatform == UnrealTargetPlatform.Mac) continue; //temp hack till mac aut
             {
                 ECProps.Add(string.Format("DirectDependencies/{0}={1}", NodePair.Key, NodePair.Value));
             }
+            foreach (var NodePair in FullNodeListSortKey)
+            {
+                ECProps.Add(string.Format("SortKey/{0}={1}", NodePair.Key, NodePair.Value));
+            }
 
             var ECJobProps = new List<string>();
             if (ExplicitTrigger != "")
