build script, disable mac binary temp storage warnings for now

diff --git a/Engine/Source/Programs/AutomationTool/TempStorage.cs b/Engine/Source/Programs/AutomationTool/TempStorage.cs
--- a/Engine/Source/Programs/AutomationTool/TempStorage.cs
+++ b/Engine/Source/Programs/AutomationTool/TempStorage.cs
@@ -63,6 +63,9 @@ namespace AutomationTool
                     bSilentOkToBeDifferent = bSilentOkToBeDifferent || (Name == "Engine/Binaries/DotNET/UnrealBuildTool.exe");
                     bSilentOkToBeDifferent = bSilentOkToBeDifferent || (Name == "Engine/Binaries/DotNET/UnrealBuildTool.exe.config");
 
+                    // Lets just allow all mac warnings to be slient
+                    bSilentOkToBeDifferent = bSilentOkToBeDifferent || Name.Contains("Engine/Binaries/Mac");
+
                     System.Diagnostics.TraceEventType LogType = bOkToBeDifferent ? System.Diagnostics.TraceEventType.Warning : System.Diagnostics.TraceEventType.Error;
                     if (bSilentOkToBeDifferent && bOkToBeDifferent)
                     {
