UAT: Fixing process exit event not being triggered. Re-enabled automation testing in GUBP

diff --git a/Engine/Source/Programs/AutomationTool/ProcessUtils.cs b/Engine/Source/Programs/AutomationTool/ProcessUtils.cs
--- a/Engine/Source/Programs/AutomationTool/ProcessUtils.cs
+++ b/Engine/Source/Programs/AutomationTool/ProcessUtils.cs
@@ -45,7 +45,7 @@ namespace AutomationTool
 		/// Creates a new process and adds it to the tracking list.
 		/// </summary>
 		/// <returns>New Process objects</returns>
-		public static ProcessResult CreateProcess(bool bAllowSpew, string LogName, Dictionary<string, string> Env = null, TraceEventType SpewVerbosity = TraceEventType.Information)
+		public static ProcessResult CreateProcess(string AppName, bool bAllowSpew, string LogName, Dictionary<string, string> Env = null, TraceEventType SpewVerbosity = TraceEventType.Information)
 		{
 			var NewProcess = HostPlatform.Current.CreateProcess(LogName);
 			if (Env != null)
@@ -62,8 +62,7 @@ namespace AutomationTool
 					}
 				}
 			}
-			var Result = new ProcessResult(NewProcess, bAllowSpew, LogName, SpewVerbosity:SpewVerbosity);
-			NewProcess.Exited += NewProcess_Exited;
+			var Result = new ProcessResult(AppName, NewProcess, bAllowSpew, LogName, SpewVerbosity: SpewVerbosity);
 			lock (SyncObject)
 			{
 				ActiveProcesses.Add(Result);
@@ -71,6 +70,14 @@ namespace AutomationTool
 			return Result;
 		}
 
+		public static void RemoveProcess(ProcessResult Proc)
+		{
+			lock (SyncObject)
+			{
+				ActiveProcesses.Remove(Proc);
+			}
+		}
+
 		public static bool CanBeKilled(string ProcessName)
 		{
 			return !HostPlatform.Current.DontKillProcessList.Contains(ProcessName, StringComparer.InvariantCultureIgnoreCase);
@@ -100,10 +107,7 @@ namespace AutomationTool
 			CommandUtils.Log("Trying to kill {0} spawned processes.", ProcessesToKill.Count);
 			foreach (var Proc in ProcessesToKill)
 			{
-				if (!Proc.HasExited)
-				{
-					CommandUtils.Log("  {0}", Proc.GetProcessName());
-				}
+				CommandUtils.Log("  {0}", Proc.GetProcessName());
 			}
             if (CommandUtils.IsBuildMachine)
             {
@@ -177,27 +181,6 @@ namespace AutomationTool
             }
 
 		}
-
-		/// <summary>
-		/// Removes a process from the list of tracked processes.
-		/// </summary>
-		/// <param name="sender"></param>
-		/// <param name="e"></param>
-		public static void NewProcess_Exited(object sender, EventArgs e)
-		{
-			var Proc = (Process)sender;
-			lock (SyncObject)
-			{
-				for (int Index = ActiveProcesses.Count - 1; Index >= 0; --Index)
-				{
-					if (ActiveProcesses[Index].ProcessObject == Proc)
-					{
-						ActiveProcesses.RemoveAt(Index);
-						break;
-					}
-				}
-			}
-		}
 	}
 
 	#region ProcessResult Helper Class
@@ -207,23 +190,46 @@ namespace AutomationTool
 	/// </summary>
 	public class ProcessResult
 	{
-		string Source = "";
-		int ProcessExitCode = -1;
-		StringBuilder ProcessOutput = new StringBuilder();
-		bool AllowSpew = true;
-		TraceEventType SpewVerbosity = TraceEventType.Information;
+		private string Source = "";
+		private int ProcessExitCode = -1;
+		private StringBuilder ProcessOutput = new StringBuilder();
+		private bool AllowSpew = true;
+		private TraceEventType SpewVerbosity = TraceEventType.Information;
+		private string AppName = String.Empty;
 		private Process Proc = null;
 		private AutoResetEvent OutputWaitHandle = new AutoResetEvent(false);
 		private AutoResetEvent ErrorWaitHandle = new AutoResetEvent(false);
 		private object ProcSyncObject;
 
-		public ProcessResult(Process InProc, bool bAllowSpew, string LogName, TraceEventType SpewVerbosity = TraceEventType.Information)
+		public ProcessResult(string InAppName, Process InProc, bool bAllowSpew, string LogName, TraceEventType SpewVerbosity = TraceEventType.Information)
 		{
+			AppName = InAppName;
 			ProcSyncObject = new object();
 			Proc = InProc;
 			Source = LogName;
 			AllowSpew = bAllowSpew;
 			this.SpewVerbosity = SpewVerbosity;
+			Proc.EnableRaisingEvents = true;
+			Proc.Exited += ProcessExited;
+		}
+
+		/// <summary>
+		/// Removes a process from the list of tracked processes.
+		/// </summary>
+		/// <param name="sender"></param>
+		/// <param name="e"></param>
+		public void ProcessExited(object sender, EventArgs e)
+		{
+			var ProcessName = GetProcessName();
+			CommandUtils.Log("ProcessExited event triggered for {0}", ProcessName);
+			lock (ProcSyncObject)
+			{
+				if (Proc != null)
+				{
+					Proc.Exited -= ProcessExited;
+				}
+			}
+			ProcessManager.RemoveProcess(this);
 		}
 
 		private void LogOutput(TraceEventType Verbosity, string Message)
@@ -328,11 +334,22 @@ namespace AutomationTool
 			string Name = null;
 			lock (ProcSyncObject)
 			{
-				if (Proc != null)
+				try
+				{
+					if (Proc != null && !Proc.HasExited)
+					{
+						Name = Proc.ProcessName;
+					}
+				}
+				catch
 				{
-					Name = Proc.ProcessName;
+					// Ignore all exceptions
 				}
 			}
+			if (String.IsNullOrEmpty(Name))
+			{
+				Name = "[EXITED] " + AppName;
+			}
 			return Name;
 		}
 
@@ -501,6 +518,7 @@ namespace AutomationTool
 				Process ProcToKill = null;
 				// At this point any call to Proc memebers will result in an exception
 				// so null the object.
+				var ProcToKillName = GetProcessName();
 				lock (ProcSyncObject)
 				{
 					ProcToKill = Proc;
@@ -513,8 +531,26 @@ namespace AutomationTool
 				}
 				try
 				{
+					const int WaitTime = 10000;
+					int Attempts = 6;
 					ProcToKill.Kill();
-					ProcToKill.WaitForExit(60000);
+					do
+					{
+						ProcToKill.WaitForExit(WaitTime);
+						if (!ProcToKill.HasExited)
+						{
+							CommandUtils.Log("Process {0} did not exit after {1}s, waiting ({0})", ProcToKillName, WaitTime / 1000, Attempts);
+						}
+					}
+					while (--Attempts > 0);
+					if (!ProcToKill.HasExited)
+					{
+						CommandUtils.Log("Process {0} failed to exit.", ProcToKillName);
+					}
+					else
+					{
+						CommandUtils.Log("Process {0} successfully exited.", ProcToKillName);
+					}
 					ProcToKill.Close();
 				}
 				catch (Exception Ex)
@@ -605,11 +641,10 @@ namespace AutomationTool
             {
                 Log(SpewVerbosity,"Run: " + App + " " + (String.IsNullOrEmpty(CommandLine) ? "" : CommandLine));
             }
-			ProcessResult Result = ProcessManager.CreateProcess(Options.HasFlag(ERunOptions.AllowSpew), Path.GetFileNameWithoutExtension(App), Env, SpewVerbosity:SpewVerbosity);
+			ProcessResult Result = ProcessManager.CreateProcess(App, Options.HasFlag(ERunOptions.AllowSpew), Path.GetFileNameWithoutExtension(App), Env, SpewVerbosity:SpewVerbosity);
 			Process Proc = Result.ProcessObject;
 
-			bool bRedirectStdOut = (Options & ERunOptions.NoStdOutRedirect) != ERunOptions.NoStdOutRedirect;
-			Proc.EnableRaisingEvents = false;
+			bool bRedirectStdOut = (Options & ERunOptions.NoStdOutRedirect) != ERunOptions.NoStdOutRedirect;			
 			Proc.StartInfo.FileName = App;
 			Proc.StartInfo.Arguments = String.IsNullOrEmpty(CommandLine) ? "" : CommandLine;
 			Proc.StartInfo.UseShellExecute = false;
@@ -655,7 +690,7 @@ namespace AutomationTool
 				{
 					// Mono's detection of process exit status is broken. It doesn't clean up after the process and it doesn't call Exited callback.
 					Proc.Dispose();
-					ProcessManager.NewProcess_Exited(Proc, null);
+					Result.ProcessExited(Proc, null);
 				}
 			}
 			else
