Cleaned up P4Utils.Changes handling of long descriptions

diff --git a/Engine/Source/Programs/AutomationTool/P4Utils.cs b/Engine/Source/Programs/AutomationTool/P4Utils.cs
--- a/Engine/Source/Programs/AutomationTool/P4Utils.cs
+++ b/Engine/Source/Programs/AutomationTool/P4Utils.cs
@@ -597,15 +597,6 @@ namespace AutomationTool
         static Dictionary<string, List<ChangeRecord>> ChangesCache = new Dictionary<string, List<ChangeRecord>>();
         public bool Changes(out List<ChangeRecord> ChangeRecords, string CommandLine, bool AllowSpew = true, bool UseCaching = false, bool LongComment = false)
         {
-            // If the user specified '-l' or '-L', the summary will appear on subsequent lines (no quotes) instead of the same line (surrounded by single quotes)
-            bool bSummaryIsOnSameLine = CommandLine.IndexOf(" -L ", StringComparison.InvariantCultureIgnoreCase) == -1 && !CommandLine.StartsWith("-L ",StringComparison.InvariantCultureIgnoreCase);
-            if (bSummaryIsOnSameLine && LongComment)
-            {
-                CommandLine = "-L " + CommandLine;
-                bSummaryIsOnSameLine = false;
-            }
-
-
             if (UseCaching && ChangesCache.ContainsKey(CommandLine))
             {
                 ChangeRecords = ChangesCache[CommandLine];
@@ -617,14 +608,20 @@ namespace AutomationTool
             {
                 // Change 1999345 on 2014/02/16 by buildmachine@BuildFarm_BUILD-23_buildmachine_++depot+UE4 'GUBP Node Shadow_LabelPromotabl'
 
+				// If the user specified '-l' or '-L', the summary will appear on subsequent lines (no quotes) instead of the same line (surrounded by single quotes)
+				bool bSummaryIsOnSameLine = CommandLine.IndexOf( "-L", StringComparison.InvariantCultureIgnoreCase ) == -1;
+				if( bSummaryIsOnSameLine && LongComment )
+				{
+					CommandLine = "-L " + CommandLine;
+					bSummaryIsOnSameLine = false;
+				}
+
                 string Output;
                 if (!LogP4Output(out Output, "changes " + CommandLine, null, AllowSpew))
                 {
                     throw new AutomationException("P4 returned failure.");
                 }
-                
-                //CommandUtils.Log("{0}", Output.Replace("\t", "[tab]").Replace("\r", "[r]").Replace("\n", "[n]"));
-                
+
                 var Lines = Output.Split(new string[] { Environment.NewLine }, StringSplitOptions.None);
                 for(int LineIndex = 0; LineIndex < Lines.Length; ++LineIndex)
                 {
@@ -644,7 +641,6 @@ namespace AutomationTool
                     int ChangeAt = Line.IndexOf(MatchChange);
                     int OnAt = Line.IndexOf(MatchOn);
                     int ByAt = Line.IndexOf(MatchBy);
-                    int AtAt = Line.IndexOf("@");
                     if (ChangeAt == 0 && OnAt > ChangeAt && ByAt > OnAt)
                     {
                         var ChangeString = Line.Substring(ChangeAt + MatchChange.Length, OnAt - ChangeAt - MatchChange.Length);
@@ -653,6 +649,7 @@ namespace AutomationTool
                         {
                             throw new AutomationException("weird CL {0} in {1}", Change.CL, Line);
                         }
+	                    int AtAt = Line.IndexOf("@");
                         Change.User = Line.Substring(ByAt + MatchBy.Length, AtAt - ByAt - MatchBy.Length);
 
 						if( bSummaryIsOnSameLine )
@@ -683,44 +680,29 @@ namespace AutomationTool
 							{
 								Line = Lines[ LineIndex ];
 
-                                bool NextLineIsTab = false;
-
-                                for (int LookAhead = LineIndex; LookAhead < Lines.Length; LookAhead++ )
-                                {
-                                    if (Lines[LookAhead].StartsWith("\t"))
-                                    {
-                                        NextLineIsTab = true;
-                                        break;
-                                    }
-                                    if (Lines[LookAhead] != "")
-                                    {
-                                        break;
-                                    }
-                                }
-
-                                if (String.IsNullOrEmpty(Line))
-                                {
-                                    if (NextLineIsTab)
-                                    {
-                                        // Add a CR if we already had some summary text
-                                        if (!String.IsNullOrEmpty(Change.Summary))
-                                        {
-                                            Change.Summary += "\n";
-                                        } 
-                                        continue;
-                                    }
-                                    // Summaries end with a blank line (no tabs)
-                                    break;
-                                }
+								int SummaryChangeAt = Line.IndexOf(MatchChange);
+								int SummaryOnAt = Line.IndexOf(MatchOn);
+								int SummaryByAt = Line.IndexOf(MatchBy);
+								int SummaryAtAt = Line.IndexOf("@");
+								if (SummaryChangeAt == 0 && SummaryOnAt > SummaryChangeAt && SummaryByAt > SummaryOnAt)
+								{
+									// OK, we found a new change. This isn't part of our summary.  We're done with the summary.  Back we go.
+									--LineIndex;
+									break;
+								}
 
 								// Summary lines are supposed to begin with a single tab character (even empty lines)
-								if( Line[0] != '\t' )
+								if( !String.IsNullOrEmpty( Line ) && Line[0] != '\t' )
 								{
-									throw new AutomationException("Was expecting every line of the P4 changes summary to start with a tab character");
+									throw new AutomationException("Was expecting every line of the P4 changes summary to start with a tab character or be totally empty");
 								}
 
 								// Remove the tab
-								var SummaryLine = Line.Substring( 1 );
+								var SummaryLine = Line;
+								if( Line.StartsWith( "\t" ) )
+								{ 
+									SummaryLine = Line.Substring( 1 );
+								}
 
 								// Add a CR if we already had some summary text
 								if( !String.IsNullOrEmpty( Change.Summary ) )
@@ -741,7 +723,7 @@ namespace AutomationTool
 					}
                 }
             }
-            catch (Exception Ex)
+			catch (Exception Ex)
             {
                 CommandUtils.Log(System.Diagnostics.TraceEventType.Warning, "Unable to get P4 changes with {0}", CommandLine);
                 CommandUtils.Log(System.Diagnostics.TraceEventType.Warning, " Exception was {0}", LogUtils.FormatException(Ex));
