Back out changelist 2064485

diff --git a/Engine/Source/Programs/AutomationTool/ProcessUtils.cs b/Engine/Source/Programs/AutomationTool/ProcessUtils.cs
--- a/Engine/Source/Programs/AutomationTool/ProcessUtils.cs
+++ b/Engine/Source/Programs/AutomationTool/ProcessUtils.cs
@@ -45,7 +45,7 @@ namespace AutomationTool
 		/// Creates a new process and adds it to the tracking list.
 		/// </summary>
 		/// <returns>New Process objects</returns>
-		public static ProcessResult CreateProcess(string AppName, bool bAllowSpew, string LogName, Dictionary<string, string> Env = null, TraceEventType SpewVerbosity = TraceEventType.Information)
+		public static ProcessResult CreateProcess(bool bAllowSpew, string LogName, Dictionary<string, string> Env = null, TraceEventType SpewVerbosity = TraceEventType.Information)
 		{
 			var NewProcess = HostPlatform.Current.CreateProcess(LogName);
 			if (Env != null)
@@ -62,7 +62,8 @@ namespace AutomationTool
 					}
 				}
 			}
-			var Result = new ProcessResult(AppName, NewProcess, bAllowSpew, LogName, SpewVerbosity: SpewVerbosity);
+			var Result = new ProcessResult(NewProcess, bAllowSpew, LogName, SpewVerbosity:SpewVerbosity);
+			NewProcess.Exited += NewProcess_Exited;
 			lock (SyncObject)
 			{
 				ActiveProcesses.Add(Result);
@@ -70,14 +71,6 @@ namespace AutomationTool
 			return Result;
 		}
 
-		public static void RemoveProcess(ProcessResult Proc)
-		{
-			lock (SyncObject)
-			{
-				ActiveProcesses.Remove(Proc);
-			}
-		}
-
 		public static bool CanBeKilled(string ProcessName)
 		{
 			return !HostPlatform.Current.DontKillProcessList.Contains(ProcessName, StringComparer.InvariantCultureIgnoreCase);
@@ -107,7 +100,10 @@ namespace AutomationTool
 			CommandUtils.Log("Trying to kill {0} spawned processes.", ProcessesToKill.Count);
 			foreach (var Proc in ProcessesToKill)
 			{
-				CommandUtils.Log("  {0}", Proc.GetProcessName());
+				if (!Proc.HasExited)
+				{
+					CommandUtils.Log("  {0}", Proc.GetProcessName());
+				}
 			}
             if (CommandUtils.IsBuildMachine)
             {
@@ -181,6 +177,27 @@ namespace AutomationTool
             }
 
 		}
+
+		/// <summary>
+		/// Removes a process from the list of tracked processes.
+		/// </summary>
+		/// <param name="sender"></param>
+		/// <param name="e"></param>
+		public static void NewProcess_Exited(object sender, EventArgs e)
+		{
+			var Proc = (Process)sender;
+			lock (SyncObject)
+			{
+				for (int Index = ActiveProcesses.Count - 1; Index >= 0; --Index)
+				{
+					if (ActiveProcesses[Index].ProcessObject == Proc)
+					{
+						ActiveProcesses.RemoveAt(Index);
+						break;
+					}
+				}
+			}
+		}
 	}
 
 	#region ProcessResult Helper Class
@@ -190,46 +207,23 @@ namespace AutomationTool
 	/// </summary>
 	public class ProcessResult
 	{
-		private string Source = "";
-		private int ProcessExitCode = -1;
-		private StringBuilder ProcessOutput = new StringBuilder();
-		private bool AllowSpew = true;
-		private TraceEventType SpewVerbosity = TraceEventType.Information;
-		private string AppName = String.Empty;
+		string Source = "";
+		int ProcessExitCode = -1;
+		StringBuilder ProcessOutput = new StringBuilder();
+		bool AllowSpew = true;
+		TraceEventType SpewVerbosity = TraceEventType.Information;
 		private Process Proc = null;
 		private AutoResetEvent OutputWaitHandle = new AutoResetEvent(false);
 		private AutoResetEvent ErrorWaitHandle = new AutoResetEvent(false);
 		private object ProcSyncObject;
 
-		public ProcessResult(string InAppName, Process InProc, bool bAllowSpew, string LogName, TraceEventType SpewVerbosity = TraceEventType.Information)
+		public ProcessResult(Process InProc, bool bAllowSpew, string LogName, TraceEventType SpewVerbosity = TraceEventType.Information)
 		{
-			AppName = InAppName;
 			ProcSyncObject = new object();
 			Proc = InProc;
 			Source = LogName;
 			AllowSpew = bAllowSpew;
 			this.SpewVerbosity = SpewVerbosity;
-			Proc.EnableRaisingEvents = true;
-			Proc.Exited += ProcessExited;
-		}
-
-		/// <summary>
-		/// Removes a process from the list of tracked processes.
-		/// </summary>
-		/// <param name="sender"></param>
-		/// <param name="e"></param>
-		public void ProcessExited(object sender, EventArgs e)
-		{
-			var ProcessName = GetProcessName();
-			CommandUtils.Log("ProcessExited event triggered for {0}", ProcessName);
-			lock (ProcSyncObject)
-			{
-				if (Proc != null)
-				{
-					Proc.Exited -= ProcessExited;
-				}
-			}
-			ProcessManager.RemoveProcess(this);
 		}
 
 		private void LogOutput(TraceEventType Verbosity, string Message)
@@ -334,22 +328,11 @@ namespace AutomationTool
 			string Name = null;
 			lock (ProcSyncObject)
 			{
-				try
-				{
-					if (Proc != null && !Proc.HasExited)
-					{
-						Name = Proc.ProcessName;
-					}
-				}
-				catch
+				if (Proc != null)
 				{
-					// Ignore all exceptions
+					Name = Proc.ProcessName;
 				}
 			}
-			if (String.IsNullOrEmpty(Name))
-			{
-				Name = "[EXITED] " + AppName;
-			}
 			return Name;
 		}
 
@@ -518,7 +501,6 @@ namespace AutomationTool
 				Process ProcToKill = null;
 				// At this point any call to Proc memebers will result in an exception
 				// so null the object.
-				var ProcToKillName = GetProcessName();
 				lock (ProcSyncObject)
 				{
 					ProcToKill = Proc;
@@ -531,26 +513,8 @@ namespace AutomationTool
 				}
 				try
 				{
-					const int WaitTime = 10000;
-					int Attempts = 6;
 					ProcToKill.Kill();
-					do
-					{
-						ProcToKill.WaitForExit(WaitTime);
-						if (!ProcToKill.HasExited)
-						{
-							CommandUtils.Log("Process {0} did not exit after {1}s, waiting ({0})", ProcToKillName, WaitTime / 1000, Attempts);
-						}
-					}
-					while (--Attempts > 0);
-					if (!ProcToKill.HasExited)
-					{
-						CommandUtils.Log("Process {0} failed to exit.", ProcToKillName);
-					}
-					else
-					{
-						CommandUtils.Log("Process {0} successfully exited.", ProcToKillName);
-					}
+					ProcToKill.WaitForExit(60000);
 					ProcToKill.Close();
 				}
 				catch (Exception Ex)
@@ -641,10 +605,11 @@ namespace AutomationTool
             {
                 Log(SpewVerbosity,"Run: " + App + " " + (String.IsNullOrEmpty(CommandLine) ? "" : CommandLine));
             }
-			ProcessResult Result = ProcessManager.CreateProcess(App, Options.HasFlag(ERunOptions.AllowSpew), Path.GetFileNameWithoutExtension(App), Env, SpewVerbosity:SpewVerbosity);
+			ProcessResult Result = ProcessManager.CreateProcess(Options.HasFlag(ERunOptions.AllowSpew), Path.GetFileNameWithoutExtension(App), Env, SpewVerbosity:SpewVerbosity);
 			Process Proc = Result.ProcessObject;
 
-			bool bRedirectStdOut = (Options & ERunOptions.NoStdOutRedirect) != ERunOptions.NoStdOutRedirect;			
+			bool bRedirectStdOut = (Options & ERunOptions.NoStdOutRedirect) != ERunOptions.NoStdOutRedirect;
+			Proc.EnableRaisingEvents = false;
 			Proc.StartInfo.FileName = App;
 			Proc.StartInfo.Arguments = String.IsNullOrEmpty(CommandLine) ? "" : CommandLine;
 			Proc.StartInfo.UseShellExecute = false;
@@ -690,7 +655,7 @@ namespace AutomationTool
 				{
 					// Mono's detection of process exit status is broken. It doesn't clean up after the process and it doesn't call Exited callback.
 					Proc.Dispose();
-					Result.ProcessExited(Proc, null);
+					ProcessManager.NewProcess_Exited(Proc, null);
 				}
 			}
 			else
