Fix for incorrect editor executable name when it can't be found (was showing PC executable names on the Mac)
#ue4
#osx

diff --git a/Engine/Source/Programs/AutomationTool/Scripts/CookCommand.Automation.cs b/Engine/Source/Programs/AutomationTool/Scripts/CookCommand.Automation.cs
--- a/Engine/Source/Programs/AutomationTool/Scripts/CookCommand.Automation.cs
+++ b/Engine/Source/Programs/AutomationTool/Scripts/CookCommand.Automation.cs
@@ -29,7 +29,7 @@ public partial class Project : CommandUtils
 		string UE4EditorExe = HostPlatform.Current.GetUE4ExePath(Params.UE4Exe);
 		if (!FileExists(UE4EditorExe))
 		{
-			throw new AutomationException("Missing " + Params.UE4Exe + " executable. Needs to be built first.");
+			throw new AutomationException("Missing " + UE4EditorExe + " executable. Needs to be built first.");
 		}
 
 		var PlatformsToCook = new List<string>();
