build script, change warnings to information on missing codeless projects

diff --git a/Engine/Source/Programs/AutomationTool/Scripts/GUBP.Automation.cs b/Engine/Source/Programs/AutomationTool/Scripts/GUBP.Automation.cs
--- a/Engine/Source/Programs/AutomationTool/Scripts/GUBP.Automation.cs
+++ b/Engine/Source/Programs/AutomationTool/Scripts/GUBP.Automation.cs
@@ -4229,7 +4229,7 @@ public class GUBP : BuildCommand
                     var Proj = Branch.FindGame(Codeless.Key);
                     if (Proj == null)
                     {
-                        Log(System.Diagnostics.TraceEventType.Warning, "{0} was listed as a codeless project by GUBP_NonCodeProjects_BaseEditorTypeOnly, however it does not exist in this branch.", Codeless.Key);
+                        Log(System.Diagnostics.TraceEventType.Information, "{0} was listed as a codeless project by GUBP_NonCodeProjects_BaseEditorTypeOnly, however it does not exist in this branch.", Codeless.Key);
                     }
                     else if (Proj.Properties.bIsCodeBasedProject)
                     {
@@ -4264,7 +4264,7 @@ public class GUBP : BuildCommand
                     }
                     else
                     {
-                        Log(System.Diagnostics.TraceEventType.Warning, "{0} was listed as a codeless formal build GUBP_NonCodeFormalBuilds_BaseEditorTypeOnly, however it does not exist in this branch.", Codeless.Key);
+                        Log(System.Diagnostics.TraceEventType.Information, "{0} was listed as a codeless formal build GUBP_NonCodeFormalBuilds_BaseEditorTypeOnly, however it does not exist in this branch.", Codeless.Key);
                     }
                 }
             }
