Add oobe.exe to list of NottoKill

diff --git a/Engine/Source/Programs/AutomationTool/WindowsHostPlatform.cs b/Engine/Source/Programs/AutomationTool/WindowsHostPlatform.cs
--- a/Engine/Source/Programs/AutomationTool/WindowsHostPlatform.cs
+++ b/Engine/Source/Programs/AutomationTool/WindowsHostPlatform.cs
@@ -217,7 +217,8 @@ namespace AutomationTool
 			"svchost",
 			"services",
 			"lsass",
-            "conhost"
+            "conhost",
+            "oobe"
 		};
 		public override string[] DontKillProcessList
 		{
