UAT: Another attempt to fix EC hangs. Re-enabling automated tests in GUBP

diff --git a/Engine/Source/Programs/AutomationTool/ProcessUtils.cs b/Engine/Source/Programs/AutomationTool/ProcessUtils.cs
--- a/Engine/Source/Programs/AutomationTool/ProcessUtils.cs
+++ b/Engine/Source/Programs/AutomationTool/ProcessUtils.cs
@@ -45,7 +45,7 @@ namespace AutomationTool
 		/// Creates a new process and adds it to the tracking list.
 		/// </summary>
 		/// <returns>New Process objects</returns>
-		public static ProcessResult CreateProcess(bool bAllowSpew, string LogName, Dictionary<string, string> Env = null, TraceEventType SpewVerbosity = TraceEventType.Information)
+		public static ProcessResult CreateProcess(string AppName, bool bAllowSpew, string LogName, Dictionary<string, string> Env = null, TraceEventType SpewVerbosity = TraceEventType.Information)
 		{
 			var NewProcess = HostPlatform.Current.CreateProcess(LogName);
 			if (Env != null)
@@ -62,8 +62,7 @@ namespace AutomationTool
 					}
 				}
 			}
-			var Result = new ProcessResult(NewProcess, bAllowSpew, LogName, SpewVerbosity:SpewVerbosity);
-			NewProcess.Exited += NewProcess_Exited;
+			var Result = new ProcessResult(AppName, NewProcess, bAllowSpew, LogName, SpewVerbosity: SpewVerbosity);
 			lock (SyncObject)
 			{
 				ActiveProcesses.Add(Result);
@@ -71,6 +70,14 @@ namespace AutomationTool
 			return Result;
 		}
 
+		public static void RemoveProcess(ProcessResult Proc)
+		{
+			lock (SyncObject)
+			{
+				ActiveProcesses.Remove(Proc);
+			}
+		}
+
 		public static bool CanBeKilled(string ProcessName)
 		{
 			return !HostPlatform.Current.DontKillProcessList.Contains(ProcessName, StringComparer.InvariantCultureIgnoreCase);
@@ -81,10 +88,16 @@ namespace AutomationTool
 		/// </summary>
 		public static void KillAll()
 		{
-			List<ProcessResult> ProcessesToKill = null;
+			List<ProcessResult> ProcessesToKill = new List<ProcessResult>();
 			lock (SyncObject)
 			{
-				ProcessesToKill = new List<ProcessResult>(ActiveProcesses);
+				foreach (var ProcResult in ActiveProcesses)
+				{
+					if (!ProcResult.HasExited)
+					{
+						ProcessesToKill.Add(ProcResult);
+					}
+				}
 				ActiveProcesses.Clear();
 			}
 			// Remove processes that can't be killed
@@ -100,10 +113,7 @@ namespace AutomationTool
 			CommandUtils.Log("Trying to kill {0} spawned processes.", ProcessesToKill.Count);
 			foreach (var Proc in ProcessesToKill)
 			{
-				if (!Proc.HasExited)
-				{
-					CommandUtils.Log("  {0}", Proc.GetProcessName());
-				}
+				CommandUtils.Log("  {0}", Proc.GetProcessName());
 			}
             if (CommandUtils.IsBuildMachine)
             {
@@ -149,17 +159,18 @@ namespace AutomationTool
             }
 			foreach (var Proc in ProcessesToKill)
 			{
+				var ProcName = Proc.GetProcessName();
 				try
 				{
 					if (!Proc.HasExited)
 					{
-						CommandUtils.Log(TraceEventType.Information, "Killing process: {0}", Proc.GetProcessName());
+						CommandUtils.Log(TraceEventType.Information, "Killing process: {0}", ProcName);
 						Proc.StopProcess(false);
 					}
 				}
 				catch (Exception Ex)
 				{
-					CommandUtils.Log(TraceEventType.Warning, "Exception while trying to kill process.");
+					CommandUtils.Log(TraceEventType.Warning, "Exception while trying to kill process {0}:", ProcName);
 					CommandUtils.Log(TraceEventType.Warning, Ex);
 				}
 			}
@@ -177,27 +188,6 @@ namespace AutomationTool
             }
 
 		}
-
-		/// <summary>
-		/// Removes a process from the list of tracked processes.
-		/// </summary>
-		/// <param name="sender"></param>
-		/// <param name="e"></param>
-		public static void NewProcess_Exited(object sender, EventArgs e)
-		{
-			var Proc = (Process)sender;
-			lock (SyncObject)
-			{
-				for (int Index = ActiveProcesses.Count - 1; Index >= 0; --Index)
-				{
-					if (ActiveProcesses[Index].ProcessObject == Proc)
-					{
-						ActiveProcesses.RemoveAt(Index);
-						break;
-					}
-				}
-			}
-		}
 	}
 
 	#region ProcessResult Helper Class
@@ -207,23 +197,34 @@ namespace AutomationTool
 	/// </summary>
 	public class ProcessResult
 	{
-		string Source = "";
-		int ProcessExitCode = -1;
-		StringBuilder ProcessOutput = new StringBuilder();
-		bool AllowSpew = true;
-		TraceEventType SpewVerbosity = TraceEventType.Information;
+		private string Source = "";
+		private int ProcessExitCode = -1;
+		private StringBuilder ProcessOutput = new StringBuilder();
+		private bool AllowSpew = true;
+		private TraceEventType SpewVerbosity = TraceEventType.Information;
+		private string AppName = String.Empty;
 		private Process Proc = null;
 		private AutoResetEvent OutputWaitHandle = new AutoResetEvent(false);
 		private AutoResetEvent ErrorWaitHandle = new AutoResetEvent(false);
 		private object ProcSyncObject;
 
-		public ProcessResult(Process InProc, bool bAllowSpew, string LogName, TraceEventType SpewVerbosity = TraceEventType.Information)
+		public ProcessResult(string InAppName, Process InProc, bool bAllowSpew, string LogName, TraceEventType SpewVerbosity = TraceEventType.Information)
 		{
+			AppName = InAppName;
 			ProcSyncObject = new object();
 			Proc = InProc;
 			Source = LogName;
 			AllowSpew = bAllowSpew;
 			this.SpewVerbosity = SpewVerbosity;
+			Proc.EnableRaisingEvents = false;
+		}
+
+		/// <summary>
+		/// Removes a process from the list of tracked processes.
+		/// </summary>
+		public void OnProcessExited()
+		{
+			ProcessManager.RemoveProcess(this);
 		}
 
 		private void LogOutput(TraceEventType Verbosity, string Message)
@@ -311,7 +312,18 @@ namespace AutomationTool
 
 		public bool HasExited
 		{
-			get { return Proc != null ? Proc.HasExited : true; }
+			get 
+			{
+				bool bHasExited = true;
+				lock (ProcSyncObject)
+				{
+					if (Proc != null)
+					{
+						bHasExited = Proc.HasExited;
+					}
+				}
+				return bHasExited; 
+			}
 		}
 
 		public Process ProcessObject
@@ -328,10 +340,21 @@ namespace AutomationTool
 			string Name = null;
 			lock (ProcSyncObject)
 			{
-				if (Proc != null)
+				try
 				{
-					Name = Proc.ProcessName;
+					if (Proc != null && !Proc.HasExited)
+					{
+						Name = Proc.ProcessName;
+					}
 				}
+				catch
+				{
+					// Ignore all exceptions
+				}
+			}
+			if (String.IsNullOrEmpty(Name))
+			{
+				Name = "[EXITED] " + AppName;
 			}
 			return Name;
 		}
@@ -501,6 +524,7 @@ namespace AutomationTool
 				Process ProcToKill = null;
 				// At this point any call to Proc memebers will result in an exception
 				// so null the object.
+				var ProcToKillName = GetProcessName();
 				lock (ProcSyncObject)
 				{
 					ProcToKill = Proc;
@@ -515,14 +539,22 @@ namespace AutomationTool
 				{
 					ProcToKill.Kill();
 					ProcToKill.WaitForExit(60000);
-					ProcToKill.Close();
+					if (!ProcToKill.HasExited)
+					{
+						CommandUtils.Log("Process {0} failed to exit.", ProcToKillName);
+					}
+					else
+					{
+						CommandUtils.Log("Process {0} successfully exited.", ProcToKillName);
+						OnProcessExited();
+					}
+					ProcToKill.Close();					
 				}
 				catch (Exception Ex)
 				{
-					CommandUtils.Log(TraceEventType.Warning, "Exception while trying to kill process.");
+					CommandUtils.Log(TraceEventType.Warning, "Exception while trying to kill process {0}:", ProcToKillName);
 					CommandUtils.Log(TraceEventType.Warning, Ex);
 				}
-
 			}
 		}
 	}
@@ -605,11 +637,10 @@ namespace AutomationTool
             {
                 Log(SpewVerbosity,"Run: " + App + " " + (String.IsNullOrEmpty(CommandLine) ? "" : CommandLine));
             }
-			ProcessResult Result = ProcessManager.CreateProcess(Options.HasFlag(ERunOptions.AllowSpew), Path.GetFileNameWithoutExtension(App), Env, SpewVerbosity:SpewVerbosity);
+			ProcessResult Result = ProcessManager.CreateProcess(App, Options.HasFlag(ERunOptions.AllowSpew), Path.GetFileNameWithoutExtension(App), Env, SpewVerbosity:SpewVerbosity);
 			Process Proc = Result.ProcessObject;
 
-			bool bRedirectStdOut = (Options & ERunOptions.NoStdOutRedirect) != ERunOptions.NoStdOutRedirect;
-			Proc.EnableRaisingEvents = false;
+			bool bRedirectStdOut = (Options & ERunOptions.NoStdOutRedirect) != ERunOptions.NoStdOutRedirect;			
 			Proc.StartInfo.FileName = App;
 			Proc.StartInfo.Arguments = String.IsNullOrEmpty(CommandLine) ? "" : CommandLine;
 			Proc.StartInfo.UseShellExecute = false;
@@ -651,11 +682,11 @@ namespace AutomationTool
 
                 }
 				Result.ExitCode = Proc.ExitCode;
+				Result.OnProcessExited();
 				if (UnrealBuildTool.Utils.IsRunningOnMono)
 				{
 					// Mono's detection of process exit status is broken. It doesn't clean up after the process and it doesn't call Exited callback.
-					Proc.Dispose();
-					ProcessManager.NewProcess_Exited(Proc, null);
+					Proc.Dispose();					
 				}
 			}
 			else
diff --git a/Engine/Source/Programs/AutomationTool/Program.cs b/Engine/Source/Programs/AutomationTool/Program.cs
--- a/Engine/Source/Programs/AutomationTool/Program.cs
+++ b/Engine/Source/Programs/AutomationTool/Program.cs
@@ -60,6 +60,12 @@ namespace AutomationTool
 			CommandUtils.ClearDirStack();
 			Environment.ExitCode = Result;
 
+			// Try to kill process before app domain exits to leave the other KillAll call to extreme edge cases
+			if (ShouldKillProcesses)
+			{
+				ProcessManager.KillAll();
+			}
+
 			Log.WriteLine(TraceEventType.Information, "AutomationTool exiting with ExitCode={0}", Result);
 			return Result;
 		}
diff --git a/Engine/Source/Programs/AutomationTool/Scripts/Tests.Automation.cs b/Engine/Source/Programs/AutomationTool/Scripts/Tests.Automation.cs
--- a/Engine/Source/Programs/AutomationTool/Scripts/Tests.Automation.cs
+++ b/Engine/Source/Programs/AutomationTool/Scripts/Tests.Automation.cs
@@ -1766,9 +1766,34 @@ class TestKillAll : BuildCommand
 
 		string Exe = CombinePaths(CmdEnv.LocalRoot, "Engine", "Binaries", "Win64", "UE4Editor.exe");
 		string ClientLogFile = CombinePaths(CmdEnv.LogFolder, "HoverGameRun");
-		string CmdLine = " ../../../Samples/HoverShip/HoverShip.uproject -game -forcelogflush -log -abslog=" + ClientLogFile;
+		string CmdLine = " ../../../Samples/Sandbox/HoverShip/HoverShip.uproject -game -forcelogflush -log -abslog=" + ClientLogFile;
+
+		Run(Exe, CmdLine, null, ERunOptions.AllowSpew | ERunOptions.NoWaitForExit | ERunOptions.AppMustExist | ERunOptions.NoStdOutRedirect);
+		Thread.Sleep(10000);
+	}
+}
+
+[Help("Spawns a process to test if it can be killed.")]
+class TestStopProcess : BuildCommand
+{
+	public override void ExecuteBuild()
+	{
+		Log("*********************** TestStopProcess");
+
+		string Exe = CombinePaths(CmdEnv.LocalRoot, "Engine", "Binaries", "Win64", "UE4Editor.exe");
+		string ClientLogFile = CombinePaths(CmdEnv.LogFolder, "HoverGameRun");
+		string CmdLine = " ../../../Samples/Sandbox/HoverShip/HoverShip.uproject -game -forcelogflush -log -ddc=noshared -abslog=" + ClientLogFile;
+
+		for (int Attempt = 0; Attempt < 5; ++Attempt)
+		{
+			Log("Attempt: {0}", Attempt);
+			var Proc = Run(Exe, CmdLine, null, ERunOptions.AllowSpew | ERunOptions.NoWaitForExit | ERunOptions.AppMustExist | ERunOptions.NoStdOutRedirect);
+			Thread.Sleep(10000);
+			Proc.StopProcess();
+		}
 
+		Log("One final attempt to test KillAll");
 		Run(Exe, CmdLine, null, ERunOptions.AllowSpew | ERunOptions.NoWaitForExit | ERunOptions.AppMustExist | ERunOptions.NoStdOutRedirect);
-		Thread.Sleep(1000);
+		Thread.Sleep(10000);
 	}
 }
\ No newline at end of file
