Fix packaging of project icon for Linux.

Do not assume any particular name for the icon (TTP #333350).

diff --git a/Engine/Source/Programs/AutomationTool/Linux/LinuxPlatform.Automation.cs b/Engine/Source/Programs/AutomationTool/Linux/LinuxPlatform.Automation.cs
--- a/Engine/Source/Programs/AutomationTool/Linux/LinuxPlatform.Automation.cs
+++ b/Engine/Source/Programs/AutomationTool/Linux/LinuxPlatform.Automation.cs
@@ -83,7 +83,7 @@ public abstract class BaseLinuxPlatform : Platform
             }
         }
 
-        SC.StageFiles(StagedFileType.NonUFS, SC.ProjectRoot, SC.ShortProjectName + ".png", false, null, SC.RelativeProjectRootForStage);
+        SC.StageFiles(StagedFileType.NonUFS, SC.ProjectRoot, "*.png", false, null, SC.RelativeProjectRootForStage);
     }
 
 	public override string GetCookPlatform(bool bDedicatedServer, bool bIsClientOnly, string CookFlavor)
