move Epic mobile provisions to NotForLicensees
updates to tool chain, UBT, and IPP to properly access them
#ue4
#ios

diff --git a/Engine/Source/Programs/AutomationTool/IOS/IOSPlatform.Automation.cs b/Engine/Source/Programs/AutomationTool/IOS/IOSPlatform.Automation.cs
--- a/Engine/Source/Programs/AutomationTool/IOS/IOSPlatform.Automation.cs
+++ b/Engine/Source/Programs/AutomationTool/IOS/IOSPlatform.Automation.cs
@@ -247,8 +247,20 @@ public class IOSPlatform : Platform
 		// Now do the .mobileprovision
 		//@TODO: Remove this mobileprovision copy, and move to a library approach like Xcode/codesign does
 		string SourceProvision = CombinePaths(LocalRoot, "Engine", "Build", "IOS", "UE4Game.mobileprovision");
-		string GameSourceProvision = ProjectRoot + "/Build/IOS/" + ShortProjectName + ".mobileprovision";
-		if (File.Exists(GameSourceProvision))
+		string GameSourceProvision = CombinePaths(ProjectRoot, "Build", "IOS", ShortProjectName + ".mobileprovision");
+		if (!File.Exists(GameSourceProvision))
+		{
+			GameSourceProvision = CombinePaths(ProjectRoot, "Build", "IOS", "NotForLicensees", ShortProjectName + ".mobileprovision");
+			if (File.Exists(GameSourceProvision))
+			{
+				SourceProvision = GameSourceProvision;
+			}
+			else if (!File.Exists(SourceProvision))
+			{
+				SourceProvision = CombinePaths(LocalRoot, "Engine", "Build", "IOS", "NotForLicensees", "UE4Game.mobileprovision");
+			}
+		}
+		else
 		{
 			SourceProvision = GameSourceProvision;
 		}
@@ -262,12 +274,24 @@ public class IOSPlatform : Platform
 
 		// install the distribution provision
 		SourceProvision = CombinePaths(LocalRoot, "Engine", "Build", "IOS", "UE4Game_Distro.mobileprovision");
-		GameSourceProvision = ProjectRoot + "/Build/IOS/" + ShortProjectName + "_Distro.mobileprovision";
-		if (File.Exists(GameSourceProvision))
+		GameSourceProvision = CombinePaths(ProjectRoot, "Build", "IOS", ShortProjectName + "_Distro.mobileprovision");
+		if (!File.Exists(GameSourceProvision))
+		{
+			GameSourceProvision = CombinePaths(ProjectRoot, "Build", "IOS", "NotForLicensees", ShortProjectName + "_Distro.mobileprovision");
+			if (File.Exists(GameSourceProvision))
+			{
+				SourceProvision = GameSourceProvision;
+			}
+			else if (!File.Exists(SourceProvision))
+			{
+				SourceProvision = CombinePaths(LocalRoot, "Engine", "Build", "IOS", "NotForLicensees", "UE4Game_Distro.mobileprovision");
+			}
+		}
+		else
 		{
 			SourceProvision = GameSourceProvision;
 		}
-		if (File.Exists (SourceProvision))
+		if (File.Exists(SourceProvision))
 		{
 			File.Copy (SourceProvision, Environment.GetEnvironmentVariable ("HOME") + "/Library/MobileDevice/Provisioning Profiles/" + ShortProjectName + "_Distro.mobileprovision", true);
 			FileInfo DestFileInfo = new FileInfo (Environment.GetEnvironmentVariable ("HOME") + "/Library/MobileDevice/Provisioning Profiles/" + ShortProjectName + "_Distro.mobileprovision");
@@ -516,8 +540,20 @@ public class IOSPlatform : Platform
 			if (GetCodeSignDesirability(Params))
 			{
 				string SourceProvision = CombinePaths(SC.LocalRoot, "Engine", "Build", "IOS", "UE4Game.mobileprovision");
-				string GameSourceProvision = SC.ProjectRoot + "/Build/IOS/" + SC.ShortProjectName + ".mobileprovision";
-				if (File.Exists(GameSourceProvision))
+				string GameSourceProvision = CombinePaths(SC.ProjectRoot, "Build", "IOS", SC.ShortProjectName + ".mobileprovision");
+				if (!File.Exists(GameSourceProvision))
+				{
+					GameSourceProvision = CombinePaths(SC.ProjectRoot, "Build", "IOS", "NotForLicensees", SC.ShortProjectName + ".mobileprovision");
+					if (File.Exists(GameSourceProvision))
+					{
+						SourceProvision = GameSourceProvision;
+					}
+					else if (!File.Exists(SourceProvision))
+					{
+						SourceProvision = CombinePaths(SC.LocalRoot, "Engine", "Build", "IOS", "NotForLicensees", "UE4Game.mobileprovision");
+					}
+				}
+				else
 				{
 					SourceProvision = GameSourceProvision;
 				}
diff --git a/Engine/Source/Programs/IOS/iPhonePackager/CodeSigning.cs b/Engine/Source/Programs/IOS/iPhonePackager/CodeSigning.cs
--- a/Engine/Source/Programs/IOS/iPhonePackager/CodeSigning.cs
+++ b/Engine/Source/Programs/IOS/iPhonePackager/CodeSigning.cs
@@ -115,10 +115,18 @@ namespace iPhonePackager
 
 			if (!File.Exists(MobileProvisionFilename))
 			{
-				MobileProvisionFilename = FileOperations.FindPrefixedFile(Config.EngineBuildDirectory, "UE4Game.mobileprovision");
+				MobileProvisionFilename = FileOperations.FindPrefixedFile(Config.BuildDirectory + "/NotForLicensees/", Program.GameName + ".mobileprovision");
 				if (!File.Exists(MobileProvisionFilename))
 				{
-					MobileProvisionFilename = FileOperations.FindAnyFileWithExtension(Config.BuildDirectory, ".mobileprovision");
+					MobileProvisionFilename = FileOperations.FindPrefixedFile(Config.EngineBuildDirectory, "UE4Game.mobileprovision");
+					if (!File.Exists(MobileProvisionFilename))
+					{
+						MobileProvisionFilename = FileOperations.FindPrefixedFile(Config.EngineBuildDirectory + "/NotForLicensees/", "UE4Game.mobileprovision");
+						if (!File.Exists(MobileProvisionFilename))
+						{
+							MobileProvisionFilename = FileOperations.FindAnyFileWithExtension(Config.BuildDirectory, ".mobileprovision");
+						}
+					}
 				}
 			}
 
diff --git a/Engine/Source/Programs/IOS/iPhonePackager/CompileTime.cs b/Engine/Source/Programs/IOS/iPhonePackager/CompileTime.cs
--- a/Engine/Source/Programs/IOS/iPhonePackager/CompileTime.cs
+++ b/Engine/Source/Programs/IOS/iPhonePackager/CompileTime.cs
@@ -203,7 +203,15 @@ namespace iPhonePackager
 			string ProvisionWithPrefix = FileOperations.FindPrefixedFile(Config.BuildDirectory, Program.GameName + ".mobileprovision");
 			if (!File.Exists(ProvisionWithPrefix))
 			{
-				ProvisionWithPrefix = FileOperations.FindPrefixedFile(Config.EngineBuildDirectory, "UE4Game.mobileprovision");
+				ProvisionWithPrefix = FileOperations.FindPrefixedFile(Config.BuildDirectory + "/NotForLicensees/", Program.GameName + ".mobileprovision");
+				if (!File.Exists(ProvisionWithPrefix))
+				{
+					ProvisionWithPrefix = FileOperations.FindPrefixedFile(Config.EngineBuildDirectory, "UE4Game.mobileprovision");
+					if (!File.Exists(ProvisionWithPrefix))
+					{
+						ProvisionWithPrefix = FileOperations.FindPrefixedFile(Config.EngineBuildDirectory + "/NotForLicensees/", "UE4Game.mobileprovision");
+					}
+				}
 			}
 			string FinalMobileProvisionFilename = Path.Combine(Config.PCXcodeStagingDir, MacMobileProvisionFilename);
 			FileOperations.CopyRequiredFile(ProvisionWithPrefix, FinalMobileProvisionFilename);
diff --git a/Engine/Source/Programs/UnrealBuildTool/IOS/UEDeployIOS.cs b/Engine/Source/Programs/UnrealBuildTool/IOS/UEDeployIOS.cs
--- a/Engine/Source/Programs/UnrealBuildTool/IOS/UEDeployIOS.cs
+++ b/Engine/Source/Programs/UnrealBuildTool/IOS/UEDeployIOS.cs
@@ -116,6 +116,17 @@ namespace UnrealBuildTool.IOS
 			{
 				ProvisionWithPrefix = BuildDirectory + "/" + InProjectName + ".mobileprovision";
 			}
+			else
+			{
+				if (File.Exists(BuildDirectory + "/NotForLicensees/" + InProjectName + ".mobileprovision"))
+				{
+					ProvisionWithPrefix = BuildDirectory + "/NotForLicensees/" + InProjectName + ".mobileprovision";
+				}
+				else if (!File.Exists(ProvisionWithPrefix))
+				{
+					ProvisionWithPrefix = InEngineDir + "/Build/IOS/NotForLicensees/UE4Game.mobileprovision";
+				}
+			}
 			if (File.Exists (ProvisionWithPrefix))
 			{
 				Directory.CreateDirectory (Environment.GetEnvironmentVariable ("HOME") + "/Library/MobileDevice/Provisioning Profiles/");
@@ -130,7 +141,18 @@ namespace UnrealBuildTool.IOS
 			{
 				ProvisionWithPrefix = BuildDirectory + "/" + InProjectName + "_Distro.mobileprovision";
 			}
-			if (File.Exists (ProvisionWithPrefix))
+			else
+			{
+				if (File.Exists(BuildDirectory + "/NotForLicensees/" + InProjectName + "_Distro.mobileprovision"))
+				{
+					ProvisionWithPrefix = BuildDirectory + "/NotForLicensees/" + InProjectName + "_Distro.mobileprovision";
+				}
+				else if (!File.Exists(ProvisionWithPrefix))
+				{
+					ProvisionWithPrefix = InEngineDir + "/Build/IOS/NotForLicensees/UE4Game_Distro.mobileprovision";
+				}
+			}
+			if (File.Exists(ProvisionWithPrefix))
 			{
 				File.Copy (ProvisionWithPrefix, Environment.GetEnvironmentVariable ("HOME") + "/Library/MobileDevice/Provisioning Profiles/" + InProjectName + "_Distro.mobileprovision", true);
 				FileInfo DestFileInfo = new FileInfo (Environment.GetEnvironmentVariable ("HOME") + "/Library/MobileDevice/Provisioning Profiles/" + InProjectName + "_Distro.mobileprovision");
