fix for IOSOnPCTestCompile node not building IPP as part of the build.
#ue4

diff --git a/Engine/Source/Programs/AutomationTool/IOS/IOSPlatform.Automation.cs b/Engine/Source/Programs/AutomationTool/IOS/IOSPlatform.Automation.cs
--- a/Engine/Source/Programs/AutomationTool/IOS/IOSPlatform.Automation.cs
+++ b/Engine/Source/Programs/AutomationTool/IOS/IOSPlatform.Automation.cs
@@ -710,10 +710,10 @@ public class IOSPlatform : Platform
 
 	public override void PreBuildAgenda(UE4Build Build, UE4Build.BuildAgenda Agenda)
 	{
-/*		if (UnrealBuildTool.ExternalExecution.GetRuntimePlatform () != UnrealTargetPlatform.Mac)
+		if (UnrealBuildTool.ExternalExecution.GetRuntimePlatform () != UnrealTargetPlatform.Mac)
 		{
 			Agenda.DotNetProjects.Add (@"Engine\Source\Programs\IOS\iPhonePackager\iPhonePackager.csproj");
-		}*/
+		}
 	}
 
 	#endregion
