UAT UnrealSync
Initial version of UnrealSync UAT script.

diff --git a/Engine/Source/Programs/AutomationTool/P4Utils.cs b/Engine/Source/Programs/AutomationTool/P4Utils.cs
--- a/Engine/Source/Programs/AutomationTool/P4Utils.cs
+++ b/Engine/Source/Programs/AutomationTool/P4Utils.cs
@@ -1637,6 +1637,138 @@ namespace AutomationTool
 			}
 			return false;
 		}
+
+		/* Pattern to parse P4 changes command output. */
+		private static readonly Regex ChangesListOutputPattern = new Regex(@"^Change\s+(?<number>\d+)\s+.+$", RegexOptions.Compiled | RegexOptions.Multiline);
+
+		/// <summary>
+		/// Gets the latest CL number submitted to the depot. It equals to the @head.
+		/// </summary>
+		/// <returns>The head CL number.</returns>
+		public int GetLatestCLNumber()
+		{
+			CheckP4Enabled();
+
+			string Output;
+			if (!LogP4Output(out Output, "changes -s submitted -m1") || string.IsNullOrWhiteSpace(Output))
+			{
+				throw new InvalidOperationException("The depot should have at least one submitted changelist. Brand new depot?");
+			}
+
+			var Match = ChangesListOutputPattern.Match(Output);
+
+			if (!Match.Success)
+			{
+				throw new InvalidOperationException("The Perforce output is not in the expected format provided by 2014.1 documentation.");
+			}
+
+			return Int32.Parse(Match.Groups["number"].Value);
+		}
+
+		/* Pattern to parse P4 labels command output. */
+		static readonly Regex LabelsListOutputPattern = new Regex(@"^Label\s+(?<name>[\w\/-]+)\s+(?<date>\d{4}/\d{2}/\d{2} \d{2}:\d{2}:\d{2})\s+'(?<description>.+)'\s*$", RegexOptions.Compiled | RegexOptions.Multiline);
+
+		/// <summary>
+		/// Gets promoted label list for given branch.
+		/// </summary>
+		/// <param name="BranchPath">A branch path.</param>
+		/// <param name="GameName">The game name for which provide the label. If null or empty then provides shared-promotable label.</param>
+		/// <returns>List of promoted labels for given branch path.</returns>
+		public string[] GetPromotedLabels(string BranchPath, string GameName)
+		{
+			var LabelNameList = new List<string>();
+
+			string Output;
+			if (LogP4Output(out Output, "labels -t -e " + BranchPath + "/Promoted" + (GameName != null ? ("-" + GameName) : "") + "-CL-*"))
+			{
+				foreach (Match LabelMatch in LabelsListOutputPattern.Matches(Output))
+				{
+					LabelNameList.Add(LabelMatch.Groups["name"].Value);
+				}
+			}
+
+			return LabelNameList.ToArray();
+		}
+
+		/// <summary>
+		/// Get latest promoted label given branch and game name.
+		/// </summary>
+		/// <param name="BranchPath">The branch path of the label.</param>
+		/// <param name="GameName">The game name for which provide the label. If null or empty then provides shared-promotable label.</param>
+		/// <param name="bVerifyContent">Verify if label tags at least one file.</param>
+		/// <returns>Label name if it exists, null otherwise.</returns>
+		public string GetLatestPromotedLabel(string BranchPath, string GameName, bool bVerifyContent)
+		{
+			CheckP4Enabled();
+
+			if (string.IsNullOrWhiteSpace(GameName))
+			{
+				GameName = null;
+			}
+
+			string Output;
+			if (LogP4Output(out Output, "labels -t -e " + BranchPath + "/Promoted" + (GameName != null ? ("-" + GameName) : "") + "-CL-*"))
+			{
+				var Labels = new Dictionary<string, DateTime>();
+
+				foreach (Match LabelMatch in LabelsListOutputPattern.Matches(Output))
+				{
+					Labels.Add(LabelMatch.Groups["name"].Value,
+						DateTime.ParseExact(
+							LabelMatch.Groups["date"].Value, "yyyy/MM/dd HH:mm:ss",
+							System.Globalization.CultureInfo.InvariantCulture));
+				}
+
+				if (Labels.Count == 0)
+				{
+					return null;
+				}
+
+				var OrderedLabels = Labels.OrderByDescending((Label) => Label.Value);
+
+				if (bVerifyContent)
+				{
+					foreach (var Label in OrderedLabels)
+					{
+						if (ValidateLabelContent(Label.Key))
+						{
+							return Label.Key;
+						}
+					}
+
+					// Haven't found valid label.
+					return null;
+				}
+
+				return OrderedLabels.First().Key;
+			}
+
+			return null;
+		}
+
+		/// <summary>
+		/// Validate label for some content.
+		/// </summary>
+		/// <returns>True if label exists and has at least one file tagged. False otherwise.</returns>
+		public bool ValidateLabelContent(string LabelName)
+		{
+			string Output;
+			if (LogP4Output(out Output, "files -m 1 @" + LabelName))
+			{
+				if (Output.StartsWith("//depot"))
+				{
+					// If it starts with depot path then label has at least one file tagged in it.
+					return true;
+				}
+			}
+			else
+			{
+				throw new InvalidOperationException("For some reason P4 files failed.");
+			}
+
+			return false;
+		}
+
 		/// <summary>
         /// returns the full name of a label. //depot/UE4/TEST-GUBP-Promotable-GameName-CL-CLNUMBER
 		/// </summary>
diff --git a/Engine/Source/Programs/AutomationTool/Scripts/AutomationScripts.Automation.csproj b/Engine/Source/Programs/AutomationTool/Scripts/AutomationScripts.Automation.csproj
--- a/Engine/Source/Programs/AutomationTool/Scripts/AutomationScripts.Automation.csproj
+++ b/Engine/Source/Programs/AutomationTool/Scripts/AutomationScripts.Automation.csproj
@@ -77,6 +77,7 @@
     <Compile Include="Tests.Automation.cs" />
     <Compile Include="TestSnippet.Automation.cs" />
     <Compile Include="UE4BuildUtils.cs" />
+    <Compile Include="UnrealSync.Automation.cs" />
     <Compile Include="UpdateLocalVersion.Automation.cs" />
   </ItemGroup>
   <ItemGroup>
diff --git a/Engine/Source/Programs/AutomationTool/Scripts/UnrealSync.Automation.cs b/Engine/Source/Programs/AutomationTool/Scripts/UnrealSync.Automation.cs
new file mode 100644
--- /dev/null
+++ b/Engine/Source/Programs/AutomationTool/Scripts/UnrealSync.Automation.cs
@@ -0,0 +1,171 @@
+﻿using AutomationTool;
+using System;
+using System.Collections.Generic;
+using System.Linq;
+using System.Text;
+using System.Text.RegularExpressions;
+using System.Xml;
+using UnrealBuildTool;
+
+namespace AutomationScripts.Automation
+{
+	[Help("Syncs promotable build. Use either -branch and -game or just -label parameter.")]
+	[Help("listlabels", "Enables special mode that lists all possible promoted labels for given branch and game.")]
+	[Help("artist", "Artist sync i.e. sync content to head and all the rest to promoted label.")]
+	[Help("branch=BranchPath", "Branch path to sync.")]
+	[Help("game=GameName", "Name of the game to sync. If not set then shared promotable will be synced.")]
+	[Help("label=LabelName", "Promotable label name to sync to.")]
+	[RequireP4]
+	class UnrealSync : BuildCommand
+	{
+		public override void ExecuteBuild()
+		{
+			var List = ParseParam("listlabels");
+			var ArtistSync = ParseParam("artist");
+			var BranchPath = ParseParamValue("branch");
+
+			if (BranchPath != null && BranchPath.EndsWith("/"))
+			{
+				BranchPath = BranchPath.Substring(0, BranchPath.Length - 1);
+			}
+
+			var GameName = ParseParamValue("game");
+
+			if (List)
+			{
+				if (string.IsNullOrWhiteSpace(BranchPath))
+				{
+					throw new AutomationException("You have to provide a branch to list promoted labels.");
+				}
+
+				Log(string.Format("Promoted labels for branch {0} and game {1}.", BranchPath, GameName));
+
+				// The P4 command will log out the possible labels.
+				P4.GetPromotedLabels(BranchPath, GameName);
+
+				return;
+			}
+
+			string ProgramSyncLabelName = null;
+
+			if (string.IsNullOrWhiteSpace(BranchPath) && string.IsNullOrWhiteSpace(GameName))
+			{
+				var LabelParam = ParseParamValue("label");
+
+				if (string.IsNullOrWhiteSpace(LabelParam))
+				{
+					throw new AutomationException("Use either -branch and -game or just -label parameter.");
+				}
+
+				if (!ValidateAndParseLabelName(LabelParam, out BranchPath, out GameName) || !P4.ValidateLabelContent(LabelParam))
+				{
+					throw new AutomationException("Label {0} either doesn't exist or is not valid promotable.");
+				}
+
+				ProgramSyncLabelName = LabelParam;
+			}
+			else
+			{
+				ProgramSyncLabelName = P4.GetLatestPromotedLabel(BranchPath, GameName, true);
+			}
+
+			if (ProgramSyncLabelName == null)
+			{
+				throw new AutomationException(string.Format("Label for {0} was not found.",
+					string.IsNullOrWhiteSpace(GameName)
+					? string.Format("branch {0} shared-promotable", BranchPath)
+					: string.Format("branch {0} and game {1}", GameName)));
+			}
+
+			var ArtistSyncRulesPath = string.Format("{0}/{1}/Build/ArtistSyncRules.xml",
+				BranchPath, string.IsNullOrWhiteSpace(GameName) ? "Samples" : GameName);
+
+			var SyncRules = string.Join("\n", P4.P4Print(ArtistSyncRulesPath + "#head"));
+
+			if (string.IsNullOrWhiteSpace(SyncRules))
+			{
+				throw new AutomationException(string.Format("The path {0} is not valid or file is empty.", ArtistSyncRulesPath));
+			}
+
+			var ProgramRevisionSpec = "@" + ProgramSyncLabelName;
+
+			// Get latest CL number to sync cause @head can change during
+			// different syncs and it could create integrity problems in
+			// workspace.
+			var ContentRevisionSpec = "@" +
+				(
+					ArtistSync
+					? P4.GetLatestCLNumber().ToString()
+					: ProgramSyncLabelName
+				);
+
+			foreach (var SyncStep in GenerateSyncSteps(SyncRules, ContentRevisionSpec, ProgramRevisionSpec))
+			{
+				P4.Sync(BranchPath + SyncStep);
+			}
+		}
+
+		/* Pattern to parse branch path and game name from label name. */
+		private static readonly Regex PromotedLabelPattern = new Regex(@"^(?<branchPath>//depot/UE4)/Promoted-((?<gameName>\w+)-)?CL-\d+$", RegexOptions.Compiled);
+
+		/// <summary>
+		/// Validates and parses promotable label name.
+		/// </summary>
+		/// <param name="LabelName">Promotable label to validate and parse.</param>
+		/// <param name="BranchPath">Parsed branch path.</param>
+		/// <param name="GameName">Parsed game name.</param>
+		/// <returns>True if label is valid. False otherwise.</returns>
+		private bool ValidateAndParseLabelName(string LabelName, out string BranchPath, out string GameName)
+		{
+			var Match = PromotedLabelPattern.Match(LabelName);
+
+			if (!Match.Success)
+			{
+				BranchPath = null;
+				GameName = null;
+
+				return false;
+			}
+
+			BranchPath = Match.Groups["branchPath"].Value;
+			GameName = Match.Groups["gameName"].Value;
+
+			return true;
+		}
+
+		/// <summary>
+		/// Generates sync steps based on the sync rules xml content.
+		/// </summary>
+		/// <param name="SyncRules">ArtistSyncRules.xml content.</param>
+		/// <param name="ContentRevisionSpec">Revision spec to which sync the content. Different for artist sync.</param>
+		/// <param name="ProgramRevisionSpec">Revision spec to which sync everything except content.</param>
+		/// <returns>An array of sync steps to perform.</returns>
+		private string[] GenerateSyncSteps(string SyncRules, string ContentRevisionSpec, string ProgramRevisionSpec)
+		{
+			var SyncRulesDocument = new XmlDocument();
+			SyncRulesDocument.LoadXml(SyncRules);
+
+			var RuleNodes = SyncRulesDocument.DocumentElement.SelectNodes("/ArtistSyncRules/Rules/string");
+
+			var OutputList = new List<string>();
+
+			foreach (XmlNode Node in RuleNodes)
+			{
+				var SyncStep = Node.InnerText.Replace("%LABEL_TO_SYNC_TO%", ProgramRevisionSpec);
+
+				// If there was no label in sync step.
+				// TODO: This is hack for messy ArtistSyncRules.xml format.
+				// Needs to be changed soon.
+				if (!SyncStep.Contains("@"))
+				{
+					SyncStep += ContentRevisionSpec;
+				}
+
+				OutputList.Add(SyncStep);
+			}
+
+			return OutputList.ToArray();
+		}
+
+	}
+}
