Add localization and splash files to packaged Linux builds.

diff --git a/Engine/Source/Programs/AutomationTool/Linux/LinuxPlatform.Automation.cs b/Engine/Source/Programs/AutomationTool/Linux/LinuxPlatform.Automation.cs
--- a/Engine/Source/Programs/AutomationTool/Linux/LinuxPlatform.Automation.cs
+++ b/Engine/Source/Programs/AutomationTool/Linux/LinuxPlatform.Automation.cs
@@ -48,6 +48,9 @@ public abstract class BaseLinuxPlatform : Platform
             SC.StageFiles(StagedFileType.NonUFS, CombinePaths(SC.LocalRoot, "Engine/Binaries/ThirdParty/OpenAL/", SC.PlatformDir), "libopenal.so.1", false, null, CombinePaths("Engine/Binaries", SC.PlatformDir));
         }
 
+        SC.StageFiles(StagedFileType.NonUFS, CombinePaths(SC.ProjectRoot, "Content/Splash"), "Splash.bmp", false, null, null, true);
+        SC.StageFiles(StagedFileType.NonUFS, CombinePaths(SC.LocalRoot, "Engine/Content/Localization"), "*.dat", true, null, null, false, !Params.Pak);
+
         List<string> Exes = GetExecutableNames(SC);
 
         foreach (var Exe in Exes)
