Attempting to fix null object reference error.  CombinePaths args appear to have been swapped, although this code itself hasn't been touched in some time.

diff --git a/Engine/Source/Programs/AutomationTool/TempStorage.cs b/Engine/Source/Programs/AutomationTool/TempStorage.cs
--- a/Engine/Source/Programs/AutomationTool/TempStorage.cs
+++ b/Engine/Source/Programs/AutomationTool/TempStorage.cs
@@ -132,7 +132,7 @@ namespace AutomationTool
 
             public void Create(List<string> InFiles, string BaseFolder)
             {
-                BaseFolder = CombinePaths(BaseFolder, "/");
+                BaseFolder = CombinePaths(PathSeparator.Slash, BaseFolder);
                 if (!BaseFolder.EndsWith("/")&& !BaseFolder.EndsWith("\\"))
                 {
                     throw new AutomationException("base folder {0} should end with a separator", BaseFolder);
