Second try at fixing null object in TempStorage.

diff --git a/Engine/Source/Programs/AutomationTool/TempStorage.cs b/Engine/Source/Programs/AutomationTool/TempStorage.cs
--- a/Engine/Source/Programs/AutomationTool/TempStorage.cs
+++ b/Engine/Source/Programs/AutomationTool/TempStorage.cs
@@ -132,7 +132,11 @@ namespace AutomationTool
 
             public void Create(List<string> InFiles, string BaseFolder)
             {
-                BaseFolder = CombinePaths(PathSeparator.Slash, BaseFolder);
+				if (BaseFolder == null)
+				{
+					BaseFolder = "";
+				}
+                BaseFolder = CombinePaths(BaseFolder, "/");
                 if (!BaseFolder.EndsWith("/")&& !BaseFolder.EndsWith("\\"))
                 {
                     throw new AutomationException("base folder {0} should end with a separator", BaseFolder);
