Removed my null-guarding in TempStorage

diff --git a/Engine/Source/Programs/AutomationTool/TempStorage.cs b/Engine/Source/Programs/AutomationTool/TempStorage.cs
--- a/Engine/Source/Programs/AutomationTool/TempStorage.cs
+++ b/Engine/Source/Programs/AutomationTool/TempStorage.cs
@@ -132,10 +132,6 @@ namespace AutomationTool
 
             public void Create(List<string> InFiles, string BaseFolder)
             {
-				if (BaseFolder == null)
-				{
-					BaseFolder = "";
-				}
                 BaseFolder = CombinePaths(BaseFolder, "/");
                 if (!BaseFolder.EndsWith("/")&& !BaseFolder.EndsWith("\\"))
                 {
