build script, preparation for deleting old builds

diff --git a/Engine/Source/Programs/AutomationTool/CommandUtils.cs b/Engine/Source/Programs/AutomationTool/CommandUtils.cs
--- a/Engine/Source/Programs/AutomationTool/CommandUtils.cs
+++ b/Engine/Source/Programs/AutomationTool/CommandUtils.cs
@@ -1813,6 +1813,114 @@ namespace AutomationTool
 		}
 
 		#endregion
+
+
+        public static string RootSharedTempStorageDirectory()
+        {
+            string StorageDirectory = "";
+            if (UnrealBuildTool.Utils.IsRunningOnMono)
+            {
+                StorageDirectory = "/Volumes/Builds";
+            }
+            else
+            {
+                StorageDirectory = CombinePaths("P:", "Builds");
+            }
+            return StorageDirectory;
+        }
+
+        static Dictionary<string, string> ResolveCache = new Dictionary<string, string>();
+        public static string ResolveSharedBuildDirectory(string GameFolder)
+        {
+            if (ResolveCache.ContainsKey(GameFolder))
+            {
+                return ResolveCache[GameFolder];
+            }
+            string Root = RootSharedTempStorageDirectory();
+            string Result = CombinePaths(Root, GameFolder);
+            if (String.IsNullOrEmpty(GameFolder) || !DirectoryExists_NoExceptions(Result))
+            {
+                string GameStr = "Game";
+                bool HadGame = false;
+                if (GameFolder.EndsWith(GameStr, StringComparison.InvariantCultureIgnoreCase))
+                {
+                    string ShortFolder = GameFolder.Substring(0, GameFolder.Length - GameStr.Length);
+                    Result = CombinePaths(Root, ShortFolder);
+                    HadGame = true;
+                }
+                if (!HadGame || !DirectoryExists_NoExceptions(Result))
+                {
+                    Result = CombinePaths(Root, "UE4");
+                    if (!DirectoryExists_NoExceptions(Result))
+                    {
+                        throw new AutomationException("Could not find an appropriate shared temp folder {0}", Result);
+                    }
+                }
+            }
+            ResolveCache.Add(GameFolder, Result);
+            return Result;
+        }
+
+        public static void CleanFormalBuilds(string DirectoryForThisBuild)
+        {
+            if (!IsBuildMachine || !DirectoryForThisBuild.StartsWith(RootSharedTempStorageDirectory()) || !P4Enabled)
+            {
+                return;
+            }
+            try
+            {
+                string CLString = P4Env.ChangelistString;
+                const int MaximumDaysToKeepTempStorage = 4;
+                string ParentDir = Path.GetDirectoryName(CombinePaths(DirectoryForThisBuild));
+                if (!DirectoryExists_NoExceptions(ParentDir))
+                {
+                    throw new AutomationException("Not cleaning formal builds, because the parent directory {0} does not exist.", ParentDir);
+                }
+                string MyDir = Path.GetFileName(CombinePaths(DirectoryForThisBuild));
+                int CLStart = MyDir.IndexOf(CLString);
+                if (CLStart < 0)
+                {
+                    throw new AutomationException("Not cleaning formal builds, because the directory {0} does not contain the CL {1}.", DirectoryForThisBuild, CLString);
+                }
+                string StartString = MyDir.Substring(0, CLStart);
+                string EndString = MyDir.Substring(CLStart + CLString.Length);
+
+
+                DirectoryInfo DirInfo = new DirectoryInfo(ParentDir);
+                var TopLevelDirs = DirInfo.GetDirectories();
+                Log("Looking for directories to delete based on {0}", DirectoryForThisBuild);
+                foreach (var TopLevelDir in TopLevelDirs)
+                {
+                    if (DirectoryExists_NoExceptions(TopLevelDir.FullName))
+                    {
+                        if (TopLevelDir.FullName.StartsWith(StartString, StringComparison.InvariantCultureIgnoreCase) && TopLevelDir.FullName.EndsWith(EndString, StringComparison.InvariantCultureIgnoreCase))
+                        {
+                            string CLPart = TopLevelDir.FullName.Substring(StartString.Length, TopLevelDir.FullName.Length - StartString.Length - EndString.Length);
+                            if (!CLPart.Contains("-") && !CLPart.Contains("+"))
+                            {
+                                DirectoryInfo ThisDirInfo = new DirectoryInfo(TopLevelDir.FullName);
+                                bool bOld = false;
+
+                                if ((DateTime.UtcNow - ThisDirInfo.CreationTimeUtc).TotalDays > MaximumDaysToKeepTempStorage)
+                                {
+                                    bOld = true;
+                                }
+                                if (bOld)
+                                {
+                                    Log("Deleteing (or will be after this looks correct) temp storage directory {0}, because it is more than {1} days old.", TopLevelDir.FullName, MaximumDaysToKeepTempStorage);
+                                    //DeleteDirectory_NoExceptions(true, TopLevelDir.FullName);
+                                }
+                            }
+                        }
+                    }
+                }
+            }
+            catch (Exception Ex)
+            {
+                Log(System.Diagnostics.TraceEventType.Warning, "Unable to Clean Directory with DirectoryForThisBuild {0}", DirectoryForThisBuild);
+                Log(System.Diagnostics.TraceEventType.Warning, " Exception was {0}", LogUtils.FormatException(Ex));
+            }
+        }
 	}
 
 
diff --git a/Engine/Source/Programs/AutomationTool/MCPPublic.cs b/Engine/Source/Programs/AutomationTool/MCPPublic.cs
--- a/Engine/Source/Programs/AutomationTool/MCPPublic.cs
+++ b/Engine/Source/Programs/AutomationTool/MCPPublic.cs
@@ -176,7 +176,7 @@ namespace EpicGames.MCP.Automation
         static public string GetBuildRootPath()
         {
             return CommandUtils.IsBuildMachine
-                ? (Utils.IsRunningOnMono ? "/Volumes/Builds" : @"P:\Builds")
+                ? CommandUtils.RootSharedTempStorageDirectory()
                 : CommandUtils.CombinePaths(CommandUtils.CmdEnv.LocalRoot, "LocalBuilds");
         }
 
diff --git a/Engine/Source/Programs/AutomationTool/Scripts/GUBP.Automation.cs b/Engine/Source/Programs/AutomationTool/Scripts/GUBP.Automation.cs
--- a/Engine/Source/Programs/AutomationTool/Scripts/GUBP.Automation.cs
+++ b/Engine/Source/Programs/AutomationTool/Scripts/GUBP.Automation.cs
@@ -2409,6 +2409,7 @@ public class GUBP : BuildCommand
             {
                 string BaseDir = ResolveSharedBuildDirectory(GameProj.GameName);
                 string ArchiveDirectory = CombinePaths(BaseDir, GetFullName(), P4Env.BuildRootEscaped + "-CL-" + P4Env.ChangelistString);
+                CleanFormalBuilds(ArchiveDirectory);
                 if (DirectoryExists_NoExceptions(ArchiveDirectory))
                 {
                     if (IsBuildMachine)
diff --git a/Engine/Source/Programs/AutomationTool/TempStorage.cs b/Engine/Source/Programs/AutomationTool/TempStorage.cs
--- a/Engine/Source/Programs/AutomationTool/TempStorage.cs
+++ b/Engine/Source/Programs/AutomationTool/TempStorage.cs
@@ -522,19 +522,6 @@ namespace AutomationTool
             }
         }
         
-        public static string RootSharedTempStorageDirectory()
-        {
-            string StorageDirectory = "";
-            if(UnrealBuildTool.Utils.IsRunningOnMono)
-            {
-                StorageDirectory = "/Volumes/Builds";
-            }
-            else
-            {
-                StorageDirectory = CombinePaths("P:", "Builds");
-            }
-            return StorageDirectory;
-        }
         public static string SharedTempStorageDirectory()
         {
             return "GUBP";
@@ -554,38 +541,6 @@ namespace AutomationTool
             return true;
         }
 
-        static Dictionary<string, string> ResolveCache = new Dictionary<string, string>();
-        public static string ResolveSharedBuildDirectory(string GameFolder)
-        {
-            if (ResolveCache.ContainsKey(GameFolder))
-            {
-                return ResolveCache[GameFolder];
-            }
-            string Root = RootSharedTempStorageDirectory();
-            string Result = CombinePaths(Root, GameFolder);
-            if (String.IsNullOrEmpty(GameFolder) || !DirectoryExists_NoExceptions(Result))
-            {
-                string GameStr = "Game";
-                bool HadGame = false;
-                if (GameFolder.EndsWith(GameStr, StringComparison.InvariantCultureIgnoreCase))
-                {
-                    string ShortFolder = GameFolder.Substring(0, GameFolder.Length - GameStr.Length);
-                    Result = CombinePaths(Root, ShortFolder);
-                    HadGame = true;
-                }
-                if (!HadGame || !DirectoryExists_NoExceptions(Result))
-                {
-                    Result = CombinePaths(Root, "UE4");
-                    if (!DirectoryExists_NoExceptions(Result))
-                    {
-                        throw new AutomationException("Could not find an appropriate shared temp folder {0}", Result);
-                    }
-                }
-            }
-            ResolveCache.Add(GameFolder, Result);
-            return Result;
-        }
-
         public static string ResolveSharedTempStorageDirectory(string GameFolder)
         {
             string SharedSubdir = SharedTempStorageDirectory();
