build script, reenable automated testing

diff --git a/Engine/Source/Programs/AutomationTool/Scripts/RunProjectCommand.Automation.cs b/Engine/Source/Programs/AutomationTool/Scripts/RunProjectCommand.Automation.cs
--- a/Engine/Source/Programs/AutomationTool/Scripts/RunProjectCommand.Automation.cs
+++ b/Engine/Source/Programs/AutomationTool/Scripts/RunProjectCommand.Automation.cs
@@ -198,15 +198,29 @@ public partial class Project : CommandUtils
             {
                 Timeout = 6 * 60 * 60;
             }
+            string LookFor = "Bringing up level for play took";
+            if (Params.RunAutomationTest != "")
+            {
+                LookFor = "Automation Test Succeeded";
+            }
+            else if (Params.RunAutomationTests)
+            {
+                LookFor = "Automation Test Queue Empty";
+            }
+            else if (Params.EditorTest)
+            {
+                LookFor = "Asset discovery search completed in";
+            }
             using (var Watchdog = new WatchdogTimer(Timeout))
 			{
 
                 string AllClientOutput = "";
+                int LastAutoFailIndex = -1;
                 ProcessResult ClientProcess = null;
 				FileStream ClientProcessLog = null;
 				StreamReader ClientLogReader = null;
 				Log("Starting Client for unattended test....");
-				ClientProcess = Run(ClientApp, ClientCmdLine + " -FORCELOGFLUSH", null, ClientRunFlags | ERunOptions.NoWaitForExit);
+                ClientProcess = Run(ClientApp, ClientCmdLine + " -FORCELOGFLUSH -testexit=\"" + LookFor + "\"", null, ClientRunFlags | ERunOptions.NoWaitForExit);
 				while (!FileExists(ClientLogFile) && !ClientProcess.HasExited)
 				{
 					Log("Waiting for client logging process to start...{0}", ClientLogFile);
@@ -235,36 +249,19 @@ public partial class Project : CommandUtils
                             AllClientOutput += ClientOutput;
 							Console.Write(ClientOutput);
 
-                            string LookFor = "Bringing up level for play took";
-                            if (Params.RunAutomationTest != "")
-                            {
-                                LookFor = "Automation Test Succeeded";
-                            }
-                            else if (Params.RunAutomationTests)
-                            {
-                                LookFor = "Automation Test Queue Empty";
-                            }
-                            else if (Params.EditorTest)
-                            {
-                                LookFor = "Asset discovery search completed in";
-                            }
-
-                            if (AllClientOutput.Contains(LookFor))
+                            if (AllClientOutput.LastIndexOf(LookFor) > AllClientOutput.IndexOf(LookFor))
 							{
 								Log("Client loaded, lets wait 30 seconds...");
 								Thread.Sleep(30000);
-								if (!ClientProcess.HasExited)
-								{
-									WelcomedCorrectly = true;
-								}
+								WelcomedCorrectly = true;
 								Log("Test complete...");
 								bKeepReading = false;
 							}
                             else if (Params.RunAutomationTests)
                             {
-                                int FailIndex = AllClientOutput.IndexOf("Automation Test Failed");
+                                int FailIndex = AllClientOutput.LastIndexOf("Automation Test Failed");
                                 int ParenIndex = AllClientOutput.LastIndexOf(")");
-                                if (FailIndex >= 0 && ParenIndex > FailIndex)
+                                if (FailIndex >= 0 && ParenIndex > FailIndex && FailIndex > LastAutoFailIndex)
                                 {
                                     string Tail = AllClientOutput.Substring(FailIndex);
                                     int CloseParenIndex = Tail.IndexOf(")");
@@ -273,16 +270,9 @@ public partial class Project : CommandUtils
                                     if (OpenParenIndex >= 0 && CloseParenIndex > OpenParenIndex)
                                     {
                                         Test = Tail.Substring(OpenParenIndex + 1, CloseParenIndex - OpenParenIndex - 1);
+                                        Log(System.Diagnostics.TraceEventType.Error, "Automated test failed ({0}).", Test);
+                                        LastAutoFailIndex = FailIndex;
                                     }
-#if true
-                                    Log("Stopping client...");
-                                    ClientProcess.StopProcess();
-                                    if (IsBuildMachine)
-                                    {
-                                        Thread.Sleep(70000);
-                                    }
-#endif
-                                    throw new AutomationException("Automated test failed ({0}).", Test);
                                 }
                             }
                         }
@@ -294,11 +284,6 @@ public partial class Project : CommandUtils
 					ClientProcess.StopProcess();
 				}
 
-                // this is hack that prevents a hang
-                if (IsBuildMachine)
-                {
-                    Thread.Sleep(70000);
-                }
 				if (!WelcomedCorrectly)
 				{
 					throw new AutomationException("Client exited before we asked it to.");
@@ -336,6 +321,7 @@ public partial class Project : CommandUtils
 		bool WelcomedCorrectly = false;
 		int NumClients = Params.NumClients;
         string AllClientOutput = "";
+        int LastAutoFailIndex = -1;
 
         if (Params.Unattended)
 		{
@@ -347,6 +333,19 @@ public partial class Project : CommandUtils
             if (Params.RunAutomationTests)
             {
                 Timeout = 6 * 60 * 60;
+            }
+            string LookFor = "Bringing up level for play took";
+            if (Params.DedicatedServer)
+            {
+                LookFor = "Welcomed by server";
+            }
+            else if (Params.RunAutomationTest != "")
+            {
+                LookFor = "Automation Test Succeeded";
+            }
+            else if (Params.RunAutomationTests)
+            {
+                LookFor = "Automation Test Queue Empty";
             }
 			using (var Watchdog = new WatchdogTimer(Timeout))
 			{
@@ -379,7 +378,7 @@ public partial class Project : CommandUtils
                                         (AllServerOutput.Contains("Game Engine Initialized") || AllServerOutput.Contains("Unreal Network File Server is ready")))
 								{
 									Log("Starting Client for unattended test....");
-									ClientProcess = Run(ClientApp, ClientCmdLine + " -FORCELOGFLUSH", null, ClientRunFlags | ERunOptions.NoWaitForExit);
+									ClientProcess = Run(ClientApp, ClientCmdLine + " -FORCELOGFLUSH -testexit=\"" + LookFor + "\"", null, ClientRunFlags | ERunOptions.NoWaitForExit);
 									//@todo no testing is done on these
 									if (NumClients > 1 && NumClients < 9)
 									{
@@ -439,21 +438,8 @@ public partial class Project : CommandUtils
                                         AllClientOutput += ClientOutput;
 										Console.Write(ClientOutput);
 
-                                        string LookFor = "Bringing up level for play took";
-                                        if (Params.DedicatedServer)
-                                        {
-                                            LookFor = "Welcomed by server";
-                                        }
-                                        else if (Params.RunAutomationTest != "")
+                                        if (AllClientOutput.LastIndexOf(LookFor) > AllClientOutput.IndexOf(LookFor))
                                         {
-                                            LookFor = "Automation Test Succeeded";
-                                        }
-                                        else if (Params.RunAutomationTests)
-                                        {
-                                            LookFor = "Automation Test Queue Empty";
-                                        }
-                                        if (AllClientOutput.Contains(LookFor))
-										{
 											if (Params.FakeClient)
 											{
 												Log("Welcomed by server or client loaded, lets wait ten minutes...");
@@ -464,32 +450,14 @@ public partial class Project : CommandUtils
 												Log("Welcomed by server or client loaded, lets wait 30 seconds...");
 												Thread.Sleep(30000);
 											}
-											if (!ServerProcess.HasExited && !ClientProcess.HasExited)
-											{
-												WelcomedCorrectly = true;
-											}
-											if (!ServerProcess.HasExited)
-											{
-												Log("Test complete");
-												if (ClientProcess != null && !ClientProcess.HasExited)
-												{
-													Log("Stopping client....");
-													ClientProcess.StopProcess();
-												}
-												Log("Stopping server....");
-												ServerProcess.StopProcess();
-                                                if (IsBuildMachine)
-                                                {
-                                                    Thread.Sleep(70000);
-                                                }
-											}
+											WelcomedCorrectly = true;
 											bKeepReading = false;
 										}
                                         else if (Params.RunAutomationTests)
                                         {
-                                            int FailIndex = AllClientOutput.IndexOf("Automation Test Failed");
+                                            int FailIndex = AllClientOutput.LastIndexOf("Automation Test Failed");
                                             int ParenIndex = AllClientOutput.LastIndexOf(")");
-                                            if (FailIndex >= 0 && ParenIndex > FailIndex)
+                                            if (FailIndex >= 0 && ParenIndex > FailIndex && FailIndex > LastAutoFailIndex)
                                             {
                                                 string Tail = AllClientOutput.Substring(FailIndex);
                                                 int CloseParenIndex = Tail.IndexOf(")");
@@ -498,33 +466,9 @@ public partial class Project : CommandUtils
                                                 if (OpenParenIndex >= 0 && CloseParenIndex > OpenParenIndex)
                                                 {
                                                     Test = Tail.Substring(OpenParenIndex + 1, CloseParenIndex - OpenParenIndex - 1);
+                                                    Log(System.Diagnostics.TraceEventType.Error, "Automated test failed ({0}).", Test);
+                                                    LastAutoFailIndex = FailIndex;
                                                 }
-#if true
-                                                if (!ServerProcess.HasExited)
-                                                {
-                                                    Log("Stopping server....");
-                                                    ServerProcess.StopProcess();
-                                                }
-                                                if (ClientProcess != null && !ClientProcess.HasExited)
-                                                {
-                                                    Log("Stopping client....");
-                                                    ClientProcess.StopProcess();
-                                                }
-                                                foreach (var OtherClient in OtherClients)
-                                                {
-                                                    if (OtherClient != null && !OtherClient.HasExited)
-                                                    {
-                                                        Log("Stopping client....");
-                                                        OtherClient.StopProcess();
-                                                    }
-                                                }
-                                                if (IsBuildMachine)
-                                                {
-                                                    Thread.Sleep(70000);
-                                                }
-
-#endif
-                                                throw new AutomationException("Automated test failed ({0}).", Test);
                                             }
                                         }
 									}
@@ -594,11 +538,6 @@ public partial class Project : CommandUtils
 		}
 		if (Params.Unattended)
 		{
-            // this is a hack to prevent hangs on exit
-            if (IsBuildMachine)
-            {
-                Thread.Sleep(70000);
-            }
             if (!WelcomedCorrectly)
             {
 			    throw new AutomationException("Server or client exited before we asked it to.");
