Platform: Add -SearchDir=PathToSomething optional parameter to BuildThirdPartyLibs automation script, allowing it be used for TPS software in a plugin or other non-standard location

diff --git a/Engine/Source/Programs/AutomationTool/Scripts/BuildThirdPartyLibs.Automation.cs b/Engine/Source/Programs/AutomationTool/Scripts/BuildThirdPartyLibs.Automation.cs
--- a/Engine/Source/Programs/AutomationTool/Scripts/BuildThirdPartyLibs.Automation.cs
+++ b/Engine/Source/Programs/AutomationTool/Scripts/BuildThirdPartyLibs.Automation.cs
@@ -49,11 +49,12 @@ using UnrealBuildTool;
 [Help("Builds third party libraries, and puts them all into a changelist")]
 [Help("Libs", "[Optional] + separated list of libraries to compile; if not specified this job will build all libraries it can find builder scripts for")]
 [Help("Changelist", "[Optional] a changelist to check out into; if not specified, a changelist will be created")]
+[Help("SearchDir", "[Optional] Directory to search for the specified Libs")]
 [RequireP4]
 class BuildThirdPartyLibs : BuildCommand
 {
 	// path to the third party directory
-	static private string LibDir = "Engine/Source/ThirdParty";
+	static private string DefaultLibraryDir = "Engine/Source/ThirdParty";
 
 	// batch/script file to look for when compiling
 	static private string WindowsCompileScript = "UE4_BuildThirdPartyLib.bat";
@@ -92,7 +93,8 @@ class BuildThirdPartyLibs : BuildCommand
 		Log("Build from {0}    Working in {1}", P4Env.Changelist, WorkingCL);
 
 		// go to the third party lib dir
-		CommandUtils.PushDir(LibDir);
+		string SearchLibraryDir = ParseParamValue("SearchDir", DefaultLibraryDir);
+		CommandUtils.PushDir(SearchLibraryDir);
 
 		// figure out what libraries to compile
 		string LibsToCompileString = ParseParamValue("Libs");
@@ -152,7 +154,7 @@ class BuildThirdPartyLibs : BuildCommand
 			CommandUtils.PopDir();
 		}
 
-		// undo the LibDir push
+		// undo the SearchLibraryDir push
 		CommandUtils.PopDir();
 
 		PrintRunTime();			
