UAT UnrealSync: Getting rid of branch parameter as UAT detects it automatically.

diff --git a/Engine/Source/Programs/AutomationTool/Scripts/UnrealSync.Automation.cs b/Engine/Source/Programs/AutomationTool/Scripts/UnrealSync.Automation.cs
--- a/Engine/Source/Programs/AutomationTool/Scripts/UnrealSync.Automation.cs
+++ b/Engine/Source/Programs/AutomationTool/Scripts/UnrealSync.Automation.cs
@@ -9,10 +9,9 @@ using UnrealBuildTool;
 
 namespace AutomationScripts.Automation
 {
-	[Help("Syncs promotable build. Use either -branch and -game or just -label parameter.")]
-	[Help("listlabels", "Enables special mode that lists all possible promoted labels for given branch and game.")]
+	[Help("Syncs promotable build. Use either -game or -label parameter.")]
+	[Help("listlabels", "Enables special mode that lists all possible promoted labels for shared promotable or given game.")]
 	[Help("artist", "Artist sync i.e. sync content to head and all the rest to promoted label.")]
-	[Help("branch=BranchPath", "Branch path to sync.")]
 	[Help("game=GameName", "Name of the game to sync. If not set then shared promotable will be synced.")]
 	[Help("label=LabelName", "Promotable label name to sync to.")]
 	[RequireP4]
@@ -22,7 +21,7 @@ namespace AutomationScripts.Automation
 		{
 			var List = ParseParam("listlabels");
 			var ArtistSync = ParseParam("artist");
-			var BranchPath = ParseParamValue("branch");
+			var BranchPath = P4Env.BuildRootP4;
 
 			if (BranchPath != null && BranchPath.EndsWith("/"))
 			{
@@ -31,14 +30,18 @@ namespace AutomationScripts.Automation
 
 			var GameName = ParseParamValue("game");
 
+			var BranchAndGameName = string.IsNullOrWhiteSpace(GameName)
+					? string.Format("branch {0} shared-promotable", BranchPath)
+					: string.Format("branch {0} and game {1}", BranchPath, GameName);
+
 			if (List)
 			{
 				if (string.IsNullOrWhiteSpace(BranchPath))
 				{
-					throw new AutomationException("You have to provide a branch to list promoted labels.");
+					throw new AutomationException("The branch path is not set. Something went wrong.");
 				}
 
-				Log(string.Format("Promoted labels for branch {0} and game {1}.", BranchPath, GameName));
+				Log(string.Format("Promoted labels for {0}.", BranchAndGameName));
 
 				// The P4 command will log out the possible labels.
 				P4.GetPromotedLabels(BranchPath, GameName);
@@ -48,7 +51,7 @@ namespace AutomationScripts.Automation
 
 			string ProgramSyncLabelName = null;
 
-			if (string.IsNullOrWhiteSpace(BranchPath) && string.IsNullOrWhiteSpace(GameName))
+			if (string.IsNullOrWhiteSpace(GameName))
 			{
 				var LabelParam = ParseParamValue("label");
 
@@ -57,9 +60,9 @@ namespace AutomationScripts.Automation
 					throw new AutomationException("Use either -branch and -game or just -label parameter.");
 				}
 
-				if (!ValidateAndParseLabelName(LabelParam, out BranchPath, out GameName) || !P4.ValidateLabelContent(LabelParam))
+				if (!ValidateAndParseLabelName(LabelParam, BranchPath, out GameName) || !P4.ValidateLabelContent(LabelParam))
 				{
-					throw new AutomationException("Label {0} either doesn't exist or is not valid promotable.");
+					throw new AutomationException("Label {0} either doesn't exist or is not valid for the current branch path {1}.", LabelParam, BranchPath);
 				}
 
 				ProgramSyncLabelName = LabelParam;
@@ -71,10 +74,7 @@ namespace AutomationScripts.Automation
 
 			if (ProgramSyncLabelName == null)
 			{
-				throw new AutomationException(string.Format("Label for {0} was not found.",
-					string.IsNullOrWhiteSpace(GameName)
-					? string.Format("branch {0} shared-promotable", BranchPath)
-					: string.Format("branch {0} and game {1}", GameName)));
+				throw new AutomationException(string.Format("Label for {0} was not found.", BranchAndGameName));
 			}
 
 			var ArtistSyncRulesPath = string.Format("{0}/{1}/Build/ArtistSyncRules.xml",
@@ -112,22 +112,27 @@ namespace AutomationScripts.Automation
 		/// Validates and parses promotable label name.
 		/// </summary>
 		/// <param name="LabelName">Promotable label to validate and parse.</param>
-		/// <param name="BranchPath">Parsed branch path.</param>
+		/// <param name="BranchPath">Current branch path.</param>
 		/// <param name="GameName">Parsed game name.</param>
 		/// <returns>True if label is valid. False otherwise.</returns>
-		private bool ValidateAndParseLabelName(string LabelName, out string BranchPath, out string GameName)
+		private bool ValidateAndParseLabelName(string LabelName, string BranchPath, out string GameName)
 		{
 			var Match = PromotedLabelPattern.Match(LabelName);
 
 			if (!Match.Success)
 			{
-				BranchPath = null;
 				GameName = null;
 
 				return false;
 			}
 
-			BranchPath = Match.Groups["branchPath"].Value;
+			if(BranchPath != Match.Groups["branchPath"].Value)
+			{
+				GameName = null;
+
+				return false;
+			}
+
 			GameName = Match.Groups["gameName"].Value;
 
 			return true;
