build script, run with -buildmachine for -unattended

diff --git a/Engine/Source/Programs/AutomationTool/Scripts/RunProjectCommand.Automation.cs b/Engine/Source/Programs/AutomationTool/Scripts/RunProjectCommand.Automation.cs
--- a/Engine/Source/Programs/AutomationTool/Scripts/RunProjectCommand.Automation.cs
+++ b/Engine/Source/Programs/AutomationTool/Scripts/RunProjectCommand.Automation.cs
@@ -693,7 +693,7 @@ public partial class Project : CommandUtils
 		{
 			TempCmdLine += "-unattended ";
 		}
-        if (IsBuildMachine)
+        if (IsBuildMachine || Params.Unattended)
         {
             TempCmdLine += "-buildmachine ";
         }
@@ -842,7 +842,7 @@ public partial class Project : CommandUtils
 				Args += " -pak";
 			}
 		}
-        if (IsBuildMachine)
+        if (IsBuildMachine || Params.Unattended)
         {
             Args += " -buildmachine"; 
         }
