build script, remove watchdog timer

diff --git a/Engine/Source/Programs/AutomationTool/Scripts/RunProjectCommand.Automation.cs b/Engine/Source/Programs/AutomationTool/Scripts/RunProjectCommand.Automation.cs
--- a/Engine/Source/Programs/AutomationTool/Scripts/RunProjectCommand.Automation.cs
+++ b/Engine/Source/Programs/AutomationTool/Scripts/RunProjectCommand.Automation.cs
@@ -193,11 +193,6 @@ public partial class Project : CommandUtils
 	{
 		if (Params.Unattended)
 		{
-            int Timeout = 60 * 30;
-            if (Params.RunAutomationTests)
-            {
-                Timeout = 6 * 60 * 60;
-            }
             string LookFor = "Bringing up level for play took";
             if (Params.RunAutomationTest != "")
             {
@@ -211,7 +206,6 @@ public partial class Project : CommandUtils
             {
                 LookFor = "Asset discovery search completed in";
             }
-            using (var Watchdog = new WatchdogTimer(Timeout))
 			{
 
                 string AllClientOutput = "";
@@ -325,15 +319,6 @@ public partial class Project : CommandUtils
 
         if (Params.Unattended)
 		{
-            int Timeout = 60 * 5;
-            if (Params.FakeClient)
-            {
-                Timeout = 60 * 15;
-            }
-            if (Params.RunAutomationTests)
-            {
-                Timeout = 6 * 60 * 60;
-            }
             string LookFor = "Bringing up level for play took";
             if (Params.DedicatedServer)
             {
@@ -347,7 +332,6 @@ public partial class Project : CommandUtils
             {
                 LookFor = "Automation Test Queue Empty";
             }
-			using (var Watchdog = new WatchdogTimer(Timeout))
 			{
 				while (!FileExists(ServerLogFile) && !ServerProcess.HasExited)
 				{
@@ -409,10 +393,6 @@ public partial class Project : CommandUtils
 								if (ClientProcess != null && ClientProcess.HasExited)
 								{
 									ServerProcess.StopProcess();
-                                    if (IsBuildMachine)
-                                    {
-                                        Thread.Sleep(70000);
-                                    }
 									throw new AutomationException("Client exited before we asked it to.");
 								}
 							}
@@ -422,10 +402,6 @@ public partial class Project : CommandUtils
 							if (ClientProcess.HasExited)
 							{
 								ServerProcess.StopProcess();
-                                if (IsBuildMachine)
-                                {
-                                    Thread.Sleep(70000);
-                                }
 								throw new AutomationException("Client exited or closed the log before we asked it to.");
 							}
 							while (!ClientProcess.HasExited && !ServerProcess.HasExited && bKeepReading)
