build script, try to improve the performance of cleaning temp

diff --git a/Engine/Source/Programs/AutomationTool/TempStorage.cs b/Engine/Source/Programs/AutomationTool/TempStorage.cs
--- a/Engine/Source/Programs/AutomationTool/TempStorage.cs
+++ b/Engine/Source/Programs/AutomationTool/TempStorage.cs
@@ -494,11 +494,16 @@ namespace AutomationTool
             return CombinePaths(Env.LocalRoot, "Engine", "Saved", "GUBP");
         }
 
-        static void CleanSharedTempStorage(string Directory)
+        static void CleanSharedTempStorage(string TopDirectory)
         {
             const int MaximumDaysToKeepTempStorage = 2;
-            DirectoryInfo DirInfo = new DirectoryInfo(Directory);
+            var StartTimeDir = DateTime.UtcNow;
+            DirectoryInfo DirInfo = new DirectoryInfo(TopDirectory);
             var TopLevelDirs = DirInfo.GetDirectories();
+            {
+                var BuildDuration = (DateTime.UtcNow - StartTimeDir).TotalMilliseconds;
+                Log("Took {0}s to enumerate {1} directories.", BuildDuration / 1000, TopLevelDirs.Length);
+            }
             foreach (var TopLevelDir in TopLevelDirs)
             {
                 if (DirectoryExists_NoExceptions(TopLevelDir.FullName))
@@ -516,7 +521,21 @@ namespace AutomationTool
                     if (bOld)
                     {
                         Log("Deleteing temp storage directory {0}, because it is more than {1} days old.", TopLevelDir.FullName, MaximumDaysToKeepTempStorage);
+                        var StartTime = DateTime.UtcNow;
+                        try
+                        {
+                            if (Directory.Exists(TopLevelDir.FullName))
+                            {
+                                // try the direct approach first
+                                Directory.Delete(TopLevelDir.FullName, true);
+                            }
+                        }
+                        catch (Exception)
+                        {
+                        }
                         DeleteDirectory_NoExceptions(true, TopLevelDir.FullName);
+                        var BuildDuration = (DateTime.UtcNow - StartTime).TotalMilliseconds;
+                        Log("Took {0}s to delete {1}.", BuildDuration / 1000, TopLevelDir.FullName);
                     }
                 }
             }
