fix for spaces in path on xbox one
fix to fall back to appxmanifest based on executable directory for content only projects on deploy

diff --git a/Engine/Source/Programs/AutomationTool/Scripts/DeployCommand.Automation.cs b/Engine/Source/Programs/AutomationTool/Scripts/DeployCommand.Automation.cs
--- a/Engine/Source/Programs/AutomationTool/Scripts/DeployCommand.Automation.cs
+++ b/Engine/Source/Programs/AutomationTool/Scripts/DeployCommand.Automation.cs
@@ -28,7 +28,7 @@ public partial class Project : CommandUtils
 					string ClientCmdLine = "-run=Deploy ";
 					ClientCmdLine += "-Device=" + Params.Device + " ";
 					ClientCmdLine += "-Targetplatform=" + SC.StageTargetPlatform.PlatformType.ToString() + " ";
-					ClientCmdLine += "-SourceDir=" + CombinePaths(Params.BaseStageDirectory, SC.StageTargetPlatform.PlatformType.ToString()) + " ";
+					ClientCmdLine += "-SourceDir=\"" + CombinePaths(Params.BaseStageDirectory, SC.StageTargetPlatform.PlatformType.ToString()) + "\" ";
 					string ClientApp = CombinePaths(CmdEnv.LocalRoot, "Engine/Binaries/Win64/UnrealFrontend.exe");
 
 					Log("Deploying via UFE:");
