build script, allow branches to opt-out of mac

diff --git a/Engine/Source/Programs/AutomationTool/Scripts/GUBP.Automation.cs b/Engine/Source/Programs/AutomationTool/Scripts/GUBP.Automation.cs
--- a/Engine/Source/Programs/AutomationTool/Scripts/GUBP.Automation.cs
+++ b/Engine/Source/Programs/AutomationTool/Scripts/GUBP.Automation.cs
@@ -18,7 +18,6 @@ public class GUBP : BuildCommand
     public List<UnrealTargetPlatform> ActivePlatforms = null;
     public BranchInfo Branch = null;
     public bool bOrthogonalizeEditorPlatforms = false;
-    public static bool bHackRunIOSCompilesOnMac = false;
     public List<UnrealTargetPlatform> HostPlatforms;
     public bool bFake = false;
     public static bool bNoIOSOnPC = false;
@@ -91,6 +90,48 @@ public class GUBP : BuildCommand
         }
     }
 
+    public abstract class GUBPBranchHacker
+    {
+        public class BranchOptions
+        {
+            public bool NoMac = false;
+        }
+        public virtual void ModifyOptions(GUBP bp, ref BranchOptions Options, string Branch)
+        {
+        }
+    }
+
+    private static List<GUBPBranchHacker> BranchHackers;
+    private GUBPBranchHacker.BranchOptions GetBranchOptions(string Branch)
+    {
+        if (BranchHackers == null)
+        {
+            BranchHackers = new List<GUBPBranchHacker>();
+            Assembly[] LoadedAssemblies = AppDomain.CurrentDomain.GetAssemblies();
+            foreach (var Dll in LoadedAssemblies)
+            {
+                Type[] AllTypes = Dll.GetTypes();
+                foreach (var PotentialConfigType in AllTypes)
+                {
+                    if (PotentialConfigType != typeof(GUBPBranchHacker) && typeof(GUBPBranchHacker).IsAssignableFrom(PotentialConfigType))
+                    {
+                        GUBPBranchHacker Config = Activator.CreateInstance(PotentialConfigType) as GUBPBranchHacker;
+                        if (Config != null)
+                        {
+                            BranchHackers.Add(Config);
+                        }
+                    }
+                }
+            }
+        }
+        var Result = new GUBPBranchHacker.BranchOptions();
+        foreach (var Hacker in BranchHackers)
+        {
+            Hacker.ModifyOptions(this, ref Result, Branch);
+        }
+        return Result;
+    }
+
     public abstract class GUBPEmailHacker
     {
         public virtual List<string> AddEmails(GUBP bp, string Branch, string NodeName)
@@ -1068,7 +1109,7 @@ public class GUBP : BuildCommand
         UnrealTargetPlatform TargetPlatform;
 
         public GamePlatformMonolithicsNode(GUBP bp, UnrealTargetPlatform InHostPlatform, BranchInfo.BranchUProject InGameProj, UnrealTargetPlatform InTargetPlatform)
-            : base((GUBP.bHackRunIOSCompilesOnMac && InTargetPlatform == UnrealTargetPlatform.IOS) ? UnrealTargetPlatform.Mac : InHostPlatform)
+            : base(InHostPlatform)
         {
             GameProj = InGameProj;
             TargetPlatform = InTargetPlatform;
@@ -1096,13 +1137,13 @@ public class GUBP : BuildCommand
             if (InGameProj.Options(InHostPlatform).bTestWithShared)  /// compiling templates is only for testing purposes, and we will group them to avoid saturating the farm
             {
                 AddPseudodependency(WaitForTestShared.StaticGetFullName());
-                AgentSharingGroup = "TemplateMonolithics" + StaticGetHostPlatformSuffix((GUBP.bHackRunIOSCompilesOnMac && InTargetPlatform == UnrealTargetPlatform.IOS) ? UnrealTargetPlatform.Mac : InHostPlatform);
+                AgentSharingGroup = "TemplateMonolithics" + StaticGetHostPlatformSuffix(InHostPlatform);
             }
         }
 
         public static string StaticGetFullName(UnrealTargetPlatform InHostPlatform, BranchInfo.BranchUProject InGameProj, UnrealTargetPlatform InTargetPlatform)
         {
-            return InGameProj.GameName + "_" + InTargetPlatform + "_Monolithics" + StaticGetHostPlatformSuffix((GUBP.bHackRunIOSCompilesOnMac && InTargetPlatform == UnrealTargetPlatform.IOS) ? UnrealTargetPlatform.Mac : InHostPlatform);
+            return InGameProj.GameName + "_" + InTargetPlatform + "_Monolithics" + StaticGetHostPlatformSuffix(InHostPlatform);
         }
         public override string GetFullName()
         {
@@ -1269,7 +1310,7 @@ public class GUBP : BuildCommand
         }
         public override void DoBuild(GUBP bp)
         {
-            if (GUBP.bHackRunIOSCompilesOnMac && TargetPlatform == UnrealTargetPlatform.IOS)
+            if (!UnrealBuildTool.Utils.IsRunningOnMono && TargetPlatform == UnrealTargetPlatform.IOS)
             {
                 throw new AutomationException("these rocket header node require real mac path.");
             }
@@ -3871,7 +3912,18 @@ public class GUBP : BuildCommand
         {
             HostPlatforms.Add(UnrealTargetPlatform.Win64);
         }
-        if (!ParseParam("NoMac"))
+        var BranchForOptions = "";
+        if (P4Enabled)
+        {
+            BranchForOptions = P4Env.BuildRootP4;
+        }
+        var BranchOptions = GetBranchOptions(BranchForOptions);
+        bool WithMac = !BranchOptions.NoMac;
+        if (ParseParam("NoMac"))
+        {
+            WithMac = false;
+        }
+        if (WithMac)
         {
             HostPlatforms.Add(UnrealTargetPlatform.Mac);
         }
@@ -3895,7 +3947,6 @@ public class GUBP : BuildCommand
         bFake = ParseParam("fake");
         bool bFakeEC = ParseParam("FakeEC");
         TimeIndex = ParseParamInt("TimeIndex", 0);
-        bHackRunIOSCompilesOnMac = !ParseParam("IOSOnRPCUtility") && !HostPlatforms.Contains(UnrealTargetPlatform.Mac);
 
         bNoIOSOnPC = HostPlatforms.Contains(UnrealTargetPlatform.Mac);
 
@@ -4159,10 +4210,6 @@ public class GUBP : BuildCommand
         AddNode(new VersionFilesNode());
 
 
-        if (bHackRunIOSCompilesOnMac)
-        {
-            AddNode(new ToolsForCompileNode(UnrealTargetPlatform.Mac));
-        }
         foreach (var HostPlatform in HostPlatforms)
         {
             AddNode(new ToolsForCompileNode(HostPlatform));
