Draft of rocket build cleanup script with GET calls to BuildInfo functioning

diff --git a/Engine/Source/Programs/AutomationTool/MCPPublic.cs b/Engine/Source/Programs/AutomationTool/MCPPublic.cs
--- a/Engine/Source/Programs/AutomationTool/MCPPublic.cs
+++ b/Engine/Source/Programs/AutomationTool/MCPPublic.cs
@@ -404,6 +404,33 @@ namespace EpicGames.MCP.Automation
 		/// <param name="McpConfigName">Which BuildInfo backend to label the build in.</param>
 		abstract public void LabelBuild(BuildPatchToolStagingInfo StagingInfo, string DestinationLabelWithPlatform, string McpConfigName);
 
+		/// <summary>
+		/// Get all labels for an app from the requested MCP environment
+		/// </summary>
+		/// <param name="AppName">App name to pull labels for</param>
+		/// <param name="McpConfigName">Which BuildInfo backend to label the build in.</param>
+		abstract public Dictionary<string, string> GetAppLabelNames(string AppName, string McpConfigName);
+
+		/// <summary>
+		/// Request the list of apps from an environment's buildinfo service, including all samples
+		/// </summary>
+		/// <param name="McpConfigName">Which BuildInfo backend to get apps from.</param>
+		abstract public List<string> GetAppNames(string McpConfigName);
+
+		/// <summary>
+		/// Get all builds for an app from the requested MCP environment
+		/// </summary>
+		/// <param name="AppName">App name to pull builds for</param>
+		/// <param name="McpConfigName">Which BuildInfo MCP backend to get builds from.</param>
+		abstract public List<string> GetAppBuildVersions(string AppName, string McpConfigName);
+
+		/// <summary>
+		/// Delete a build from the requested MCP environment
+		/// </summary>
+		/// <param name="BuildVersion">Unique ID for the Build to delete</param>
+		/// <param name="McpConfigName">Which BuildInfo MCP backend to delete the build in.</param>
+		abstract public void DeleteBuild(string BuildVersion, string McpConfigName);
+
         /// <summary>
         /// Informs Patcher Service of a new build availability after async labeling is complete
         /// (this usually means the build was copied to a public file server before the label could be applied).
@@ -461,7 +488,7 @@ namespace EpicGames.MCP.Automation
         /// <param name="command">Build command (used to allow the -CDNDrive cmdline override).</param>
         /// <param name="stagingInfo">Staging info used to determine where the chunks are to copy.</param>
         abstract public void MirrorCloudDirToInternalCDN(BuildPatchToolStagingInfo stagingInfo);
-    }
+	}
     /// <summary>
     /// Helpers for using the MCP account service
     /// </summary>
