Added error handling to CodeSign in iOSPlatform automation script

diff --git a/Engine/Source/Programs/AutomationTool/IOS/IOSPlatform.Automation.cs b/Engine/Source/Programs/AutomationTool/IOS/IOSPlatform.Automation.cs
--- a/Engine/Source/Programs/AutomationTool/IOS/IOSPlatform.Automation.cs
+++ b/Engine/Source/Programs/AutomationTool/IOS/IOSPlatform.Automation.cs
@@ -313,11 +313,15 @@ public class IOSPlatform : Platform
 		Arguments += " - iOS'";
 		Arguments += " -configuration " + TargetConfig.ToString();
 		Arguments += " CODE_SIGN_IDENTITY=" + (Distribution ? "\"iPhone Distribution\"" : "\"iPhone Developer\"");
-		Run ("/usr/bin/env", Arguments, null, ERunOptions.Default);
+		ProcessResult Result = Run ("/usr/bin/env", Arguments, null, ERunOptions.Default);
 		if (bWasGenerated)
 		{
 			InternalUtils.SafeDeleteDirectory( XcodeProj, true);
 		}
+		if (Result.ExitCode != 0)
+		{
+			throw new AutomationException("CodeSign Failed");
+		}
 	}
 
 	private void PackageIPA(string BaseDirectory, string GameName, string ProjectName, string ProjectDirectory, UnrealTargetConfiguration TargetConfig, bool Distribution = false)
