Implemented DeleteBuild and corrected some logic errors

diff --git a/Engine/Source/Programs/AutomationTool/MCPPublic.cs b/Engine/Source/Programs/AutomationTool/MCPPublic.cs
--- a/Engine/Source/Programs/AutomationTool/MCPPublic.cs
+++ b/Engine/Source/Programs/AutomationTool/MCPPublic.cs
@@ -429,7 +429,7 @@ namespace EpicGames.MCP.Automation
 		/// </summary>
 		/// <param name="BuildVersion">Unique ID for the Build to delete</param>
 		/// <param name="McpConfigName">Which BuildInfo MCP backend to delete the build in.</param>
-		abstract public void DeleteBuild(string BuildVersion, string McpConfigName);
+		abstract public void DeleteBuild(string BuildVersion, string AppName, string McpConfigName);
 
         /// <summary>
         /// Informs Patcher Service of a new build availability after async labeling is complete
