UAT: WatchdogTimer improvements

diff --git a/Engine/Source/Programs/AutomationTool/ProcessUtils.cs b/Engine/Source/Programs/AutomationTool/ProcessUtils.cs
--- a/Engine/Source/Programs/AutomationTool/ProcessUtils.cs
+++ b/Engine/Source/Programs/AutomationTool/ProcessUtils.cs
@@ -20,6 +20,13 @@ namespace AutomationTool
 		CTRL_SHUTDOWN_EVENT
 	}
 
+	public interface IProcess
+	{
+		void StopProcess(bool KillDescendants = true);
+		bool HasExited { get; }
+		string GetProcessName();
+	}
+
 	/// <summary>
 	/// Tracks all active processes.
 	/// </summary>
@@ -34,7 +41,7 @@ namespace AutomationTool
 		/// <summary>
 		/// List of active (running) processes.
 		/// </summary>
-		private static List<ProcessResult> ActiveProcesses = new List<ProcessResult>();
+		private static List<IProcess> ActiveProcesses = new List<IProcess>();
 		/// <summary>
 		/// Synchronization object
 		/// </summary>
@@ -63,14 +70,19 @@ namespace AutomationTool
 				}
 			}
 			var Result = new ProcessResult(AppName, NewProcess, bAllowSpew, LogName, SpewVerbosity: SpewVerbosity);
+			AddProcess(Result);
+			return Result;
+		}
+
+		public static void AddProcess(IProcess Proc)
+		{
 			lock (SyncObject)
 			{
-				ActiveProcesses.Add(Result);
+				ActiveProcesses.Add(Proc);
 			}
-			return Result;
 		}
 
-		public static void RemoveProcess(ProcessResult Proc)
+		public static void RemoveProcess(IProcess Proc)
 		{
 			lock (SyncObject)
 			{
@@ -88,7 +100,7 @@ namespace AutomationTool
 		/// </summary>
 		public static void KillAll()
 		{
-			List<ProcessResult> ProcessesToKill = new List<ProcessResult>();
+			List<IProcess> ProcessesToKill = new List<IProcess>();
 			lock (SyncObject)
 			{
 				foreach (var ProcResult in ActiveProcesses)
@@ -195,7 +207,7 @@ namespace AutomationTool
 	/// <summary>
 	/// Class containing the result of process execution.
 	/// </summary>
-	public class ProcessResult
+	public class ProcessResult : IProcess
 	{
 		private string Source = "";
 		private int ProcessExitCode = -1;
diff --git a/Engine/Source/Programs/AutomationTool/Scripts/Tests.Automation.cs b/Engine/Source/Programs/AutomationTool/Scripts/Tests.Automation.cs
--- a/Engine/Source/Programs/AutomationTool/Scripts/Tests.Automation.cs
+++ b/Engine/Source/Programs/AutomationTool/Scripts/Tests.Automation.cs
@@ -1282,18 +1282,37 @@ public class TestWatchdogTimer : BuildCommand
 {
 	public override void ExecuteBuild()
 	{
-		Log("Starting the first timer (1s). This should not crash.");
+		Log("Starting 1st timer (1s). This should not crash.");
 		using (var SafeTimer = new WatchdogTimer(1))
 		{
 			// Wait 500ms
+			Log("Started {0}", SafeTimer.GetProcessName());
 			Thread.Sleep(500);
 		}
 		Log("First timer disposed successfully.");
 
-		Log("Starting the second timer (2s). This should crash after 2 seconds.");
+		Log("Starting 2nd timer (2s). This should throw an exception after 1 second.");
+		try
+		{
+			using (var CrashTimer = new WatchdogTimer(2))
+			{
+				// Wait 5s (this will trigger the watchdog timer)
+				Log("Started {0}", CrashTimer.GetProcessName());
+				Thread.Sleep(1000);
+				throw new Exception("Test exceptions under WatchdogTimer");
+			}
+		}
+		catch (Exception Ex)
+		{
+			Log("Triggered exception guarded by WatchdogTimer:");
+			Log(System.Diagnostics.TraceEventType.Information, Ex);
+		}
+
+		Log("Starting 3rd timer (2s). This should crash after 2 seconds.");
 		using (var CrashTimer = new WatchdogTimer(2))
 		{
 			// Wait 5s (this will trigger the watchdog timer)
+			Log("Started {0}", CrashTimer.GetProcessName());
 			Thread.Sleep(5000);
 		}
 	}
diff --git a/Engine/Source/Programs/AutomationTool/WatchdogTimer.cs b/Engine/Source/Programs/AutomationTool/WatchdogTimer.cs
--- a/Engine/Source/Programs/AutomationTool/WatchdogTimer.cs
+++ b/Engine/Source/Programs/AutomationTool/WatchdogTimer.cs
@@ -13,10 +13,12 @@ namespace AutomationTool
 	/// <summary>
 	/// Timer that exits the application of the execution of the command takes too long.
 	/// </summary>
-	public class WatchdogTimer : IDisposable
+	public class WatchdogTimer : IDisposable, IProcess
 	{
 		private Timer CountownTimer;
 		private System.Diagnostics.StackFrame TimerFrame;
+		private object SyncObject = new object();
+		private bool bTicking = false;
 
 		[MethodImplAttribute(MethodImplOptions.NoInlining)]
 		public WatchdogTimer(int Seconds)
@@ -25,22 +27,61 @@ namespace AutomationTool
 			CountownTimer = new Timer(Seconds * 1000.0);
 			CountownTimer.Elapsed += Elapsed;
 			CountownTimer.Start();
+			bTicking = true;
+			ProcessManager.AddProcess(this);
 		}
 
 		void Elapsed(object sender, ElapsedEventArgs e)
 		{
+			lock (SyncObject)
+			{
+				bTicking = false;
+			}
+			int Interval = (int)(CountownTimer.Interval * 0.001);
+			Dispose();
 			var Method = TimerFrame.GetMethod();
-			Log.TraceError("BUILD FAILED: WatchdogTimer timed out after {0}s in {1}.{2}", (int)(CountownTimer.Interval * 0.001), Method.DeclaringType.Name, Method.Name);
+			Log.TraceError("BUILD FAILED: WatchdogTimer timed out after {0}s in {1}.{2}", Interval, Method.DeclaringType.Name, Method.Name);			
 			Environment.Exit(1);
 		}
 
 		public void Dispose()
 		{
-			if (CountownTimer != null)
+			Log.TraceInformation("WatchdogTimer.Dispose()");
+			lock (SyncObject)
+			{
+				if (CountownTimer != null)
+				{
+					CountownTimer.Stop();
+					CountownTimer.Dispose();
+					CountownTimer = null;
+				}
+				bTicking = false;
+			}
+			ProcessManager.RemoveProcess(this);
+		}
+
+		public void StopProcess(bool KillDescendants = true)
+		{
+			Dispose();			
+		}
+
+		public bool HasExited
+		{
+			get 
 			{
-				CountownTimer.Dispose();
-				CountownTimer = null;
+				bool bResult = true;
+ 				lock (SyncObject)
+				{
+					bResult = bTicking;
+				}
+				return bResult;
 			}
 		}
+
+		public string GetProcessName()
+		{
+			var Method = TimerFrame.GetMethod();
+			return String.Format("WatchdogTimer_{0}.{1}", Method.DeclaringType.Name, Method.Name);
+		}
 	}
 }
