build script, fix temp storage cleaner to only do a dir once per run

diff --git a/Engine/Source/Programs/AutomationTool/TempStorage.cs b/Engine/Source/Programs/AutomationTool/TempStorage.cs
--- a/Engine/Source/Programs/AutomationTool/TempStorage.cs
+++ b/Engine/Source/Programs/AutomationTool/TempStorage.cs
@@ -493,9 +493,15 @@ namespace AutomationTool
         {
             return CombinePaths(Env.LocalRoot, "Engine", "Saved", "GUBP");
         }
-
+        static HashSet<string> TopDirectoryTestedForClean = new HashSet<string>();
         static void CleanSharedTempStorage(string TopDirectory)
         {
+            if (TopDirectoryTestedForClean.Contains(TopDirectory))
+            {
+                return;
+            }
+            TopDirectoryTestedForClean.Add(TopDirectory);
+
             const int MaximumDaysToKeepTempStorage = 2;
             var StartTimeDir = DateTime.UtcNow;
             DirectoryInfo DirInfo = new DirectoryInfo(TopDirectory);
