UAT UnrealSync: Add -preview option. Improved P4 performance. Changed the label listing parameters to:
-listpromotedlabels -- lists all promoted labels for a branch. If game parameter is set then list them for a game, if not for shared promotable.
-listlabels -- lists all labels for a branch.

diff --git a/Engine/Source/Programs/AutomationTool/P4Utils.cs b/Engine/Source/Programs/AutomationTool/P4Utils.cs
--- a/Engine/Source/Programs/AutomationTool/P4Utils.cs
+++ b/Engine/Source/Programs/AutomationTool/P4Utils.cs
@@ -1690,6 +1690,27 @@ namespace AutomationTool
 			return LabelNameList.ToArray();
 		}
 
+		/// <summary>
+		/// Gets labels list for given branch.
+		/// </summary>
+		/// <param name="BranchPath">A branch path.</param>
+		/// <returns>List of labels for given branch path.</returns>
+		public string[] GetBranchLabels(string BranchPath)
+		{
+			var LabelNameList = new List<string>();
+
+			string Output;
+			if (LogP4Output(out Output, "labels -t -e " + BranchPath + "/*"))
+			{
+				foreach (Match LabelMatch in LabelsListOutputPattern.Matches(Output))
+				{
+					LabelNameList.Add(LabelMatch.Groups["name"].Value);
+				}
+			}
+
+			return LabelNameList.ToArray();
+		}
+
 		/// <summary>
 		/// Get latest promoted label given branch and game name.
 		/// </summary>
diff --git a/Engine/Source/Programs/AutomationTool/Scripts/UnrealSync.Automation.cs b/Engine/Source/Programs/AutomationTool/Scripts/UnrealSync.Automation.cs
--- a/Engine/Source/Programs/AutomationTool/Scripts/UnrealSync.Automation.cs
+++ b/Engine/Source/Programs/AutomationTool/Scripts/UnrealSync.Automation.cs
@@ -10,16 +10,21 @@ using UnrealBuildTool;
 namespace AutomationScripts.Automation
 {
 	[Help("Syncs promotable build. Use either -game or -label parameter.")]
-	[Help("listlabels", "Enables special mode that lists all possible promoted labels for shared promotable or given game.")]
+	[Help("listpromotedlabels", "Enables special mode that lists all possible promoted labels for shared promotable or given game.")]
+	[Help("listlabels", "Enables special mode that lists all labels for current branch.")]
 	[Help("artist", "Artist sync i.e. sync content to head and all the rest to promoted label.")]
+	[Help("preview", "This option makes that syncs are done in the preview mode (i.e. p4 sync -n).")]
 	[Help("game=GameName", "Name of the game to sync. If not set then shared promotable will be synced.")]
 	[Help("label=LabelName", "Promotable label name to sync to.")]
 	[RequireP4]
+	[DoesNotNeedP4CL]
 	class UnrealSync : BuildCommand
 	{
 		public override void ExecuteBuild()
 		{
 			var List = ParseParam("listlabels");
+			var ListPromoted = ParseParam("listpromotedlabels");
+			var Preview = ParseParam("preview");
 			var ArtistSync = ParseParam("artist");
 			var BranchPath = P4Env.BuildRootP4;
 
@@ -34,17 +39,27 @@ namespace AutomationScripts.Automation
 					? string.Format("branch {0} shared-promotable", BranchPath)
 					: string.Format("branch {0} and game {1}", BranchPath, GameName);
 
-			if (List)
+			if (List || ListPromoted)
 			{
 				if (string.IsNullOrWhiteSpace(BranchPath))
 				{
 					throw new AutomationException("The branch path is not set. Something went wrong.");
 				}
 
-				Log(string.Format("Promoted labels for {0}.", BranchAndGameName));
+				if (ListPromoted)
+				{
+					Log(string.Format("Promoted labels for {0}.", BranchAndGameName));
 
-				// The P4 command will log out the possible labels.
-				P4.GetPromotedLabels(BranchPath, GameName);
+					// The P4 command will log out the possible labels.
+					P4.GetPromotedLabels(BranchPath, GameName);
+				}
+				else
+				{
+					Log(string.Format("Labels for {0}.", BranchPath));
+
+					// The P4 command will log out the possible labels.
+					P4.GetBranchLabels(BranchPath);
+				}
 
 				return;
 			}
@@ -57,7 +72,7 @@ namespace AutomationScripts.Automation
 
 				if (string.IsNullOrWhiteSpace(LabelParam))
 				{
-					throw new AutomationException("Use either -branch and -game or just -label parameter.");
+					throw new AutomationException("Use either -game or -label parameter.");
 				}
 
 				if (!ValidateAndParseLabelName(LabelParam, BranchPath, out GameName) || !P4.ValidateLabelContent(LabelParam))
@@ -77,6 +92,18 @@ namespace AutomationScripts.Automation
 				throw new AutomationException(string.Format("Label for {0} was not found.", BranchAndGameName));
 			}
 
+			var ProgramRevisionSpec = "@" + ProgramSyncLabelName;
+
+			// Get latest CL number to sync cause @head can change during
+			// different syncs and it could create integrity problems in
+			// workspace.
+			var ContentRevisionSpec = "@" +
+			(
+				ArtistSync
+				? P4.GetLatestCLNumber().ToString()
+				: ProgramSyncLabelName
+			);
+
 			var ArtistSyncRulesPath = string.Format("{0}/{1}/Build/ArtistSyncRules.xml",
 				BranchPath, string.IsNullOrWhiteSpace(GameName) ? "Samples" : GameName);
 
@@ -87,26 +114,17 @@ namespace AutomationScripts.Automation
 				throw new AutomationException(string.Format("The path {0} is not valid or file is empty.", ArtistSyncRulesPath));
 			}
 
-			var ProgramRevisionSpec = "@" + ProgramSyncLabelName;
-
-			// Get latest CL number to sync cause @head can change during
-			// different syncs and it could create integrity problems in
-			// workspace.
-			var ContentRevisionSpec = "@" +
-				(
-					ArtistSync
-					? P4.GetLatestCLNumber().ToString()
-					: ProgramSyncLabelName
-				);
-
 			foreach (var SyncStep in GenerateSyncSteps(SyncRules, ContentRevisionSpec, ProgramRevisionSpec))
 			{
-				P4.Sync(BranchPath + SyncStep);
+				P4.Sync((Preview ? "-n " : "") + BranchPath + SyncStep);
 			}
 		}
 
-		/* Pattern to parse branch path and game name from label name. */
-		private static readonly Regex PromotedLabelPattern = new Regex(@"^(?<branchPath>//depot/UE4)/Promoted-((?<gameName>\w+)-)?CL-\d+$", RegexOptions.Compiled);
+		/* Pattern to parse branch path and the rest from label name. */
+		private static readonly Regex LabelPattern = new Regex(@"^(?<branchPath>//depot/.+)/(?<rest>[^/]+)$", RegexOptions.Compiled);
+
+		/* Pattern to parse the rest of the label name (without branch path). */
+		private static readonly Regex PromotedLabelRestPattern = new Regex(@"^Promoted-((?<gameName>\w+)-)?CL-\d+$", RegexOptions.Compiled);
 
 		/// <summary>
 		/// Validates and parses promotable label name.
@@ -117,23 +135,25 @@ namespace AutomationScripts.Automation
 		/// <returns>True if label is valid. False otherwise.</returns>
 		private bool ValidateAndParseLabelName(string LabelName, string BranchPath, out string GameName)
 		{
-			var Match = PromotedLabelPattern.Match(LabelName);
+			var Match = LabelPattern.Match(LabelName);
+
+			GameName = null;
 
 			if (!Match.Success)
 			{
-				GameName = null;
-
 				return false;
 			}
 
 			if(BranchPath != Match.Groups["branchPath"].Value)
 			{
-				GameName = null;
-
 				return false;
 			}
 
-			GameName = Match.Groups["gameName"].Value;
+			var GameMatch = PromotedLabelRestPattern.Match(Match.Groups["rest"].Value);
+			if(GameMatch.Success)
+			{
+				GameName = GameMatch.Groups["gameName"].Value;
+			}
 
 			return true;
 		}
@@ -145,7 +165,7 @@ namespace AutomationScripts.Automation
 		/// <param name="ContentRevisionSpec">Revision spec to which sync the content. Different for artist sync.</param>
 		/// <param name="ProgramRevisionSpec">Revision spec to which sync everything except content.</param>
 		/// <returns>An array of sync steps to perform.</returns>
-		private string[] GenerateSyncSteps(string SyncRules, string ContentRevisionSpec, string ProgramRevisionSpec)
+		private List<string> GenerateSyncSteps(string SyncRules, string ContentRevisionSpec, string ProgramRevisionSpec)
 		{
 			var SyncRulesDocument = new XmlDocument();
 			SyncRulesDocument.LoadXml(SyncRules);
@@ -169,7 +189,7 @@ namespace AutomationScripts.Automation
 				OutputList.Add(SyncStep);
 			}
 
-			return OutputList.ToArray();
+			return OutputList;
 		}
 
 	}
