UAT UnrealSync: Minor refactoring.

diff --git a/Engine/Source/Programs/AutomationTool/Scripts/UnrealSync.Automation.cs b/Engine/Source/Programs/AutomationTool/Scripts/UnrealSync.Automation.cs
--- a/Engine/Source/Programs/AutomationTool/Scripts/UnrealSync.Automation.cs
+++ b/Engine/Source/Programs/AutomationTool/Scripts/UnrealSync.Automation.cs
@@ -26,12 +26,7 @@ namespace AutomationScripts.Automation
 			var ListPromoted = ParseParam("listpromotedlabels");
 			var Preview = ParseParam("preview");
 			var ArtistSync = ParseParam("artist");
-			var BranchPath = P4Env.BuildRootP4;
-
-			if (BranchPath != null && BranchPath.EndsWith("/"))
-			{
-				BranchPath = BranchPath.Substring(0, BranchPath.Length - 1);
-			}
+			var BranchPath = CommandUtils.GetDirectoryName(P4Env.BuildRootP4);
 
 			var GameName = ParseParamValue("game");
 
@@ -48,14 +43,14 @@ namespace AutomationScripts.Automation
 
 				if (ListPromoted)
 				{
-					Log(string.Format("Promoted labels for {0}.", BranchAndGameName));
+					Log("Promoted labels for {0}.", BranchAndGameName);
 
 					// The P4 command will log out the possible labels.
 					P4.GetPromotedLabels(BranchPath, GameName);
 				}
 				else
 				{
-					Log(string.Format("Labels for {0}.", BranchPath));
+					Log("Labels for {0}.", BranchPath);
 
 					// The P4 command will log out the possible labels.
 					P4.GetBranchLabels(BranchPath);
@@ -89,7 +84,7 @@ namespace AutomationScripts.Automation
 
 			if (ProgramSyncLabelName == null)
 			{
-				throw new AutomationException(string.Format("Label for {0} was not found.", BranchAndGameName));
+				throw new AutomationException("Label for {0} was not found.", BranchAndGameName);
 			}
 
 			var ProgramRevisionSpec = "@" + ProgramSyncLabelName;
@@ -111,7 +106,7 @@ namespace AutomationScripts.Automation
 
 			if (string.IsNullOrWhiteSpace(SyncRules))
 			{
-				throw new AutomationException(string.Format("The path {0} is not valid or file is empty.", ArtistSyncRulesPath));
+				throw new AutomationException("The path {0} is not valid or file is empty.", ArtistSyncRulesPath);
 			}
 
 			foreach (var SyncStep in GenerateSyncSteps(SyncRules, ContentRevisionSpec, ProgramRevisionSpec))
