Fixing BuildFarms

diff --git a/Engine/Source/Programs/AutomationTool/ProjectUtils.cs b/Engine/Source/Programs/AutomationTool/ProjectUtils.cs
--- a/Engine/Source/Programs/AutomationTool/ProjectUtils.cs
+++ b/Engine/Source/Programs/AutomationTool/ProjectUtils.cs
@@ -181,14 +181,17 @@ namespace AutomationTool
 			}
 
 			// Get all ini files
-			var EngineDirectory = CommandUtils.CombinePaths(CommandUtils.CmdEnv.LocalRoot, "Engine");
-			CommandUtils.Log("Loading ini files for {0}", RawProjectPath);
-			foreach (UnrealTargetPlatform TargetPlatformType in Enum.GetValues(typeof(UnrealTargetPlatform)))
+			if (!String.IsNullOrWhiteSpace(RawProjectPath))
 			{
-				if (TargetPlatformType != UnrealTargetPlatform.Unknown)
+				var EngineDirectory = CommandUtils.CombinePaths(CommandUtils.CmdEnv.LocalRoot, "Engine");
+				CommandUtils.Log("Loading ini files for {0}", RawProjectPath);
+				foreach (UnrealTargetPlatform TargetPlatformType in Enum.GetValues(typeof(UnrealTargetPlatform)))
 				{
-					var Config = new ConfigCacheIni(TargetPlatformType, "Engine", Path.GetDirectoryName(RawProjectPath), EngineDirectory);
-					Properties.EngineConfigs.Add(TargetPlatformType, Config);
+					if (TargetPlatformType != UnrealTargetPlatform.Unknown)
+					{
+						var Config = new ConfigCacheIni(TargetPlatformType, "Engine", Path.GetDirectoryName(RawProjectPath), EngineDirectory);
+						Properties.EngineConfigs.Add(TargetPlatformType, Config);
+					}
 				}
 			}
 
