Fix running content-only projects from UAT without packaging.

diff --git a/Engine/Source/Programs/AutomationTool/Win/WinPlatform.Automation.cs b/Engine/Source/Programs/AutomationTool/Win/WinPlatform.Automation.cs
--- a/Engine/Source/Programs/AutomationTool/Win/WinPlatform.Automation.cs
+++ b/Engine/Source/Programs/AutomationTool/Win/WinPlatform.Automation.cs
@@ -193,7 +193,7 @@ public abstract class BaseWinPlatform : Platform
 				foreach (var StageExecutable in SC.StageExecutables)
 				{
 					string ExeName = SC.StageTargetPlatform.GetPlatformExecutableName(StageExecutable);
-					if(!SC.IsCodeBasedProject && !bIsRun)
+					if(!SC.IsCodeBasedProject && (!bIsRun || !SC.Stage))
 					{
 						ExecutableNames.Add(CombinePaths(SC.RuntimeRootDir, "Engine/Binaries", SC.PlatformDir, ExeName + Ext));
 					}
