build script, add mmc to the list of processes not to kill

diff --git a/Engine/Source/Programs/AutomationTool/WindowsHostPlatform.cs b/Engine/Source/Programs/AutomationTool/WindowsHostPlatform.cs
--- a/Engine/Source/Programs/AutomationTool/WindowsHostPlatform.cs
+++ b/Engine/Source/Programs/AutomationTool/WindowsHostPlatform.cs
@@ -218,7 +218,8 @@ namespace AutomationTool
 			"services",
 			"lsass",
             "conhost",
-            "oobe"
+            "oobe",
+            "mmc"
 		};
 		public override string[] DontKillProcessList
 		{
