UAT: A few tweaks to prevent blocking operations on app domain unload

diff --git a/Engine/Source/Programs/AutomationTool/Log.cs b/Engine/Source/Programs/AutomationTool/Log.cs
--- a/Engine/Source/Programs/AutomationTool/Log.cs
+++ b/Engine/Source/Programs/AutomationTool/Log.cs
@@ -14,6 +14,8 @@ namespace AutomationTool
 
 	public class LogUtils
 	{
+		private static AutomationFileTraceListener FileLog;
+
 		/// <summary>
 		/// Initializes trace logging.
 		/// </summary>
@@ -23,7 +25,8 @@ namespace AutomationTool
 			var VerbosityLevel = CommandUtils.ParseParam(CommandLine, "-Verbose") ? TraceEventType.Verbose : TraceEventType.Information;
 			var Filter = new VerbosityFilter(VerbosityLevel);
 			Trace.Listeners.Add(new AutomationConsoleTraceListener());
-			Trace.Listeners.Add(new AutomationFileTraceListener());
+			FileLog = new AutomationFileTraceListener();
+			Trace.Listeners.Add(FileLog);
 			Trace.Listeners.Add(new AutomationMemoryLogListener());
 			foreach (TraceListener Listener in Trace.Listeners)
 			{
@@ -31,6 +34,20 @@ namespace AutomationTool
 			}
 		}
 
+		/// <summary>
+		/// Closes logging system
+		/// </summary>
+		public static void CloseFileLogging()
+		{
+			if (FileLog != null)
+			{
+				Trace.Listeners.Remove(FileLog);
+				FileLog.Close();
+				FileLog.Dispose();
+				FileLog = null;
+			}
+		}
+
 		/// <summary>
 		/// Formats message for logging.
 		/// </summary>
@@ -406,26 +423,25 @@ namespace AutomationTool
 					LogFile.Close();
 					LogFile.Dispose();
 					LogFile = null;
-
-					try
-					{
-						// Try to copy the log file to the log folder. The reason why it's done here is that
-						// at the time the log file is being initialized the env var may not yet be set (this 
-						// applies to local log folder in particular)
-						var LogFolder = Environment.GetEnvironmentVariable(EnvVarNames.LogFolder);
-						if (!String.IsNullOrEmpty(LogFolder) && Directory.Exists(LogFolder) &&
-								!String.IsNullOrEmpty(LogFilename) && File.Exists(LogFilename))
-						{
-							var DestFilename = CommandUtils.CombinePaths(LogFolder, "UAT_" + Path.GetFileName(LogFilename));
-							SafeCopyLogFile(LogFilename, DestFilename);
-						}
-					}
-					catch (Exception)
-					{
-						// Silently ignore, logging is pointless because eveything is shut down at this point
-					}
 				}
 			}
+			try
+			{
+				// Try to copy the log file to the log folder. The reason why it's done here is that
+				// at the time the log file is being initialized the env var may not yet be set (this 
+				// applies to local log folder in particular)
+				var LogFolder = Environment.GetEnvironmentVariable(EnvVarNames.LogFolder);
+				if (!String.IsNullOrEmpty(LogFolder) && Directory.Exists(LogFolder) &&
+						!String.IsNullOrEmpty(LogFilename) && File.Exists(LogFilename))
+				{
+					var DestFilename = CommandUtils.CombinePaths(LogFolder, "UAT_" + Path.GetFileName(LogFilename));
+					SafeCopyLogFile(LogFilename, DestFilename);
+				}
+			}
+			catch (Exception)
+			{
+				// Silently ignore, logging is pointless because eveything is shut down at this point
+			}
 			base.Close();
 		}
 
diff --git a/Engine/Source/Programs/AutomationTool/Program.cs b/Engine/Source/Programs/AutomationTool/Program.cs
--- a/Engine/Source/Programs/AutomationTool/Program.cs
+++ b/Engine/Source/Programs/AutomationTool/Program.cs
@@ -67,6 +67,8 @@ namespace AutomationTool
 			}
 
 			Log.WriteLine(TraceEventType.Information, "AutomationTool exiting with ExitCode={0}", Result);
+			LogUtils.CloseFileLogging();
+
 			return Result;
 		}
 
@@ -77,19 +79,18 @@ namespace AutomationTool
 			{
 				// Force exit
 				Environment.Exit(3);
-			}
+			}			
 			return true;
 		}
 
 		static void Domain_ProcessExit(object sender, EventArgs e)
 		{
-			// Kill all spawned processes.
+			// Kill all spawned processes (Console instead of Log because logging is closed at this time anyway)
 			Console.WriteLine("Domain_ProcessExit");
 			if (ShouldKillProcesses)
 			{			
 				ProcessManager.KillAll();
 			}
-			// Close logging
 			Trace.Close();
 		}
 
diff --git a/Engine/Source/Programs/AutomationToolLauncher/Launcher.cs b/Engine/Source/Programs/AutomationToolLauncher/Launcher.cs
--- a/Engine/Source/Programs/AutomationToolLauncher/Launcher.cs
+++ b/Engine/Source/Programs/AutomationToolLauncher/Launcher.cs
@@ -45,6 +45,7 @@ namespace AutomationToolLauncher
             {
                 Console.WriteLine("AutomationToolLauncher failed to execute the assembly");
                 Console.WriteLine(Ex.Message);
+								Console.WriteLine(Ex.StackTrace);
             }
             int Retry = 5;
             bool Unloaded = false;
