UAT UnrealSync: Made all default type of querty for UnrealSyncList.

diff --git a/Engine/Source/Programs/AutomationTool/Scripts/UnrealSync.Automation.cs b/Engine/Source/Programs/AutomationTool/Scripts/UnrealSync.Automation.cs
--- a/Engine/Source/Programs/AutomationTool/Scripts/UnrealSync.Automation.cs
+++ b/Engine/Source/Programs/AutomationTool/Scripts/UnrealSync.Automation.cs
@@ -129,7 +129,14 @@ namespace AutomationScripts.Automation
 		/// <returns>Enum value of the query type param.</returns>
 		private QueryType GetTypeParam()
 		{
-			switch(ParseParamValue("type").ToLower())
+			var TypeString = ParseParamValue("type");
+
+			if(string.IsNullOrWhiteSpace(TypeString))
+			{
+				return QueryType.All;
+			}
+
+			switch (TypeString.ToLower())
 			{
 				case "all":
 					return QueryType.All;
