build script, additional logging to see why we aren't going to delete formal builds

diff --git a/Engine/Source/Programs/AutomationTool/CommandUtils.cs b/Engine/Source/Programs/AutomationTool/CommandUtils.cs
--- a/Engine/Source/Programs/AutomationTool/CommandUtils.cs
+++ b/Engine/Source/Programs/AutomationTool/CommandUtils.cs
@@ -1888,7 +1888,7 @@ namespace AutomationTool
 
                 DirectoryInfo DirInfo = new DirectoryInfo(ParentDir);
                 var TopLevelDirs = DirInfo.GetDirectories();
-                Log("Looking for directories to delete based on {0}", DirectoryForThisBuild);
+                Log("Looking for directories to delete in {0}   {1} dirs", ParentDir, TopLevelDirs.Length);
                 foreach (var TopLevelDir in TopLevelDirs)
                 {
                     if (DirectoryExists_NoExceptions(TopLevelDir.FullName))
@@ -1911,6 +1911,14 @@ namespace AutomationTool
                                     //DeleteDirectory_NoExceptions(true, TopLevelDir.FullName);
                                 }
                             }
+                            else
+                            {
+                                Log("skipping {0}, because the CL part {1} had weird characters", TopLevelDir.FullName, CLPart);
+                            }
+                        }
+                        else
+                        {
+                            Log("skipping {0}, because it didn't start with {1} or end with {2}", TopLevelDir.FullName, StartString, EndString);
                         }
                     }
                 }
