build script, dial back the CIS frequency of lots and lots of things

diff --git a/Engine/Source/Programs/AutomationTool/Scripts/GUBP.Automation.cs b/Engine/Source/Programs/AutomationTool/Scripts/GUBP.Automation.cs
--- a/Engine/Source/Programs/AutomationTool/Scripts/GUBP.Automation.cs
+++ b/Engine/Source/Programs/AutomationTool/Scripts/GUBP.Automation.cs
@@ -1085,7 +1085,7 @@ public class GUBP : BuildCommand
             int Result = base.CISFrequencyQuantumShift(bp);
             if (GameProj.Options(HostPlatform).bTestWithShared)
             {
-                Result += 2; // only every 80m
+                Result += 3;
             }
             return Result;
         }
@@ -1162,11 +1162,11 @@ public class GUBP : BuildCommand
             int Result = base.CISFrequencyQuantumShift(bp);
             if (GameProj.GameName != bp.Branch.BaseEngineProject.GameName)
             {
-                Result += 2; // only every 80m
+                Result += 3; // only every 80m
             }
             else if (TargetPlatform != HostPlatform)
             {
-                Result += 1; // only every 40m
+                Result += 2; // only every 40m
             }
             return Result;
         }
@@ -1870,7 +1870,7 @@ public class GUBP : BuildCommand
         }
         public override int CISFrequencyQuantumShift(GUBP bp)
         {
-            return base.CISFrequencyQuantumShift(bp) + 4;
+            return base.CISFrequencyQuantumShift(bp) + 5;
         }
     }
 
