build script, allow UAT files to be not retrieved...still working on the hang

diff --git a/Engine/Source/Programs/AutomationTool/TempStorage.cs b/Engine/Source/Programs/AutomationTool/TempStorage.cs
--- a/Engine/Source/Programs/AutomationTool/TempStorage.cs
+++ b/Engine/Source/Programs/AutomationTool/TempStorage.cs
@@ -821,12 +821,27 @@ namespace AutomationTool
                 }
                 var RelativeFile = Filename.Substring(BlockPath.Length);
                 var DestFile = CombinePaths(BaseFolder, RelativeFile);
+                bool bDoCopy = true;
                 if (FileExists_NoExceptions(true, DestFile))
                 {
                     Log("Dest file {0} already exists, deleting and overwriting", DestFile);
-                    DeleteFile(DestFile);
+                    DeleteFile_NoExceptions(DestFile);
+                    if (FileExists_NoExceptions(true, DestFile))
+                    {
+                        bool bOkToBeDifferent = DestFile.Contains("AutomationTool.exe");
+                        bOkToBeDifferent = bOkToBeDifferent || DestFile.Contains("\\DotNET\\AutomationScripts\\") || DestFile.Contains("/DotNET/AutomationScripts/");
+                        if (bOkToBeDifferent)
+                        {
+                            bDoCopy = false;
+                        }
+                        System.Diagnostics.TraceEventType LogType = bOkToBeDifferent ? System.Diagnostics.TraceEventType.Warning : System.Diagnostics.TraceEventType.Error;
+                        Log(LogType, "Dest file {0} already exists and couldn't be deleted.", DestFile);
+                    }
+                }
+                if (bDoCopy)
+                {
+                    CopyFile(Filename, DestFile, true);
                 }
-                CopyFile(Filename, DestFile, true);
 
                 Robust_FileExists_NoExceptions(true, DestFile, "Could not copy to {0}");
 
