build script, more robust pull script

diff --git a/Engine/Source/Programs/AutomationTool/Scripts/Tests.Automation.cs b/Engine/Source/Programs/AutomationTool/Scripts/Tests.Automation.cs
--- a/Engine/Source/Programs/AutomationTool/Scripts/Tests.Automation.cs
+++ b/Engine/Source/Programs/AutomationTool/Scripts/Tests.Automation.cs
@@ -81,6 +81,115 @@ class GitPullRequest : BuildCommand
         return Result.Output.Trim();
     }
 
+    bool ScanForBranchAndCL_BaseVersion(string GitCommand, out string Depot, out int CL)
+    {
+        Depot = null;
+        CL = 0;
+        try
+        {
+            var Base = RunGit(GitCommand);
+
+            string BaseStart = "Engine source (";
+            string BaseEnd = ")";
+
+            if (Base.Contains(BaseStart) && Base.Contains(BaseEnd))
+            {
+                Base = Base.Substring(Base.IndexOf(BaseStart) + BaseStart.Length);
+                if (Base.StartsWith("4."))
+                {
+                    Depot = "//depot/UE4-Releases/4." + Base.Substring(2, 1);
+                }
+                else if (Base.StartsWith("Main"))
+                {
+                    Depot = "//depot/UE4";
+                }
+                else
+                {
+                    throw new AutomationException("Unrecognized branch.");
+                }
+                Log("Depot {0}", Depot);
+
+                Base = Base.Substring(0, Base.IndexOf(BaseEnd));
+                if (!Base.Contains(" "))
+                {
+                    throw new AutomationException("Unrecognized commit3.");
+                }
+                Base = Base.Substring(Base.LastIndexOf(" "));
+                Log("CL String {0}", Base);
+                CL = int.Parse(Base);
+            }
+            Log("CL int {0}", CL);
+            if (CL < 2000000 || String.IsNullOrWhiteSpace(Depot))
+            {
+                throw new AutomationException("Unrecognized commit3.");
+            }
+
+            return true;
+        }
+        catch (Exception)
+        {
+            CL = 0;
+            return false;
+        }
+    }
+    
+    bool ScanForBranchAndCL_LiveVersion(string GitCommand, out string Depot, out int CL)
+    {
+        Depot = null;
+        CL = 0;
+        try
+        {
+            var Base = RunGit(GitCommand);
+
+            string LiveStart = "[CL ";
+            string LiveEnd = " branch]";
+
+            if (Base.Contains(LiveStart) && Base.Contains(LiveEnd))
+            {
+                var CLStuff = Base.Substring(Base.IndexOf(LiveStart) + LiveStart.Length);
+                if (CLStuff.IndexOf(" ") <= 0)
+                {
+                    throw new AutomationException("Unrecognized commit5.");
+                }
+                CLStuff = CLStuff.Substring(0, CLStuff.IndexOf(" "));
+                Log("CL String {0}", CLStuff);
+                CL = int.Parse(CLStuff);
+
+                var BranchStuff = Base.Substring(Base.IndexOf(LiveStart) + LiveStart.Length, Base.IndexOf(LiveEnd) - Base.IndexOf(LiveStart) - LiveStart.Length);
+                if (BranchStuff.IndexOf(" ") <= 0)
+                {
+                    throw new AutomationException("Unrecognized commit6.");
+                }
+                BranchStuff = BranchStuff.Substring(BranchStuff.LastIndexOf(" ") + 1);
+                Log("Branch String {0}", BranchStuff);
+                if (BranchStuff.StartsWith("4."))
+                {
+                    Depot = "//depot/UE4-Releases/4." + BranchStuff.Substring(2, 1);
+                }
+                else if (BranchStuff.StartsWith("Main"))
+                {
+                    Depot = "//depot/UE4";
+                }
+                else
+                {
+                    throw new AutomationException("Unrecognized branch2.");
+                }
+            }
+            Log("CL int {0}", CL);
+            if (CL < 2000000 || String.IsNullOrWhiteSpace(Depot))
+            {
+                throw new AutomationException("Unrecognized commit3.");
+            }
+
+            return true;
+        }
+        catch (Exception)
+        {
+            CL = 0;
+            return false;
+        }
+    }
+
     void ExecuteInner(string Dir, int PR)
     {
         string PRNum = PR.ToString();
@@ -97,47 +206,39 @@ class GitPullRequest : BuildCommand
         RunGit(String.Format("fetch origin refs/pull/{0}/head:pr/{1}", PRNum, PRNum));
         RunGit(String.Format("checkout pr/{0} --", PRNum));
 
-		// after the fetch we do git log --Author... to figure out the P4 branch and CL
-		// the -1 limits it to the first one with the right author
-		var Base = RunGit(String.Format("log --author=TimSweeney --author=UnrealBot -1 pr/{0} --", PRNum));
-        string LookFor = "(";
-        if (!Base.Contains(LookFor))
-        {
-            throw new AutomationException("Unrecognized commit.");
-        }
-        Base = Base.Substring(Base.IndexOf(LookFor) + LookFor.Length);
-        string Depot = null;
-        if (Base.StartsWith("4."))
+        int CLBase;
+        string DepotBase;
+        ScanForBranchAndCL_BaseVersion(String.Format("log --author=TimSweeney --author=UnrealBot -100 pr/{0} --", PRNum), out DepotBase, out CLBase);
+
+
+        int CLLive;
+        string DepotLive;
+        ScanForBranchAndCL_LiveVersion(String.Format("log -100 pr/{0} --", PRNum), out DepotLive, out CLLive);
+
+        if (CLLive == 0 && CLBase == 0)
         {
-            Depot = "//depot/UE4-Releases/4." + Base.Substring(2, 1);
+            throw new AutomationException("Could not find a base change and branch using either method.");
         }
-        else if (Base.StartsWith("Main"))
+
+        int CL = 0;
+        string Depot = "";
+
+        if (CLBase > CLLive)
         {
-            Depot = "//depot/UE4";
+            CL = CLBase;
+            Depot = DepotBase;
         }
         else
         {
-            throw new AutomationException("Unrecognized branch.");
-        }
-        Log("Depot {0}", Depot);
-        if (!Base.Contains(")"))
-        {
-            throw new AutomationException("Unrecognized commit2.");
+            CL = CLLive;
+            Depot = DepotLive;
         }
-        Base = Base.Substring(0, Base.IndexOf(")"));
-        if (!Base.Contains(" "))
+        if (CL < 2000000 || String.IsNullOrWhiteSpace(Depot))
         {
-            throw new AutomationException("Unrecognized commit3.");
-        }
-        Base = Base.Substring(Base.LastIndexOf(" "));
-        Log("CL String {0}", Base);
-        int CL = int.Parse(Base);
-        Log("CL int {0}", CL);
-        if (CL < 2000000)
-        {
-            throw new AutomationException("Unrecognized commit3.");
+            throw new AutomationException("Could not find a base change and branch using either method.");
         }
 
+
         P4ClientInfo NewClient = P4.GetClientInfo(P4Env.Client);
 
         foreach (var p in NewClient.View)
