Enable OneSky portion of LauncherLocalization build.
Fix for a regression that was preventing Portable Object(.po) localization files from being exported/imported to the appropriate culture folders.  Enabled the OneSky integration in the LauncherLocalization build script and also made OneSky configuration settings more flexible to set/retrieve/hide.

diff --git a/Engine/Source/Editor/UnrealEd/Private/Commandlets/InternationalizationExportCommandlet.cpp b/Engine/Source/Editor/UnrealEd/Private/Commandlets/InternationalizationExportCommandlet.cpp
--- a/Engine/Source/Editor/UnrealEd/Private/Commandlets/InternationalizationExportCommandlet.cpp
+++ b/Engine/Source/Editor/UnrealEd/Private/Commandlets/InternationalizationExportCommandlet.cpp
@@ -1042,7 +1042,16 @@ bool UInternationalizationExportCommandlet::DoExport( const FString& SourcePath,
 				// Write out the Portable Object to .po file.
 				{
 					FString OutputString = PortableObj.ToString();
-					FString OutputFileName = DestinationPath / Filename;
+					FString OutputFileName;
+					
+					if (CulturesToGenerate.Num() > 1)
+					{
+						OutputFileName = DestinationPath / CultureName / Filename;
+					}
+					else
+					{
+						OutputFileName = DestinationPath / Filename;
+					}
 
 					//@TODO We force UTF8 at the moment but we want this to be based on the format found in the header info.
 					if( !FFileHelper::SaveStringToFile(OutputString, *OutputFileName, FFileHelper::EEncodingOptions::ForceUTF8) )
@@ -1080,7 +1089,16 @@ bool UInternationalizationExportCommandlet::DoImport(const FString& SourcePath,
 	{
 		// Load the Portable Object file if found
 		const FString CultureName = CulturesToGenerate[Culture];
-		FString POFilePath = SourcePath / Filename;
+		FString POFilePath;
+					
+		if (CulturesToGenerate.Num() > 1)
+		{
+			POFilePath = DestinationPath / CultureName / Filename;
+		}
+		else
+		{
+			POFilePath = DestinationPath / Filename;
+		}
 
 		if( !FPaths::FileExists(POFilePath) )
 		{
diff --git a/Engine/Source/Programs/AutomationTool/Scripts/AutomationScripts.Automation.csproj b/Engine/Source/Programs/AutomationTool/Scripts/AutomationScripts.Automation.csproj
--- a/Engine/Source/Programs/AutomationTool/Scripts/AutomationScripts.Automation.csproj
+++ b/Engine/Source/Programs/AutomationTool/Scripts/AutomationScripts.Automation.csproj
@@ -38,6 +38,10 @@
       <HintPath>..\..\..\..\Binaries\DotNET\Ionic.Zip.Reduced.dll</HintPath>
       <Private>False</Private>
     </Reference>
+    <Reference Include="OneSky, Version=1.0.0.0, Culture=neutral, processorArchitecture=MSIL">
+      <SpecificVersion>True</SpecificVersion>
+      <HintPath>..\..\..\ThirdParty\OneSky\bin\Release\OneSky.dll</HintPath>
+    </Reference>
     <Reference Include="System" />
     <Reference Include="System.Core" />
     <Reference Include="System.Runtime.Serialization" />
diff --git a/Engine/Source/Programs/AutomationTool/Scripts/LauncherLocalization.Automation.cs b/Engine/Source/Programs/AutomationTool/Scripts/LauncherLocalization.Automation.cs
--- a/Engine/Source/Programs/AutomationTool/Scripts/LauncherLocalization.Automation.cs
+++ b/Engine/Source/Programs/AutomationTool/Scripts/LauncherLocalization.Automation.cs
@@ -4,13 +4,17 @@ using System.Collections.Generic;
 using System.Globalization;
 using System.Linq;
 using System.IO;
+using System.Reflection;
 using AutomationTool;
-//using OneSky;
+using OneSky;
+using EpicGames.OneSkyLocalization.Config;
 
 class LauncherLocalization : BuildCommand
 {
+    
     public override void ExecuteBuild()
     {
+
         var EditorExe = CombinePaths(CmdEnv.LocalRoot, @"Engine/Binaries/Win64/UE4Editor-Cmd.exe");
 
         if (P4Enabled)
@@ -27,20 +31,21 @@ class LauncherLocalization : BuildCommand
             Log("Localize from label {0}", P4Env.LabelToSync);            
         }
 
-        //var oneSkyService = new OneSkyService("Reo5M8fnKcqVANHb4Zjc5nVUeXeehhGZ", "ByaWoHTjLM6clyxWLrrhhgdTY1btatdB");
+        OneSkyConfigData OneSkyConfig = OneSkyConfigHelper.Find("OneSkyConfig_EpicGames");
+        var oneSkyService = new OneSkyService(OneSkyConfig.ApiKey, OneSkyConfig.ApiSecret);
 
-        //// Export Launcher text from OneSky
-        //{
-        //    var launcherGroup = GetLauncherGroup(oneSkyService);
-        //    var appProject = GetAppProject(oneSkyService);
-        //    var appFile = appProject.UploadedFiles.FirstOrDefault(f => f.Filename == "App.po");
+        // Export Launcher text from OneSky
+        {
+            var launcherGroup = GetLauncherGroup(oneSkyService);
+            var appProject = GetAppProject(oneSkyService);
+            var appFile = appProject.UploadedFiles.FirstOrDefault(f => f.Filename == "App.po");
 
-        //    //Export
-        //    if (appFile != null)
-        //    {
-        //        ExportFileToDirectory(appFile, new DirectoryInfo(CmdEnv.LocalRoot + "/Engine/Programs/NoRedist/UnrealEngineLauncher/Content/Localization/App"), launcherGroup.EnabledCultures);
-        //    }
-        //}
+            //Export
+            if (appFile != null)
+            {
+                ExportFileToDirectory(appFile, new DirectoryInfo(CmdEnv.LocalRoot + "/Engine/Programs/NoRedist/UnrealEngineLauncher/Content/Localization/App"), launcherGroup.EnabledCultures);
+            }
+        }
 
         // Setup editor arguments for SCC.
         string EditorArguments = String.Empty;
@@ -82,108 +87,184 @@ class LauncherLocalization : BuildCommand
         }
 
         // Upload Launcher text to OneSky
-        //UploadDirectoryToProject(GetAppProject(oneSkyService), new DirectoryInfo(CmdEnv.LocalRoot + "/Engine/Programs/NoRedist/UnrealEngineLauncher/Content/Localization/App"), "*.po");
+        UploadDirectoryToProject(GetAppProject(oneSkyService), new DirectoryInfo(CmdEnv.LocalRoot + "/Engine/Programs/NoRedist/UnrealEngineLauncher/Content/Localization/App"), "*.po");
+    }
+
+    private static ProjectGroup GetLauncherGroup(OneSkyService oneSkyService)
+    {
+        var launcherGroup = oneSkyService.ProjectGroups.FirstOrDefault(g => g.Name == "Launcher");
+
+        if (launcherGroup == null)
+        {
+            launcherGroup = new ProjectGroup("Launcher", new CultureInfo("en"));
+            oneSkyService.ProjectGroups.Add(launcherGroup);
+        }
+
+        return launcherGroup;
+    }
+
+    private static OneSky.Project GetAppProject(OneSkyService oneSkyService)
+    {
+        var launcherGroup = GetLauncherGroup(oneSkyService);
+
+        OneSky.Project appProject = launcherGroup.Projects.FirstOrDefault(p => p.Name == "App");
+
+        if (appProject == null)
+        {
+            ProjectType projectType = oneSkyService.ProjectTypes.First(pt => pt.Code == "website");
+
+            appProject = new OneSky.Project("App", "The core application text that ships with the Launcher", projectType);
+            launcherGroup.Projects.Add(appProject);
+        }
+
+        return appProject;
+    }
+
+    private static void ExportFileToDirectory(UploadedFile file, DirectoryInfo destination, IEnumerable<CultureInfo> cultures)
+    {
+        foreach (var culture in cultures)
+        {
+            var cultureDirectory = new DirectoryInfo(Path.Combine(destination.FullName, culture.Name));
+            if (!cultureDirectory.Exists)
+            {
+                cultureDirectory.Create();
+            }
+
+            using (var memoryStream = new MemoryStream())
+            {
+                var exportFile = new FileInfo(Path.Combine(cultureDirectory.FullName, file.Filename));
+
+                var exportTranslationState = file.ExportTranslation(culture, memoryStream).Result;
+                if (exportTranslationState == UploadedFile.ExportTranslationState.Success)
+                {
+                    memoryStream.Position = 0;
+                    using (Stream fileStream = File.OpenWrite(exportFile.FullName))
+                    {
+                        memoryStream.CopyTo(fileStream);
+                        Console.WriteLine("[SUCCESS] Exporting: " + exportFile.FullName + " Locale: " + culture.Name);
+                    }
+                }
+                else if (exportTranslationState == UploadedFile.ExportTranslationState.NoContent)
+                {
+                    Console.WriteLine("[WARNING] Exporting: " + exportFile.FullName + " Locale: " + culture.Name + " has no translations!");
+                }
+                else
+                {
+                    Console.WriteLine("[FAILED] Exporting: " + exportFile.FullName + " Locale: " + culture.Name);
+                }
+            }
+        }
     }
 
-    //private static ProjectGroup GetLauncherGroup(OneSkyService oneSkyService)
-    //{
-    //    var launcherGroup = oneSkyService.ProjectGroups.FirstOrDefault(g => g.Name == "Launcher");
-
-    //    if (launcherGroup == null)
-    //    {
-    //        launcherGroup = new ProjectGroup("Launcher", new CultureInfo("en"));
-    //        oneSkyService.ProjectGroups.Add(launcherGroup);
-    //    }
-
-    //    return launcherGroup;
-    //}
-
-    //private static OneSky.Project GetAppProject(OneSkyService oneSkyService)
-    //{
-    //    var launcherGroup = GetLauncherGroup(oneSkyService);
-
-    //    OneSky.Project appProject = launcherGroup.Projects.FirstOrDefault(p => p.Name == "App");
-
-    //    if (appProject == null)
-    //    {
-    //        ProjectType projectType = oneSkyService.ProjectTypes.First(pt => pt.Code == "website");
-
-    //        appProject = new OneSky.Project("App", "The core application text that ships with the Launcher", projectType);
-    //        launcherGroup.Projects.Add(appProject);
-    //    }
-
-    //    return appProject;
-    //}
-
-    //private static void ExportFileToDirectory(UploadedFile file, DirectoryInfo destination, IEnumerable<CultureInfo> cultures)
-    //{
-    //    foreach (var culture in cultures)
-    //    {
-    //        var cultureDirectory = new DirectoryInfo(Path.Combine(destination.FullName, culture.Name));
-    //        if (!cultureDirectory.Exists)
-    //        {
-    //            cultureDirectory.Create();
-    //        }
-
-    //        using (var memoryStream = new MemoryStream())
-    //        {
-    //            var exportFile = new FileInfo(Path.Combine(cultureDirectory.FullName, file.Filename));
-
-    //            var exportTranslationState = file.ExportTranslation(culture, memoryStream).Result;
-    //            if (exportTranslationState == UploadedFile.ExportTranslationState.Success)
-    //            {
-    //                memoryStream.Position = 0;
-    //                using (Stream fileStream = File.OpenWrite(exportFile.FullName))
-    //                {
-    //                    memoryStream.CopyTo(fileStream);
-    //                    Console.WriteLine("[SUCCESS] Exporting: " + exportFile.FullName + " Locale: " + culture.Name);
-    //                }
-    //            }
-    //            else if (exportTranslationState == UploadedFile.ExportTranslationState.NoContent)
-    //            {
-    //                Console.WriteLine("[WARNING] Exporting: " + exportFile.FullName + " Locale: " + culture.Name + " has no translations!");
-    //            }
-    //            else
-    //            {
-    //                Console.WriteLine("[FAILED] Exporting: " + exportFile.FullName + " Locale: " + culture.Name);
-    //            }
-    //        }
-    //    }
-    //}
-
-    //static void UploadDirectoryToProject(OneSky.Project project, DirectoryInfo directory, string fileExtension)
-    //{
-    //    foreach (var file in Directory.GetFiles(directory.FullName, fileExtension, SearchOption.AllDirectories))
-    //    {
-    //        DirectoryInfo parentDirectory = Directory.GetParent(file);
-    //        string localeName = parentDirectory.Name;
-    //        string currentFile = file;
-
-    //        using (var fileStream = File.OpenRead(currentFile))
-    //        {
-    //            // Read the BOM
-    //            var bom = new byte[3];
-    //            fileStream.Read(bom, 0, 3);
-
-    //            //We want to ignore the utf8 BOM
-    //            if (bom[0] != 0xef || bom[1] != 0xbb || bom[2] != 0xbf)
-    //            {
-    //                fileStream.Position = 0;
-    //            }
-
-    //            Console.WriteLine("Uploading: " + currentFile + " Locale: " + localeName);
-    //            var uploadedFile = project.Upload(Path.GetFileName(currentFile), fileStream, new CultureInfo(localeName)).Result;
-
-    //            if (uploadedFile == null)
-    //            {
-    //                Console.WriteLine("[FAILED] Uploading: " + currentFile + " Locale: " + localeName);
-    //            }
-    //            else
-    //            {
-    //                Console.WriteLine("[SUCCESS] Uploading: " + currentFile + " Locale: " + localeName);
-    //            }
-    //        }
-    //    }
-    //}
+    static void UploadDirectoryToProject(OneSky.Project project, DirectoryInfo directory, string fileExtension)
+    {
+        foreach (var file in Directory.GetFiles(directory.FullName, fileExtension, SearchOption.AllDirectories))
+        {
+            DirectoryInfo parentDirectory = Directory.GetParent(file);
+            string localeName = parentDirectory.Name;
+            string currentFile = file;
+
+            using (var fileStream = File.OpenRead(currentFile))
+            {
+                // Read the BOM
+                var bom = new byte[3];
+                fileStream.Read(bom, 0, 3);
+
+                //We want to ignore the utf8 BOM
+                if (bom[0] != 0xef || bom[1] != 0xbb || bom[2] != 0xbf)
+                {
+                    fileStream.Position = 0;
+                }
+
+                Console.WriteLine("Uploading: " + currentFile + " Locale: " + localeName);
+                var uploadedFile = project.Upload(Path.GetFileName(currentFile), fileStream, new CultureInfo(localeName)).Result;
+
+                if (uploadedFile == null)
+                {
+                    Console.WriteLine("[FAILED] Uploading: " + currentFile + " Locale: " + localeName);
+                }
+                else
+                {
+                    Console.WriteLine("[SUCCESS] Uploading: " + currentFile + " Locale: " + localeName);
+                }
+            }
+        }
+    }
 }
 
+
+namespace EpicGames.OneSkyLocalization.Config
+{
+    /// <summary>
+    /// Class for retrieving OneSky localization configuration data
+    /// </summary>
+    public class OneSkyConfigHelper
+    {
+        // List of configs is cached off for fetching from multiple times
+        private static Dictionary<string, OneSkyConfigData> Configs;
+
+        public static OneSkyConfigData Find(string ConfigName)
+        {
+            if (Configs == null)
+            {
+                // Load all secret configs by trying to instantiate all classes derived from OneSkyConfigData from all loaded DLLs.
+                // Note that we're using the default constructor on the secret config types.
+                Configs = new Dictionary<string, OneSkyConfigData>();
+                Assembly[] LoadedAssemblies = AppDomain.CurrentDomain.GetAssemblies();
+                foreach (var Dll in LoadedAssemblies)
+                {
+                    Type[] AllTypes = Dll.GetTypes();
+                    foreach (var PotentialConfigType in AllTypes)
+                    {
+                        if (PotentialConfigType != typeof(OneSkyConfigData) && typeof(OneSkyConfigData).IsAssignableFrom(PotentialConfigType))
+                        {
+                            try
+                            {
+                                OneSkyConfigData Config = Activator.CreateInstance(PotentialConfigType) as OneSkyConfigData;
+                                if (Config != null)
+                                {
+                                    Configs.Add(Config.Name, Config);
+                                }
+                            }
+                            catch
+                            {
+                                BuildCommand.LogWarning("Unable to create OneSky config data: {0}", PotentialConfigType.Name);
+                            }
+                        }
+                    }
+                }
+            }
+            OneSkyConfigData LoadedConfig;
+            Configs.TryGetValue(ConfigName, out LoadedConfig);
+            if (LoadedConfig == null)
+            {
+                throw new AutomationException("Unable to find requested OneSky config data: {0}", ConfigName);
+            }
+            return LoadedConfig;
+        }
+    }
+
+    /// <summary>
+    /// Class for storing OneSky Localization configuration data
+    /// </summary>
+    public class OneSkyConfigData
+    {
+        public OneSkyConfigData(string InName, string InApiKey, string InApiSecret)
+        {
+            Name = InName;
+            ApiKey = InApiKey;
+            ApiSecret = InApiSecret;
+        }
+
+        public string Name;
+        public string ApiKey;
+        public string ApiSecret;
+
+        public void SpewValues()
+        {
+            CommandUtils.Log("Name : {0}", Name);
+            CommandUtils.Log("ApiKey : {0}", ApiKey);
+            //CommandUtils.Log("ApiSecret : {0}", ApiSecret);  // This should probably never be revealed.
+        }
+    }
+}
