build script, allow UBT files to be not retrieved

diff --git a/Engine/Source/Programs/AutomationTool/TempStorage.cs b/Engine/Source/Programs/AutomationTool/TempStorage.cs
--- a/Engine/Source/Programs/AutomationTool/TempStorage.cs
+++ b/Engine/Source/Programs/AutomationTool/TempStorage.cs
@@ -830,6 +830,7 @@ namespace AutomationTool
                     {
                         bool bOkToBeDifferent = DestFile.Contains("AutomationTool.exe");
                         bOkToBeDifferent = bOkToBeDifferent || DestFile.Contains("\\DotNET\\AutomationScripts\\") || DestFile.Contains("/DotNET/AutomationScripts/");
+						bOkToBeDifferent = bOkToBeDifferent || DestFile.Contains("UnrealBuildTool.exe");
                         if (bOkToBeDifferent)
                         {
                             bDoCopy = false;
