build script, fix build cleanup

diff --git a/Engine/Source/Programs/AutomationTool/CommandUtils.cs b/Engine/Source/Programs/AutomationTool/CommandUtils.cs
--- a/Engine/Source/Programs/AutomationTool/CommandUtils.cs
+++ b/Engine/Source/Programs/AutomationTool/CommandUtils.cs
@@ -1893,9 +1893,10 @@ namespace AutomationTool
                 {
                     if (DirectoryExists_NoExceptions(TopLevelDir.FullName))
                     {
-                        if (TopLevelDir.FullName.StartsWith(StartString, StringComparison.InvariantCultureIgnoreCase) && TopLevelDir.FullName.EndsWith(EndString, StringComparison.InvariantCultureIgnoreCase))
+                        var JustDir = Path.GetFileName(CombinePaths(TopLevelDir.FullName));
+                        if (JustDir.StartsWith(StartString, StringComparison.InvariantCultureIgnoreCase) && (String.IsNullOrEmpty(EndString) || JustDir.EndsWith(EndString, StringComparison.InvariantCultureIgnoreCase)))
                         {
-                            string CLPart = TopLevelDir.FullName.Substring(StartString.Length, TopLevelDir.FullName.Length - StartString.Length - EndString.Length);
+                            string CLPart = JustDir.Substring(StartString.Length, TopLevelDir.FullName.Length - StartString.Length - EndString.Length);
                             if (!CLPart.Contains("-") && !CLPart.Contains("+"))
                             {
                                 DirectoryInfo ThisDirInfo = new DirectoryInfo(TopLevelDir.FullName);
@@ -1913,12 +1914,12 @@ namespace AutomationTool
                             }
                             else
                             {
-                                Log("skipping {0}, because the CL part {1} had weird characters", TopLevelDir.FullName, CLPart);
+                                Log("skipping {0}, because the CL part {1} had weird characters", JustDir, CLPart);
                             }
                         }
                         else
                         {
-                            Log("skipping {0}, because it didn't start with {1} or end with {2}", TopLevelDir.FullName, StartString, EndString);
+                            Log("skipping {0}, because it didn't start with {1} or end with {2}", JustDir, StartString, EndString);
                         }
                     }
                 }
