Integrate API docs from 4.2 release back to main branch.

diff --git a/Engine/Source/Programs/AutomationTool/Scripts/BuildDocumentation.Automation.cs b/Engine/Source/Programs/AutomationTool/Scripts/BuildDocumentation.Automation.cs
--- a/Engine/Source/Programs/AutomationTool/Scripts/BuildDocumentation.Automation.cs
+++ b/Engine/Source/Programs/AutomationTool/Scripts/BuildDocumentation.Automation.cs
@@ -83,7 +83,7 @@ public class BuildDocumentation : BuildCommand
 		// Execute the clean
 		if(bCleanEnv)
 		{
-			RunAndLog(CmdEnv, ApiDocToolPath, "-cleantargetinfo" + ApiToolCommandLine, "ApiDocTool-CleanTargetInfo");
+			RunAndLog(CmdEnv, ApiDocToolPath, "-cleantarget" + ApiToolCommandLine, "ApiDocTool-CleanTargetInfo");
 		}
 		if (bCleanMeta)
 		{
@@ -109,7 +109,7 @@ public class BuildDocumentation : BuildCommand
 		// Execute the build
 		if (bBuildEnv)
 		{
-			RunAndLog(CmdEnv, ApiDocToolPath, "-buildtargetinfo" + ApiToolCommandLine, "ApiDocTool-BuildTargetInfo");
+			RunAndLog(CmdEnv, ApiDocToolPath, "-buildtarget" + ApiToolCommandLine, "ApiDocTool-BuildTargetInfo");
 		}
 		if (bBuildMeta)
 		{
diff --git a/Engine/Source/Programs/UnrealDocTool/APIDocTool/APIDocTool/APIIndex.cs b/Engine/Source/Programs/UnrealDocTool/APIDocTool/APIDocTool/APIIndex.cs
--- a/Engine/Source/Programs/UnrealDocTool/APIDocTool/APIDocTool/APIIndex.cs
+++ b/Engine/Source/Programs/UnrealDocTool/APIDocTool/APIDocTool/APIIndex.cs
@@ -86,7 +86,8 @@ namespace APIDocTool
 				APIModule Module = APIModule.Build(Parent, Pair.Key);
 				Parent.Children.Add(Module);
 
-				Pair.Value.AddModule(Module);
+				Pair.Value.Modules.Add(Module);
+//				Pair.Value.AddModule(Module);
 			}
 
 			// Expand the Core section by default
diff --git a/Engine/Source/Programs/UnrealDocTool/APIDocTool/APIDocTool/APIModule.cs b/Engine/Source/Programs/UnrealDocTool/APIDocTool/APIDocTool/APIModule.cs
--- a/Engine/Source/Programs/UnrealDocTool/APIDocTool/APIDocTool/APIModule.cs
+++ b/Engine/Source/Programs/UnrealDocTool/APIDocTool/APIDocTool/APIModule.cs
@@ -65,7 +65,7 @@ namespace APIDocTool
 				// If we didn't have one set explicitly, generate one from the subdirectory
 				if(FilterPath == null)
 				{
-					string EntityFile = Path.GetFullPath(Entity.File);
+					string EntityFile = String.IsNullOrEmpty(Entity.File)? "" : Path.GetFullPath(Entity.File);
 					if(EntityFile.ToLowerInvariant().StartsWith(NormalizedBaseSrcDir))
 					{
 						int MinIndex = EntityFile.IndexOf('\\', NormalizedBaseSrcDir.Length);
