build script, fix and enable build cleaning

diff --git a/Engine/Source/Programs/AutomationTool/CommandUtils.cs b/Engine/Source/Programs/AutomationTool/CommandUtils.cs
--- a/Engine/Source/Programs/AutomationTool/CommandUtils.cs
+++ b/Engine/Source/Programs/AutomationTool/CommandUtils.cs
@@ -1861,15 +1861,18 @@ namespace AutomationTool
             return Result;
         }
 
-        public static void CleanFormalBuilds(string DirectoryForThisBuild)
+        public static void CleanFormalBuilds(string DirectoryForThisBuild, string CLString = "")
         {
-            if (!IsBuildMachine || !DirectoryForThisBuild.StartsWith(RootSharedTempStorageDirectory()) || !P4Enabled)
+            if (CLString == "" && (!IsBuildMachine || !DirectoryForThisBuild.StartsWith(RootSharedTempStorageDirectory()) || !P4Enabled))
             {
                 return;
             }
             try
             {
-                string CLString = P4Env.ChangelistString;
+                if (P4Enabled && CLString == "")
+                {
+                    CLString = P4Env.ChangelistString;
+                }
                 const int MaximumDaysToKeepTempStorage = 4;
                 string ParentDir = Path.GetDirectoryName(CombinePaths(DirectoryForThisBuild));
                 if (!DirectoryExists_NoExceptions(ParentDir))
@@ -1896,7 +1899,7 @@ namespace AutomationTool
                         var JustDir = Path.GetFileName(CombinePaths(TopLevelDir.FullName));
                         if (JustDir.StartsWith(StartString, StringComparison.InvariantCultureIgnoreCase) && (String.IsNullOrEmpty(EndString) || JustDir.EndsWith(EndString, StringComparison.InvariantCultureIgnoreCase)))
                         {
-                            string CLPart = JustDir.Substring(StartString.Length, TopLevelDir.FullName.Length - StartString.Length - EndString.Length);
+                            string CLPart = JustDir.Substring(StartString.Length, JustDir.Length - StartString.Length - EndString.Length);
                             if (!CLPart.Contains("-") && !CLPart.Contains("+"))
                             {
                                 DirectoryInfo ThisDirInfo = new DirectoryInfo(TopLevelDir.FullName);
@@ -1908,8 +1911,12 @@ namespace AutomationTool
                                 }
                                 if (bOld)
                                 {
-                                    Log("Deleteing (or will be after this looks correct) temp storage directory {0}, because it is more than {1} days old.", TopLevelDir.FullName, MaximumDaysToKeepTempStorage);
-                                    //DeleteDirectory_NoExceptions(true, TopLevelDir.FullName);
+                                    Log("Deleting temp storage directory {0}, because it is more than {1} days old.", TopLevelDir.FullName, MaximumDaysToKeepTempStorage);
+                                    DeleteDirectory_NoExceptions(true, TopLevelDir.FullName);
+                                }
+                                else
+                                {
+                                    Log("Not Deleteing temp storage directory {0}, because it is less than {1} days old.", TopLevelDir.FullName, MaximumDaysToKeepTempStorage);
                                 }
                             }
                             else
diff --git a/Engine/Source/Programs/AutomationTool/Scripts/Tests.Automation.cs b/Engine/Source/Programs/AutomationTool/Scripts/Tests.Automation.cs
--- a/Engine/Source/Programs/AutomationTool/Scripts/Tests.Automation.cs
+++ b/Engine/Source/Programs/AutomationTool/Scripts/Tests.Automation.cs
@@ -1893,6 +1893,19 @@ class TestKillAll : BuildCommand
 	}
 }
 
+[Help("Tests CleanFormalBuilds.")]
+class TestCleanFormalBuilds : BuildCommand
+{
+    public override void ExecuteBuild()
+    {
+        Log("*********************** TestCleanFormalBuilds");
+        var Dir = ParseParamValue("Dir", @"P:\Builds\Soul\Soul_Android_Shipping_MakeBuild\++depot+UE4-CL-2077154");
+        var CLString = ParseParamValue("CL", "2077154");
+        CleanFormalBuilds(Dir, CLString);
+    }
+}
+
+
 [Help("Spawns a process to test if it can be killed.")]
 class TestStopProcess : BuildCommand
 {
