build script, better dependencies for test compiles

diff --git a/Engine/Source/Programs/AutomationTool/Scripts/GUBP.Automation.cs b/Engine/Source/Programs/AutomationTool/Scripts/GUBP.Automation.cs
--- a/Engine/Source/Programs/AutomationTool/Scripts/GUBP.Automation.cs
+++ b/Engine/Source/Programs/AutomationTool/Scripts/GUBP.Automation.cs
@@ -1118,8 +1118,8 @@ public class GUBP : BuildCommand
                 AddPseudodependency(EditorGameNode.StaticGetFullName(InHostPlatform, GameProj));
                 if (bp.HasNode(GamePlatformMonolithicsNode.StaticGetFullName(HostPlatform, bp.Branch.BaseEngineProject, TargetPlatform)))
                 {
-                AddPseudodependency(GamePlatformMonolithicsNode.StaticGetFullName(HostPlatform, bp.Branch.BaseEngineProject, TargetPlatform));
-            }
+                    AddPseudodependency(GamePlatformMonolithicsNode.StaticGetFullName(HostPlatform, bp.Branch.BaseEngineProject, TargetPlatform));
+                }
             }
             else
             {
@@ -2494,6 +2494,7 @@ public class GUBP : BuildCommand
         public NonUnityTestNode(UnrealTargetPlatform InHostPlatform)
             : base(InHostPlatform)
         {
+            AddPseudodependency(RootEditorNode.StaticGetFullName(HostPlatform));
         }
         public static string StaticGetFullName(UnrealTargetPlatform InHostPlatform)
         {
@@ -2535,9 +2536,12 @@ public class GUBP : BuildCommand
 
     public class IOSOnPCTestNode : TestNode
     {
-        public IOSOnPCTestNode()
+        public IOSOnPCTestNode(GUBP bp)
             : base(UnrealTargetPlatform.Win64)
         {
+            AddDependency(ToolsForCompileNode.StaticGetFullName(UnrealTargetPlatform.Win64));
+            AddDependency(ToolsNode.StaticGetFullName(UnrealTargetPlatform.Win64));
+            AddPseudodependency(GamePlatformMonolithicsNode.StaticGetFullName(UnrealTargetPlatform.Mac, bp.Branch.BaseEngineProject, UnrealTargetPlatform.IOS));
         }
         public static string StaticGetFullName()
         {
@@ -2556,7 +2560,7 @@ public class GUBP : BuildCommand
             var Build = new UE4Build(bp);
             var Agenda = new UE4Build.BuildAgenda();
 
-            Agenda.AddTargets(new string[] { "UE4Game" }, UnrealTargetPlatform.IOS, UnrealTargetConfiguration.Development);
+            Agenda.AddTargets(new string[] { bp.Branch.BaseEngineProject.Properties.Targets[TargetRules.TargetType.Game].TargetName }, UnrealTargetPlatform.IOS, UnrealTargetConfiguration.Development);
 
             Build.Build(Agenda, InDeleteBuildProducts: true, InUpdateVersionFiles: false);
 
@@ -4327,10 +4331,6 @@ public class GUBP : BuildCommand
             if (DoASharedPromotable)
             {
                 AddNode(new NonUnityTestNode(HostPlatform));
-                if (HostPlatform == UnrealTargetPlatform.Win64)
-                {
-                    AddNode(new IOSOnPCTestNode());
-                }
 
                 var AgentSharingGroup = "Shared_EditorTests" + HostPlatformNode.StaticGetHostPlatformSuffix(HostPlatform);
 
@@ -4790,6 +4790,10 @@ if (HostPlatform == UnrealTargetPlatform.Mac) continue; //temp hack till mac aut
 
         if (HasNode(ToolsForCompileNode.StaticGetFullName(UnrealTargetPlatform.Win64)))
         {
+            if (HasNode(GamePlatformMonolithicsNode.StaticGetFullName(UnrealTargetPlatform.Mac, Branch.BaseEngineProject, UnrealTargetPlatform.IOS)) && HasNode(ToolsNode.StaticGetFullName(UnrealTargetPlatform.Win64)))
+            {
+                AddNode(new IOSOnPCTestNode(this));
+            }
             AddNode(new CleanSharedTempStorageNode(this));
         }
 
