build script, test shared temp storage for write

diff --git a/Engine/Source/Programs/AutomationTool/CommandUtils.cs b/Engine/Source/Programs/AutomationTool/CommandUtils.cs
--- a/Engine/Source/Programs/AutomationTool/CommandUtils.cs
+++ b/Engine/Source/Programs/AutomationTool/CommandUtils.cs
@@ -1829,6 +1829,34 @@ namespace AutomationTool
             return StorageDirectory;
         }
 
+        static bool DirectoryExistsAndIsWritable_NoExceptions(string Dir)
+        {
+            try
+            {
+                if (!DirectoryExists_NoExceptions(Dir))
+                {
+                    return false;
+                }
+                var TestGUID = Guid.NewGuid();
+                var Filename = CombinePaths(Dir, TestGUID.ToString() + ".Temp.txt");
+                WriteAllText_NoExceptions(Filename, "Test");
+                if (FileExists_NoExceptions(true, Filename))
+                {
+                    DeleteFile_NoExceptions(Filename, true);
+                    //Log(System.Diagnostics.TraceEventType.Information, "Resolved shared dir {0}", Dir);
+                    return true;
+                }
+                Log(System.Diagnostics.TraceEventType.Warning, "Shared dir {0} is not writable", Dir);
+
+            }
+            catch (Exception Ex)
+            {
+                Log(System.Diagnostics.TraceEventType.Warning, "Failed to resolve shared dir {0}", Dir);
+                Log(System.Diagnostics.TraceEventType.Warning, LogUtils.FormatException(Ex));
+            }
+            return false;
+        }
+
         static Dictionary<string, string> ResolveCache = new Dictionary<string, string>();
         public static string ResolveSharedBuildDirectory(string GameFolder)
         {
@@ -1838,7 +1866,7 @@ namespace AutomationTool
             }
             string Root = RootSharedTempStorageDirectory();
             string Result = CombinePaths(Root, GameFolder);
-            if (String.IsNullOrEmpty(GameFolder) || !DirectoryExists_NoExceptions(Result))
+            if (String.IsNullOrEmpty(GameFolder) || !DirectoryExistsAndIsWritable_NoExceptions(Result))
             {
                 string GameStr = "Game";
                 bool HadGame = false;
@@ -1848,10 +1876,10 @@ namespace AutomationTool
                     Result = CombinePaths(Root, ShortFolder);
                     HadGame = true;
                 }
-                if (!HadGame || !DirectoryExists_NoExceptions(Result))
+                if (!HadGame || !DirectoryExistsAndIsWritable_NoExceptions(Result))
                 {
                     Result = CombinePaths(Root, "UE4");
-                    if (!DirectoryExists_NoExceptions(Result))
+                    if (!DirectoryExistsAndIsWritable_NoExceptions(Result))
                     {
                         throw new AutomationException("Could not find an appropriate shared temp folder {0}", Result);
                     }
diff --git a/Engine/Source/Programs/AutomationTool/Scripts/Tests.Automation.cs b/Engine/Source/Programs/AutomationTool/Scripts/Tests.Automation.cs
--- a/Engine/Source/Programs/AutomationTool/Scripts/Tests.Automation.cs
+++ b/Engine/Source/Programs/AutomationTool/Scripts/Tests.Automation.cs
@@ -459,6 +459,9 @@ class TestTempStorage : BuildCommand
 	public override void ExecuteBuild()
 	{
 		Log("TestTempStorage********");
+
+        Log("Resolved Ocean to {0}", ResolveSharedBuildDirectory("Ocean"));
+
 		DeleteLocalTempStorageManifests(CmdEnv);
 		DeleteSharedTempStorageManifests(CmdEnv, "Test");
 		if (TempStorageExists(CmdEnv, "Test"))
