Build fix for earlier GitHub merge.

diff --git a/Engine/Source/Programs/AutomationTool/Android/AndroidPlatform.Automation.cs b/Engine/Source/Programs/AutomationTool/Android/AndroidPlatform.Automation.cs
--- a/Engine/Source/Programs/AutomationTool/Android/AndroidPlatform.Automation.cs
+++ b/Engine/Source/Programs/AutomationTool/Android/AndroidPlatform.Automation.cs
@@ -102,7 +102,7 @@ public class AndroidPlatform : Platform
 
 		// packaging just takes a pak file and makes it the .obb
         UEBuildConfiguration.bOBBinAPK = Params.OBBinAPK; // Make sure this setting is sync'd pre-build
-        UEBuildDeploy Deploy = UEBuildDeploy.GetBuildDeploy(UnrealTargetPlatform.Android);
+        var Deploy = UEBuildDeploy.GetBuildDeploy(UnrealTargetPlatform.Android);
 		Deploy.PrepForUATPackageOrDeploy(Params.ShortProjectName, SC.ProjectRoot, SOName, SC.LocalRoot + "/Engine", Params.Distribution);
         
 		// first, look for a .pak file in the staged directory
@@ -201,7 +201,7 @@ public class AndroidPlatform : Platform
 		string ApkName = GetFinalApkName(Params, SC.StageExecutables[0], true);
 
 		// make sure APK is up to date (this is fast if so)
-		UEBuildDeploy Deploy = UEBuildDeploy.GetBuildDeploy(UnrealTargetPlatform.Android);
+		var Deploy = UEBuildDeploy.GetBuildDeploy(UnrealTargetPlatform.Android);
 		Deploy.PrepForUATPackageOrDeploy(Params.ShortProjectName, SC.ProjectRoot, SOName, SC.LocalRoot + "/Engine", Params.Distribution);
 
 		// now we can use the apk to get more info
diff --git a/Engine/Source/Programs/AutomationTool/IOS/IOSPlatform.Automation.cs b/Engine/Source/Programs/AutomationTool/IOS/IOSPlatform.Automation.cs
--- a/Engine/Source/Programs/AutomationTool/IOS/IOSPlatform.Automation.cs
+++ b/Engine/Source/Programs/AutomationTool/IOS/IOSPlatform.Automation.cs
@@ -76,7 +76,7 @@ public class IOSPlatform : Platform
 		if (UnrealBuildTool.ExternalExecution.GetRuntimePlatform() == UnrealTargetPlatform.Mac)
 		{
 			// copy in all of the artwork and plist
-			UEBuildDeploy DeployHandler = UEBuildDeploy.GetBuildDeploy(UnrealTargetPlatform.IOS);
+			var DeployHandler = UEBuildDeploy.GetBuildDeploy(UnrealTargetPlatform.IOS);
 			DeployHandler.PrepForUATPackageOrDeploy(Params.ShortProjectName,
 				Path.GetDirectoryName(Params.RawProjectPath),
 				CombinePaths(Path.GetDirectoryName(Params.ProjectGameExeFilename), SC.StageExecutables[0]),
