build script, add support for UT to gitpullrequest

diff --git a/Engine/Source/Programs/AutomationTool/Scripts/Tests.Automation.cs b/Engine/Source/Programs/AutomationTool/Scripts/Tests.Automation.cs
--- a/Engine/Source/Programs/AutomationTool/Scripts/Tests.Automation.cs
+++ b/Engine/Source/Programs/AutomationTool/Scripts/Tests.Automation.cs
@@ -35,10 +35,15 @@ class TestP4_Info : BuildCommand
 [Help("Dir=", "Directory of the Git repo.")]
 [Help("PR=", "PR # to shelve, can use a range PR=5-25")]
 [RequireP4]
+[DoesNotNeedP4CL]
 class GitPullRequest : BuildCommand
 {
 	/// URL to our UnrealEngine repository on GitHub
-	static readonly string GitRepositoryURL = "https://github.com/EpicGames/UnrealEngine.git";
+	static readonly string GitRepositoryURL_Engine = "https://github.com/EpicGames/UnrealEngine.git";
+    static readonly string GitRepositoryURL_UT = "https://github.com/EpicGames/UnrealTournament.git";
+    string GitRepositoryURL = null;
+
+    bool bDoingUT = false;
 
     string FindExeFromPath(string ExeName, string ExpectedPathSubstring = null)
     {
@@ -103,6 +108,10 @@ class GitPullRequest : BuildCommand
                 {
                     Depot = "//depot/UE4";
                 }
+                else if (Base.StartsWith("UT"))
+                {
+                    Depot = "//depot/UE4-UT";
+                }
                 else
                 {
                     throw new AutomationException("Unrecognized branch.");
@@ -170,6 +179,10 @@ class GitPullRequest : BuildCommand
                 {
                     Depot = "//depot/UE4";
                 }
+                else if (BranchStuff.StartsWith("UT"))
+                {
+                    Depot = "//depot/UE4-UT";
+                }
                 else
                 {
                     throw new AutomationException("Unrecognized branch2.");
@@ -263,18 +276,28 @@ class GitPullRequest : BuildCommand
         P4Sub.Sync(String.Format("-f -k -q {0}/...@{1}", Depot, CL));
 
         var Change = P4Sub.CreateChange(null, String.Format("GitHub pull request #{0}", PRNum));
-        P4Sub.ReconcileNoDeletes(Change, CombinePaths(Dir, "Engine", "..."));
+        P4Sub.ReconcileNoDeletes(Change, CombinePaths(Dir, bDoingUT ? "UnrealTournament" : "Engine", "..."));
         P4Sub.Shelve(Change);
         P4Sub.Revert(Change, "-k //...");
     }
 
     public override void ExecuteBuild()
     {
+        if (ParseParam("UT"))
+        {
+            bDoingUT = true;
+            GitRepositoryURL = GitRepositoryURL_UT;
+        }
+        else
+        {
+            bDoingUT = true;
+            GitRepositoryURL = GitRepositoryURL_Engine;
+        }
         var Dir = ParseParamValue("Dir");
         if (String.IsNullOrEmpty(Dir))
         {
 			// No Git repo directory was specified, so we'll choose a directory automatically
-			Dir = Path.GetFullPath( Path.Combine( CmdEnv.LocalRoot, "Engine", "Intermediate", "PullRequestGitRepo" ) );
+            Dir = Path.GetFullPath(Path.Combine(CmdEnv.LocalRoot, "Engine", "Intermediate", bDoingUT ? "PullRequestGitRepo_UT" : "PullRequestGitRepo"));
         }
 
         var PRNum = ParseParamValue("PR");
