Do not package .png files for Linux.

New template projects don't seem to have them, making packaging fail (TTP #333350).

diff --git a/Engine/Source/Programs/AutomationTool/Linux/LinuxPlatform.Automation.cs b/Engine/Source/Programs/AutomationTool/Linux/LinuxPlatform.Automation.cs
--- a/Engine/Source/Programs/AutomationTool/Linux/LinuxPlatform.Automation.cs
+++ b/Engine/Source/Programs/AutomationTool/Linux/LinuxPlatform.Automation.cs
@@ -85,8 +85,6 @@ public abstract class BaseLinuxPlatform : Platform
                 throw new AutomationException("Can't stage the exe {0} because it doesn't start with {1} or {2}", Exe, CombinePaths(SC.RuntimeProjectRootDir, "Binaries", SC.PlatformDir), CombinePaths(SC.RuntimeRootDir, "Engine/Binaries", SC.PlatformDir));
             }
         }
-
-        SC.StageFiles(StagedFileType.NonUFS, SC.ProjectRoot, "*.png", false, null, SC.RelativeProjectRootForStage);
     }
 
 	public override string GetCookPlatform(bool bDedicatedServer, bool bIsClientOnly, string CookFlavor)
