UE4: Made the maximum days to keep temporary builds an optional parameter

diff --git a/Engine/Source/Programs/AutomationTool/CommandUtils.cs b/Engine/Source/Programs/AutomationTool/CommandUtils.cs
--- a/Engine/Source/Programs/AutomationTool/CommandUtils.cs
+++ b/Engine/Source/Programs/AutomationTool/CommandUtils.cs
@@ -1889,7 +1889,7 @@ namespace AutomationTool
             return Result;
         }
 
-        public static void CleanFormalBuilds(string DirectoryForThisBuild, string CLString = "")
+        public static void CleanFormalBuilds(string DirectoryForThisBuild, string CLString = "", int MaximumDaysToKeepTempStorage = 4)
         {
             if (CLString == "" && (!IsBuildMachine || !DirectoryForThisBuild.StartsWith(RootSharedTempStorageDirectory()) || !P4Enabled))
             {
@@ -1901,7 +1901,6 @@ namespace AutomationTool
                 {
                     CLString = P4Env.ChangelistString;
                 }
-                const int MaximumDaysToKeepTempStorage = 4;
                 string ParentDir = Path.GetDirectoryName(CombinePaths(DirectoryForThisBuild));
                 if (!DirectoryExists_NoExceptions(ParentDir))
                 {
