build script, fix bug in git pull request

diff --git a/Engine/Source/Programs/AutomationTool/Scripts/Tests.Automation.cs b/Engine/Source/Programs/AutomationTool/Scripts/Tests.Automation.cs
--- a/Engine/Source/Programs/AutomationTool/Scripts/Tests.Automation.cs
+++ b/Engine/Source/Programs/AutomationTool/Scripts/Tests.Automation.cs
@@ -290,7 +290,7 @@ class GitPullRequest : BuildCommand
         }
         else
         {
-            bDoingUT = true;
+            bDoingUT = false;
             GitRepositoryURL = GitRepositoryURL_Engine;
         }
         var Dir = ParseParamValue("Dir");
