build script, wait longer for process output

diff --git a/Engine/Source/Programs/AutomationTool/ProcessUtils.cs b/Engine/Source/Programs/AutomationTool/ProcessUtils.cs
--- a/Engine/Source/Programs/AutomationTool/ProcessUtils.cs
+++ b/Engine/Source/Programs/AutomationTool/ProcessUtils.cs
@@ -394,7 +394,7 @@ namespace AutomationTool
 			if (!bProcTerminated)
 			{
 				// Check messages
-				int MaxWaitUntilMessagesReceived = 5;
+				int MaxWaitUntilMessagesReceived = 120;
 				while (MaxWaitUntilMessagesReceived > 0 && !(bStdOutSignalReceived && bStdErrSignalReceived))
 				{
 					if (!bStdOutSignalReceived)
@@ -416,6 +416,10 @@ namespace AutomationTool
 						MaxWaitUntilMessagesReceived--;
 					}
 				}
+                if (!(bStdOutSignalReceived && bStdErrSignalReceived))
+                {
+                    CommandUtils.Log("Waited for a long time for output of {0}, some output may be missing; we gave up.", AppName);
+                }
 
 				// Double-check if the process terminated
 				lock (ProcSyncObject)
