UAT: Adding CommandUtils.ThreadedCopyFiles, ScopedTimer

diff --git a/Engine/Source/Programs/AutomationTool/CommandUtils.cs b/Engine/Source/Programs/AutomationTool/CommandUtils.cs
--- a/Engine/Source/Programs/AutomationTool/CommandUtils.cs
+++ b/Engine/Source/Programs/AutomationTool/CommandUtils.cs
@@ -1398,6 +1398,124 @@ namespace AutomationTool
 
 		#endregion
 
+		#region Threaded Copy
+
+		/// <summary>
+		/// Single batch of files to copy on a thread
+		/// </summary>
+		class CopyRequest
+		{
+			public string[] Source;
+			public string[] Dest;
+			public Exception Result;
+			public Thread Worker;
+		}
+
+		/// <summary>
+		/// Main thread procedure for copying files
+		/// </summary>
+		private static void CopyThreadProc(object Request)
+		{
+			const bool bQuiet = true;
+			var FileToCopy = (CopyRequest)Request;
+			try
+			{
+				for (int Index = 0; Index < FileToCopy.Source.Length; ++Index)
+				{
+					CopyFile(FileToCopy.Source[Index], FileToCopy.Dest[Index], bQuiet);
+				}
+			}
+			catch (Exception Ex)
+			{
+				FileToCopy.Result = Ex;
+			}
+		}
+
+		/// <summary>
+		/// Copies files using miltiple threads
+		/// </summary>
+		/// <param name="SourceDirectory"></param>
+		/// <param name="DestDirectory"></param>
+		/// <param name="MaxThreads"></param>
+		public static void ThreadedCopyFiles(string SourceDirectory, string DestDirectory, int MaxThreads = 64)
+		{
+			CreateDirectory(DestDirectory);
+
+			var SourceFiles = Directory.GetFiles(SourceDirectory, "*", SearchOption.AllDirectories);
+			var SourceBase = GetDirectoryName(SourceDirectory) + GetPathSeparatorChar(PathSeparator.Default);
+			var DestBase = GetDirectoryName(DestDirectory) + GetPathSeparatorChar(PathSeparator.Default);
+			var DestFiles = new string[SourceFiles.Length];
+			for (int Index = 0; Index < SourceFiles.Length; ++Index)
+			{
+				DestFiles[Index] = SourceFiles[Index].Replace(SourceBase, DestBase);
+			}
+			ThreadedCopyFiles(SourceFiles, DestFiles, MaxThreads);
+		}
+
+		/// <summary>
+		/// Copies files using miltiple threads
+		/// </summary>
+		/// <param name="SourceDirectory"></param>
+		/// <param name="DestDirectory"></param>
+		/// <param name="MaxThreads"></param>
+		public static void ThreadedCopyFiles(string[] Source, string[] Dest, int MaxThreads = 64)
+		{
+			Log("Copying {0} file(s) using max {1} thread(s)", Source.Length, MaxThreads);
+
+			if (Source.Length != Dest.Length)
+			{
+				throw new AutomationException("Source count ({0}) does not match Dest count ({1})", Source.Length, Dest.Length);
+			}
+			const int MinFilesPerThread = 10;
+			if (MaxThreads > 1 && Source.Length > MinFilesPerThread)
+			{
+				// Split evenly across the threads
+				int FilesPerThread = Math.Max(Source.Length / MaxThreads, MinFilesPerThread);
+				int ThreadCount = Math.Min(Source.Length / FilesPerThread + 1, MaxThreads);
+				FilesPerThread = Source.Length / ThreadCount + 1;
+
+				// Divide and copy
+				var WorkerThreads = new List<CopyRequest>(ThreadCount);
+				var FilesToCopy = Source.Length;
+				var FirstFileIndex = 0;
+				while (FilesToCopy > 0)
+				{
+					var Request = new CopyRequest();
+					WorkerThreads.Add(Request);
+
+					Request.Source = new string[Math.Min(FilesToCopy, FilesPerThread)];
+					Request.Dest = new string[Request.Source.Length];
+					Array.Copy(Source, FirstFileIndex, Request.Source, 0, Request.Source.Length);
+					Array.Copy(Dest, FirstFileIndex, Request.Dest, 0, Request.Dest.Length);
+
+					Request.Worker = new Thread(new ParameterizedThreadStart(CopyThreadProc));
+					Request.Worker.Start(Request);
+
+					FirstFileIndex += Request.Source.Length;
+					FilesToCopy -= Request.Source.Length;
+				}
+
+				// Wait for completion
+				foreach (var Request in WorkerThreads)
+				{
+					if (Request.Worker.IsAlive)
+					{
+						Request.Worker.Join();
+					}
+				}
+			}
+			else
+			{
+				const bool bQuiet = true;
+				for (int Index = 0; Index < Source.Length; ++Index)
+				{
+					CopyFile(Source[Index], Dest[Index], bQuiet);
+				}
+			}
+		}
+
+		#endregion
+
 		#region Environment variables
 
 		/// <summary>
diff --git a/Engine/Source/Programs/AutomationTool/Scripts/Tests.Automation.cs b/Engine/Source/Programs/AutomationTool/Scripts/Tests.Automation.cs
--- a/Engine/Source/Programs/AutomationTool/Scripts/Tests.Automation.cs
+++ b/Engine/Source/Programs/AutomationTool/Scripts/Tests.Automation.cs
@@ -1955,4 +1955,39 @@ class TestStopProcess : BuildCommand
 		Run(Exe, CmdLine, null, ERunOptions.AllowSpew | ERunOptions.NoWaitForExit | ERunOptions.AppMustExist | ERunOptions.NoStdOutRedirect);
 		Thread.Sleep(10000);
 	}
+}
+
+
+[Help("Copies all files from source directory to destination directory using ThreadedCopyFiles")]
+[Help("Source", "Source path")]
+[Help("Dest", "Destination path")]
+[Help("Threads", "Number of threads used to copy files (64 by default)")]
+class TestThreadedCopyFiles : BuildCommand
+{
+	public override void ExecuteBuild()
+	{
+		var SourcePath = ParseParamValue("Source=");
+		var DestPath = ParseParamValue("Dest=");
+		var Threads = ParseParamInt("Threads=", 64);
+
+		if (String.IsNullOrEmpty(SourcePath))
+		{
+			throw new AutomationException("Source path not specified, please use -source=Path");
+		}
+		if (String.IsNullOrEmpty(DestPath))
+		{
+			throw new AutomationException("Destination path not specified, please use -dest=Path");
+		}
+		if (!DirectoryExists(SourcePath))
+		{
+			throw new AutomationException("Source path {0} does not exist", SourcePath);
+		}
+
+		DeleteDirectory(DestPath);
+
+		using (var ScopedCopyTimer = new ScopedTimer("ThreadedCopyFiles"))
+		{
+			ThreadedCopyFiles(SourcePath, DestPath, Threads);
+		}
+	}
 }
\ No newline at end of file
diff --git a/Engine/Source/Programs/AutomationTool/Utils.cs b/Engine/Source/Programs/AutomationTool/Utils.cs
--- a/Engine/Source/Programs/AutomationTool/Utils.cs
+++ b/Engine/Source/Programs/AutomationTool/Utils.cs
@@ -975,4 +975,39 @@ namespace AutomationTool
 	}
 
 	#endregion
+
+	#region Scoped Timer
+
+	public class ScopedTimer : IDisposable
+	{
+		DateTime StartTime;
+		string TimerName;
+		static int Indent = 0;
+		static object IndentLock = new object();
+
+		public ScopedTimer(string Name)
+		{
+			TimerName = Name;
+			lock (IndentLock)
+			{
+				Indent++;
+			}
+			StartTime = DateTime.UtcNow;
+		}
+
+		public void Dispose()
+		{
+			var TotalSeconds = (DateTime.UtcNow - StartTime).TotalSeconds;
+			var LogIndent = 0;
+			lock (IndentLock)
+			{
+				LogIndent = --Indent;
+			}
+			var IndentText = new StringBuilder(LogIndent * 2);
+			IndentText.Append(' ', LogIndent * 2);
+			Log.TraceVerbose("{0}{1} took {2}s", IndentText.ToString(), TimerName, TotalSeconds);
+		}
+	}
+
+	#endregion
 }
