[INTEGRATE] Change 2083281

missing file for copying the cookie storage for save game
#ue4-tappychicken
#html5

diff --git a/Engine/Source/Programs/AutomationTool/HTML5/HTML5Platform.Automation.cs b/Engine/Source/Programs/AutomationTool/HTML5/HTML5Platform.Automation.cs
--- a/Engine/Source/Programs/AutomationTool/HTML5/HTML5Platform.Automation.cs
+++ b/Engine/Source/Programs/AutomationTool/HTML5/HTML5Platform.Automation.cs
@@ -39,8 +39,15 @@ public class HTML5Platform : Platform
 		}
 		// put the HTML file to the binaries directory
 		string TemplateFile = Path.Combine(CombinePaths(CmdEnv.LocalRoot, "Engine"), "Build", "HTML5", "Game.html.template");
-		string OutputFile = Path.Combine(Path.GetDirectoryName(Path.GetFullPath(Params.ProjectGameExeFilename)), Params.ProjectBinariesFolder, Params.ShortProjectName) + ".html";
+		string OutputFile = Path.Combine(Path.GetDirectoryName(Path.GetFullPath(Params.ProjectGameExeFilename)), Params.ProjectBinariesFolder, (Params.ClientConfigsToBuild[0].ToString() != "Development" ? (Params.ShortProjectName + "-HTML5-" + Params.ClientConfigsToBuild[0].ToString()) : Params.ShortProjectName)) + ".html";
     	GenerateFileFromTemplate(TemplateFile, OutputFile, Params.ShortProjectName, Params.ClientConfigsToBuild[0].ToString(), Params.StageCommandline, !Params.IsCodeBasedProject);
+
+		// copy the jstorage files to the binaries directory
+		string JSDir = Path.Combine(CombinePaths(CmdEnv.LocalRoot, "Engine"), "Build", "HTML5");
+		string OutDir = Path.Combine(Path.GetDirectoryName(Path.GetFullPath(Params.ProjectGameExeFilename)), Params.ProjectBinariesFolder);
+		File.Copy(JSDir + "/json2.js", OutDir + "/json2.js", true);
+		File.Copy(JSDir + "/jstorage.js", OutDir + "/jstorage.js", true);
+		File.Copy(JSDir + "/moz_binarystring.js", OutDir + "/moz_binarystring.js", true);
 		PrintRunTime();
 	}
 
