UAT: fixes for Mono.

- Fixes correct (normalized) filename not being used in FileExists().
- Removes unused variables (safe fix).

This is a subset of changes (compiler warning fixes needed for Linux editor) from pull request 185 by sbc100, 3dluvr, wshearn, abarbu et al (#ue4linux community): https://github.com/EpicGames/UnrealEngine/pull/185

diff --git a/Engine/Source/Programs/AutomationTool/Android/AndroidPlatform.Automation.cs b/Engine/Source/Programs/AutomationTool/Android/AndroidPlatform.Automation.cs
--- a/Engine/Source/Programs/AutomationTool/Android/AndroidPlatform.Automation.cs
+++ b/Engine/Source/Programs/AutomationTool/Android/AndroidPlatform.Automation.cs
@@ -154,7 +154,6 @@ public class AndroidPlatform : Platform
 		{
 			throw new AutomationException("Android is currently only able to package one target configuration at a time, but StageTargetConfigurations contained {0} configurations", SC.StageTargetConfigurations.Count);
 		}
-		var TargetConfiguration = SC.StageTargetConfigurations[0];
 
 		string ApkName = GetFinalApkName(Params, SC.StageExecutables[0], true);
 		string ObbName = GetFinalObbName(ApkName);
diff --git a/Engine/Source/Programs/AutomationTool/CommandUtils.cs b/Engine/Source/Programs/AutomationTool/CommandUtils.cs
--- a/Engine/Source/Programs/AutomationTool/CommandUtils.cs
+++ b/Engine/Source/Programs/AutomationTool/CommandUtils.cs
@@ -651,7 +651,7 @@ namespace AutomationTool
 			foreach (var Filename in Filenames)
 			{
 				var NormalizedFilename = ConvertSeparators(PathSeparator.Default, Filename);
-				bExists = InternalUtils.SafeFileExists(Filename) && bExists;
+				bExists = InternalUtils.SafeFileExists(NormalizedFilename) && bExists;
 			}
 			return bExists;
 		}
@@ -677,7 +677,7 @@ namespace AutomationTool
             foreach (var Filename in Filenames)
             {
                 var NormalizedFilename = ConvertSeparators(PathSeparator.Default, Filename);
-                bExists = InternalUtils.SafeFileExists(Filename, bQuiet) && bExists;
+                bExists = InternalUtils.SafeFileExists(NormalizedFilename, bQuiet) && bExists;
             }
             return bExists;
         }
@@ -1354,8 +1354,6 @@ namespace AutomationTool
                 throw new AutomationException("Failed to create directory {0} for copy", TargetPathBase);
             }
 
-			string[] SourceList = Directory.GetFiles(SourcePathBase);
-
 			DirectoryInfo SourceDirectory = new DirectoryInfo(SourcePathBase);
 			DirectoryInfo[] SourceSubdirectories = SourceDirectory.GetDirectories();
 
diff --git a/Engine/Source/Programs/AutomationTool/Scripts/GUBP.Automation.cs b/Engine/Source/Programs/AutomationTool/Scripts/GUBP.Automation.cs
--- a/Engine/Source/Programs/AutomationTool/Scripts/GUBP.Automation.cs
+++ b/Engine/Source/Programs/AutomationTool/Scripts/GUBP.Automation.cs
@@ -3282,8 +3282,6 @@ public class GUBP : BuildCommand
     }
     public void SaveGraphVisualization(List<string> Nodes)
     {
-        var TimerStartTime = DateTime.UtcNow;
-
         var GraphNodes = new List<GraphNode>();
 
         var NodeToGraphNodeMap = new Dictionary<string, GraphNode>();
@@ -4742,7 +4740,6 @@ if (HostPlatform == UnrealTargetPlatform.Mac) continue; //temp hack till mac aut
         {
             if (CodeProj.Properties.Targets.ContainsKey(TargetRules.TargetType.Editor))
             {
-                var Target = CodeProj.Properties.Targets[TargetRules.TargetType.Editor];
                 bool AnySeparate = false;
                 var PromotedHosts = new List<UnrealTargetPlatform>();
 
diff --git a/Engine/Source/Programs/AutomationTool/Scripts/RunProjectCommand.Automation.cs b/Engine/Source/Programs/AutomationTool/Scripts/RunProjectCommand.Automation.cs
--- a/Engine/Source/Programs/AutomationTool/Scripts/RunProjectCommand.Automation.cs
+++ b/Engine/Source/Programs/AutomationTool/Scripts/RunProjectCommand.Automation.cs
@@ -120,7 +120,6 @@ public partial class Project : CommandUtils
 		{
 			if (Params.ClientTargetPlatforms.Count > 0)
 			{
-				UnrealTargetPlatform ClientPlatform = Params.ClientTargetPlatforms[0];
 				Platform ClientPlatformInst = Params.ClientTargetPlatformInstances[0];
 				string TargetCook = ClientPlatformInst.GetCookPlatform(false, Params.HasDedicatedServerAndClient, Params.CookFlavor);
 				ServerProcess = RunCookOnTheFlyServer(Params.RawProjectPath, Params.NoClient ? "" : ServerLogFile, TargetCook, Params.RunCommandline);
