UAT: Don't try to add empty non-param help attributes to param list. Better handling of empty help attributes.

diff --git a/Engine/Source/Programs/AutomationTool/HelpUtils.cs b/Engine/Source/Programs/AutomationTool/HelpUtils.cs
--- a/Engine/Source/Programs/AutomationTool/HelpUtils.cs
+++ b/Engine/Source/Programs/AutomationTool/HelpUtils.cs
@@ -168,23 +168,35 @@ namespace AutomationTool
 			Dictionary<string, string> ParamDict = new Dictionary<string, string>(StringComparer.InvariantCultureIgnoreCase);
 
 			// Extract Params/Descriptions into Key/Value pairs
-			foreach (var param in Params)
+			foreach (var Param in Params)
 			{
 				// Find the first space (should be following the param name)
-				int SplitPoint = param.IndexOf(' ');
+				if (!String.IsNullOrWhiteSpace(Param))
+				{
+					var ParamName = String.Empty;
+					var ParamDesc = String.Empty;
 
-				// Extract the name and description seperately
-				string ParamName = param.Substring(0, SplitPoint);
-				string ParamDesc = param.Substring(SplitPoint + 1, param.Length - (SplitPoint + 1));
+					var SplitPoint = Param.IndexOf(' ');
+					if (SplitPoint > 0)
+					{
+						// Extract the name and description seperately
+						ParamName = Param.Substring(0, SplitPoint);
+						ParamDesc = Param.Substring(SplitPoint + 1, Param.Length - (SplitPoint + 1));
+					}
+					else
+					{
+						ParamName = Param;
+					}
 
-				// build dictionary using Name and Desc as Key and Value
-				if (!ParamDict.ContainsKey(ParamName))
-				{
-					ParamDict.Add(ParamName, ParamDesc);
-				}
-				else
-				{
-					LogWarning("Duplicated help parameter \"{0}\"", ParamName);
+					// build dictionary using Name and Desc as Key and Value
+					if (!ParamDict.ContainsKey(ParamName))
+					{
+						ParamDict.Add(ParamName, ParamDesc);
+					}
+					else
+					{
+						LogWarning("Duplicated help parameter \"{0}\"", ParamName);
+					}
 				}
 			}
 
@@ -333,7 +345,7 @@ namespace AutomationTool
 				foreach (var CustomAttribute in AllMemeberAtributes)
 				{
 					var HelpAttribute = CustomAttribute as HelpAttribute;
-					if (HelpAttribute != null)
+					if (HelpAttribute != null && HelpAttribute.IsParam)
 					{
 						Params.Add(HelpAttribute.Description);
 					}
