Remove MirrorCloudDirToInternalCDN

diff --git a/Engine/Source/Programs/AutomationTool/MCPPublic.cs b/Engine/Source/Programs/AutomationTool/MCPPublic.cs
--- a/Engine/Source/Programs/AutomationTool/MCPPublic.cs
+++ b/Engine/Source/Programs/AutomationTool/MCPPublic.cs
@@ -420,13 +420,6 @@ namespace EpicGames.MCP.Automation
         /// <param name="command">Build command (used to allow the -CDNDrive cmdline override).</param>
         /// <param name="manifestUrlPath">relative path to the manifest file from the build info service</param>
         abstract public void CopyChunksToProductionCDN(string manifestUrlPath);
-
-        /// <summary>
-        /// Mirrors the CloudDir with the internal web server that can also serve chunks.
-        /// </summary>
-        /// <param name="command">Build command (used to allow the -CDNDrive cmdline override).</param>
-        /// <param name="stagingInfo">Staging info used to determine where the chunks are to copy.</param>
-        abstract public void MirrorCloudDirToInternalCDN(BuildPatchToolStagingInfo stagingInfo);
 	}
     /// <summary>
     /// Helpers for using the MCP account service
