build script, Add custom shared promotable workflow

diff --git a/Engine/Source/Programs/AutomationTool/Scripts/GUBP.Automation.cs b/Engine/Source/Programs/AutomationTool/Scripts/GUBP.Automation.cs
--- a/Engine/Source/Programs/AutomationTool/Scripts/GUBP.Automation.cs
+++ b/Engine/Source/Programs/AutomationTool/Scripts/GUBP.Automation.cs
@@ -1718,6 +1718,10 @@ public class GUBP : BuildCommand
         {
             return WaitForPromotionUserInput.StaticGetFullName("Shared", bInLabelPromoted);
         }
+        public override string GetTriggerStateName()
+        {
+            return GetFullName();
+        } 
     }
 
     public class LabelPromotableNode : GUBPNode
