build script, parallel read and writes to temp storage

diff --git a/Engine/Source/Programs/AutomationTool/TempStorage.cs b/Engine/Source/Programs/AutomationTool/TempStorage.cs
--- a/Engine/Source/Programs/AutomationTool/TempStorage.cs
+++ b/Engine/Source/Programs/AutomationTool/TempStorage.cs
@@ -691,6 +691,12 @@ namespace AutomationTool
         {
             return LocalTempStorageExists(Env, StorageBlockName, bQuiet) || (!bLocalOnly && SharedTempStorageExists(Env, StorageBlockName, GameFolder, bQuiet));
         }
+
+        static int ThreadsToCopyWith()
+        {
+            return 8;
+        }
+
         public static void StoreToTempStorage(CommandEnvironment Env, string StorageBlockName, List<string> Files, bool bLocalOnly = false, string GameFolder = "", string BaseFolder = "")
         {
             if (String.IsNullOrEmpty(BaseFolder))
@@ -720,26 +726,57 @@ namespace AutomationTool
                 {
                     throw new AutomationException("Storage Block Could Not Be Created! {0}", BlockPath);
                 }
+
                 var DestFiles = new List<string>();
-                foreach (string InFilename in Files)
+                if (ThreadsToCopyWith() < 2)
                 {
-                    var Filename = CombinePaths(InFilename);
-                    Robust_FileExists_NoExceptions(false, Filename, "Could not add {0} to manifest because it does not exist");
+                    foreach (string InFilename in Files)
+                    {
+                        var Filename = CombinePaths(InFilename);
+                        Robust_FileExists_NoExceptions(false, Filename, "Could not add {0} to manifest because it does not exist");
 
-                    if (!Filename.StartsWith(BaseFolder, StringComparison.InvariantCultureIgnoreCase))
+                        if (!Filename.StartsWith(BaseFolder, StringComparison.InvariantCultureIgnoreCase))
+                        {
+                            throw new AutomationException("Could not add {0} to manifest because it does not start with the base folder {1}", Filename, BaseFolder);
+                        }
+                        var RelativeFile = Filename.Substring(BaseFolder.Length);
+                        var DestFile = CombinePaths(BlockPath, RelativeFile);
+                        if (FileExists_NoExceptions(true, DestFile))
+                        {
+                            throw new AutomationException("Dest file {0} already exists.", DestFile);
+                        }
+                        CopyFile(Filename, DestFile, true);
+                        Robust_FileExists_NoExceptions(true, DestFile, "Could not copy to {0}");
+
+                        DestFiles.Add(DestFile);
+                    }
+                }
+                else
+                {
+                    var SrcFiles = new List<string>();
+                    foreach (string InFilename in Files)
                     {
-                        throw new AutomationException("Could not add {0} to manifest because it does not start with the base folder {1}", Filename, BaseFolder);
+                        var Filename = CombinePaths(InFilename);
+                        Robust_FileExists_NoExceptions(false, Filename, "Could not add {0} to manifest because it does not exist");
+
+                        if (!Filename.StartsWith(BaseFolder, StringComparison.InvariantCultureIgnoreCase))
+                        {
+                            throw new AutomationException("Could not add {0} to manifest because it does not start with the base folder {1}", Filename, BaseFolder);
+                        }
+                        var RelativeFile = Filename.Substring(BaseFolder.Length);
+                        var DestFile = CombinePaths(BlockPath, RelativeFile);
+                        if (FileExists_NoExceptions(true, DestFile))
+                        {
+                            throw new AutomationException("Dest file {0} already exists.", DestFile);
+                        }
+                        SrcFiles.Add(Filename);
+                        DestFiles.Add(DestFile);
                     }
-                    var RelativeFile = Filename.Substring(BaseFolder.Length);
-                    var DestFile = CombinePaths(BlockPath, RelativeFile);
-                    if (FileExists_NoExceptions(true, DestFile))
+                    ThreadedCopyFiles(SrcFiles.ToArray(), DestFiles.ToArray(), ThreadsToCopyWith());
+                    foreach (string DestFile in DestFiles)
                     {
-                        throw new AutomationException("Dest file {0} already exists.", DestFile);
+                        Robust_FileExists_NoExceptions(true, DestFile, "Could not copy to {0}");
                     }
-                    CopyFile(Filename, DestFile, true);
-                    Robust_FileExists_NoExceptions(true, DestFile, "Could not copy to {0}");
-
-                    DestFiles.Add(DestFile);
                 }
                 var Shared = SaveSharedTempStorageManifest(Env, StorageBlockName, GameFolder, DestFiles);
                 if (!Local.Compare(Shared))
@@ -810,33 +847,72 @@ namespace AutomationTool
             var SharedFiles = Shared.GetFiles(BlockPath);
 
             var DestFiles = new List<string>();
-            foreach (string InFilename in SharedFiles)
+            if (ThreadsToCopyWith() < 2)
             {
-                var Filename = CombinePaths(InFilename);
-                Robust_FileExists_NoExceptions(true, Filename, "Could not add {0} to manifest because it does not exist");
-
-                if (!Filename.StartsWith(BlockPath, StringComparison.InvariantCultureIgnoreCase))
+                foreach (string InFilename in SharedFiles)
                 {
-                    throw new AutomationException("Could not add {0} to manifest because it does not start with the base folder {1}", Filename, BlockPath);
+                    var Filename = CombinePaths(InFilename);
+                    Robust_FileExists_NoExceptions(true, Filename, "Could not add {0} to manifest because it does not exist");
+
+                    if (!Filename.StartsWith(BlockPath, StringComparison.InvariantCultureIgnoreCase))
+                    {
+                        throw new AutomationException("Could not add {0} to manifest because it does not start with the base folder {1}", Filename, BlockPath);
+                    }
+                    var RelativeFile = Filename.Substring(BlockPath.Length);
+                    var DestFile = CombinePaths(BaseFolder, RelativeFile);
+                    if (FileExists_NoExceptions(true, DestFile))
+                    {
+                        Log("Dest file {0} already exists, deleting and overwriting", DestFile);
+                        DeleteFile(DestFile);
+                    }
+                    CopyFile(Filename, DestFile, true);
+
+                    Robust_FileExists_NoExceptions(true, DestFile, "Could not copy to {0}");
+
+                    if (UnrealBuildTool.Utils.IsRunningOnMono && IsProbablyAMacOrIOSExe(DestFile))
+                    {
+                        SetExecutableBit(DestFile);
+                    }
+
+                    FileInfo Info = new FileInfo(DestFile);
+                    DestFiles.Add(Info.FullName);
                 }
-                var RelativeFile = Filename.Substring(BlockPath.Length);
-                var DestFile = CombinePaths(BaseFolder, RelativeFile);
-                if (FileExists_NoExceptions(true, DestFile))
+            }
+            else
+            {
+                var SrcFiles = new List<string>();
+                foreach (string InFilename in SharedFiles)
                 {
-                    Log("Dest file {0} already exists, deleting and overwriting", DestFile);
-                    DeleteFile(DestFile);
-                }
-                CopyFile(Filename, DestFile, true);
-
-                Robust_FileExists_NoExceptions(true, DestFile, "Could not copy to {0}");
+                    var Filename = CombinePaths(InFilename);
+                    //Robust_FileExists_NoExceptions(true, Filename, "Could not add {0} to manifest because it does not exist");
 
-                if (UnrealBuildTool.Utils.IsRunningOnMono && IsProbablyAMacOrIOSExe(DestFile))
+                    if (!Filename.StartsWith(BlockPath, StringComparison.InvariantCultureIgnoreCase))
+                    {
+                        throw new AutomationException("Could not add {0} to manifest because it does not start with the base folder {1}", Filename, BlockPath);
+                    }
+                    var RelativeFile = Filename.Substring(BlockPath.Length);
+                    var DestFile = CombinePaths(BaseFolder, RelativeFile);
+                    if (FileExists_NoExceptions(true, DestFile))
+                    {
+                        Log("Dest file {0} already exists, deleting and overwriting", DestFile);
+                        DeleteFile(DestFile);
+                    }
+                    SrcFiles.Add(Filename);
+                    DestFiles.Add(DestFile);
+                }
+                ThreadedCopyFiles(SrcFiles.ToArray(), DestFiles.ToArray(), ThreadsToCopyWith());
+                var NewDestFiles = new List<string>();
+                foreach (string DestFile in DestFiles)
                 {
-                    SetExecutableBit(DestFile);
+                    Robust_FileExists_NoExceptions(true, DestFile, "Could not copy to {0}");
+                    if (UnrealBuildTool.Utils.IsRunningOnMono && IsProbablyAMacOrIOSExe(DestFile))
+                    {
+                        SetExecutableBit(DestFile);
+                    }
+                    FileInfo Info = new FileInfo(DestFile);
+                    NewDestFiles.Add(Info.FullName);
                 }
-
-                FileInfo Info = new FileInfo(DestFile);
-                DestFiles.Add(Info.FullName);
+                DestFiles = NewDestFiles;
             }
             var NewLocal = SaveLocalTempStorageManifest(Env, BaseFolder, StorageBlockName, DestFiles);
             if (!NewLocal.Compare(Shared))
