[336846] Include OneSky library referenced by automation scripts.

- AutomationScripts.Automation.csproj now uses a valid ToolsVersion number
- The OneSky library reference now points to the DotNet binaries folder similar to other references
- Add OneSky release binary to the Binaries\\DotNet folder

diff --git a/Engine/Source/Programs/AutomationTool/Scripts/AutomationScripts.Automation.csproj b/Engine/Source/Programs/AutomationTool/Scripts/AutomationScripts.Automation.csproj
--- a/Engine/Source/Programs/AutomationTool/Scripts/AutomationScripts.Automation.csproj
+++ b/Engine/Source/Programs/AutomationTool/Scripts/AutomationScripts.Automation.csproj
@@ -1,5 +1,5 @@
 ﻿<?xml version="1.0" encoding="utf-8"?>
-<Project ToolsVersion="11.0" DefaultTargets="Build" xmlns="http://schemas.microsoft.com/developer/msbuild/2003">
+<Project ToolsVersion="4.0" DefaultTargets="Build" xmlns="http://schemas.microsoft.com/developer/msbuild/2003">
   <Import Project="$(MSBuildExtensionsPath)\$(MSBuildToolsVersion)\Microsoft.Common.props" Condition="Exists('$(MSBuildExtensionsPath)\$(MSBuildToolsVersion)\Microsoft.Common.props')" />
   <PropertyGroup>
     <Configuration Condition=" '$(Configuration)' == '' ">Debug</Configuration>
@@ -39,8 +39,9 @@
       <Private>False</Private>
     </Reference>
     <Reference Include="OneSky, Version=1.0.0.0, Culture=neutral, processorArchitecture=MSIL">
-      <SpecificVersion>True</SpecificVersion>
-      <HintPath>..\..\..\ThirdParty\OneSky\bin\Release\OneSky.dll</HintPath>
+      <HintPath>..\..\..\..\Binaries\DotNET\OneSky.dll</HintPath>
+      <SpecificVersion>False</SpecificVersion>
+      <Private>False</Private>
     </Reference>
     <Reference Include="System" />
     <Reference Include="System.Core" />
