build script, Use specific withnightlys workflow when in main branch

diff --git a/Engine/Source/Programs/AutomationTool/Scripts/GUBP.Automation.cs b/Engine/Source/Programs/AutomationTool/Scripts/GUBP.Automation.cs
--- a/Engine/Source/Programs/AutomationTool/Scripts/GUBP.Automation.cs
+++ b/Engine/Source/Programs/AutomationTool/Scripts/GUBP.Automation.cs
@@ -1715,14 +1715,27 @@ public class GUBP : BuildCommand
             return "Make the shared promotable label.";
         }
         public static string StaticGetFullName(bool bInLabelPromoted)
+        {            
+            return WaitForPromotionUserInput.StaticGetFullName("Shared", bInLabelPromoted) + (bIsMainBranch() ? "_WithNightlys" : "");
+        }
+        public static bool bIsMainBranch()
         {
-            return WaitForPromotionUserInput.StaticGetFullName("Shared", bInLabelPromoted);
+            bool bisMain = false;
+            if (P4Enabled)
+            {
+                string CurrentBranch = P4Env.BuildRootP4;
+                if (CurrentBranch == "//depot/UE4")
+                {
+                    bisMain = true;
+                }
+            }
+            return bisMain;
         }
         public override string GetTriggerStateName()
         {
             return GetFullName();
         } 
-    }
+    }    
 
     public class LabelPromotableNode : GUBPNode
     {
