build script, correct _withnightlys logic

diff --git a/Engine/Source/Programs/AutomationTool/Scripts/GUBP.Automation.cs b/Engine/Source/Programs/AutomationTool/Scripts/GUBP.Automation.cs
--- a/Engine/Source/Programs/AutomationTool/Scripts/GUBP.Automation.cs
+++ b/Engine/Source/Programs/AutomationTool/Scripts/GUBP.Automation.cs
@@ -1716,20 +1716,20 @@ public class GUBP : BuildCommand
         }
         public static string StaticGetFullName(bool bInLabelPromoted)
         {            
-            return WaitForPromotionUserInput.StaticGetFullName("Shared", bInLabelPromoted) + (bIsMainBranch() ? "_WithNightlys" : "");
+            return WaitForPromotionUserInput.StaticGetFullName("Shared", bInLabelPromoted) + IsMainBranch();
         }
-        public static bool bIsMainBranch()
+        public static string IsMainBranch()
         {
-            bool bisMain = false;
+            string isMain = "";
             if (P4Enabled)
             {
                 string CurrentBranch = P4Env.BuildRootP4;
                 if (CurrentBranch == "//depot/UE4")
                 {
-                    bisMain = true;
+                    isMain = "_WithNightlys";
                 }
             }
-            return bisMain;
+            return isMain;
         }
         public override string GetTriggerStateName()
         {
