build script, use correct getfullname function

diff --git a/Engine/Source/Programs/AutomationTool/Scripts/GUBP.Automation.cs b/Engine/Source/Programs/AutomationTool/Scripts/GUBP.Automation.cs
--- a/Engine/Source/Programs/AutomationTool/Scripts/GUBP.Automation.cs
+++ b/Engine/Source/Programs/AutomationTool/Scripts/GUBP.Automation.cs
@@ -1733,7 +1733,7 @@ public class GUBP : BuildCommand
         }
         public override string GetTriggerStateName()
         {
-            return GetFullName();
+            return StaticGetFullName(bLabelPromoted);
         } 
     }    
 
