build script, back out previous change to allow shared promotables

diff --git a/Engine/Source/Programs/AutomationTool/Scripts/GUBP.Automation.cs b/Engine/Source/Programs/AutomationTool/Scripts/GUBP.Automation.cs
--- a/Engine/Source/Programs/AutomationTool/Scripts/GUBP.Automation.cs
+++ b/Engine/Source/Programs/AutomationTool/Scripts/GUBP.Automation.cs
@@ -1733,7 +1733,7 @@ public class GUBP : BuildCommand
         }
         public override string GetTriggerStateName()
         {
-            return StaticGetFullName(bLabelPromoted);
+            return GetFullName();
         } 
     }    
 
