Add script to email out BUNS lists to every engineer, taking into account integrations into the release branch.

diff --git a/Engine/Source/Programs/AutomationTool/P4Utils.cs b/Engine/Source/Programs/AutomationTool/P4Utils.cs
--- a/Engine/Source/Programs/AutomationTool/P4Utils.cs
+++ b/Engine/Source/Programs/AutomationTool/P4Utils.cs
@@ -765,6 +765,7 @@ namespace AutomationTool
             public string User = "";
             public string UserEmail = "";
             public string Summary = "";
+			public string Header = "";
 			
 			public class DescribeFile
 			{
@@ -853,6 +854,7 @@ namespace AutomationTool
 
 						DescribeRecord.CL = CurrentChangelist;
                         DescribeRecord.User = Line.Substring(ByAt + MatchBy.Length, AtAt - ByAt - MatchBy.Length);
+						DescribeRecord.Header = Line;
 
 						++LineIndex;
 						if( LineIndex >= Lines.Length )
