build script, add back _withnightlys

diff --git a/Engine/Source/Programs/AutomationTool/Scripts/GUBP.Automation.cs b/Engine/Source/Programs/AutomationTool/Scripts/GUBP.Automation.cs
--- a/Engine/Source/Programs/AutomationTool/Scripts/GUBP.Automation.cs
+++ b/Engine/Source/Programs/AutomationTool/Scripts/GUBP.Automation.cs
@@ -1718,6 +1718,10 @@ public class GUBP : BuildCommand
         {            
             return WaitForPromotionUserInput.StaticGetFullName("Shared", bInLabelPromoted) + IsMainBranch();
         }
+        public override string GetFullName()
+        {
+            return base.GetFullName() + IsMainBranch();
+        }
         public static string IsMainBranch()
         {
             string isMain = "";
