build script, fixing branch setup

diff --git a/Engine/Source/Programs/AutomationTool/Scripts/GUBP.Automation.cs b/Engine/Source/Programs/AutomationTool/Scripts/GUBP.Automation.cs
--- a/Engine/Source/Programs/AutomationTool/Scripts/GUBP.Automation.cs
+++ b/Engine/Source/Programs/AutomationTool/Scripts/GUBP.Automation.cs
@@ -1720,7 +1720,7 @@ public class GUBP : BuildCommand
         }
         public override string GetFullName()
         {
-            return base.GetFullName() + IsMainBranch();
+            return base.GetFullName();// + IsMainBranch();
         }
         public static string IsMainBranch()
         {
