#UE4 #HTML5 TTP: 333781 HTML5: Change Heap size per game project

      Platform ini now contains a [BuildSettings] section with HeapSize values which can be overriden per config.

diff --git a/Engine/Build/HTML5/Game.html.template b/Engine/Build/HTML5/Game.html.template
--- a/Engine/Build/HTML5/Game.html.template
+++ b/Engine/Build/HTML5/Game.html.template
@@ -161,6 +161,7 @@ function preRunHandler() {
 var Module = {
   preRun: [preRunHandler],
   postRun: [],
+  TOTAL_MEMORY: %HEAPSIZE%,
   noImageDecoding: true,
   noAudioDecoding: true,
         arguments: [%UE4CMDLINE%],
diff --git a/Engine/Config/HTML5/HTML5Engine.ini b/Engine/Config/HTML5/HTML5Engine.ini
--- a/Engine/Config/HTML5/HTML5Engine.ini
+++ b/Engine/Config/HTML5/HTML5Engine.ini
@@ -6,3 +6,6 @@ benabled=false
 
 [SystemSettings]
 r.setres=800x600
+
+[BuildSettings]
+HeapSize=1024
diff --git a/Engine/Source/Programs/AutomationTool/HTML5/HTML5Platform.Automation.cs b/Engine/Source/Programs/AutomationTool/HTML5/HTML5Platform.Automation.cs
--- a/Engine/Source/Programs/AutomationTool/HTML5/HTML5Platform.Automation.cs
+++ b/Engine/Source/Programs/AutomationTool/HTML5/HTML5Platform.Automation.cs
@@ -40,7 +40,29 @@ public class HTML5Platform : Platform
 		// put the HTML file to the binaries directory
 		string TemplateFile = Path.Combine(CombinePaths(CmdEnv.LocalRoot, "Engine"), "Build", "HTML5", "Game.html.template");
 		string OutputFile = Path.Combine(Path.GetDirectoryName(Path.GetFullPath(Params.ProjectGameExeFilename)), Params.ProjectBinariesFolder, (Params.ClientConfigsToBuild[0].ToString() != "Development" ? (Params.ShortProjectName + "-HTML5-" + Params.ClientConfigsToBuild[0].ToString()) : Params.ShortProjectName)) + ".html";
-    	GenerateFileFromTemplate(TemplateFile, OutputFile, Params.ShortProjectName, Params.ClientConfigsToBuild[0].ToString(), Params.StageCommandline, !Params.IsCodeBasedProject);
+    	
+
+		// find Heap Size.
+		ulong HeapSize;
+		var ConfigCache = new UnrealBuildTool.ConfigCacheIni(UnrealTargetPlatform.HTML5, "Engine", Path.GetDirectoryName(Params.RawProjectPath),CombinePaths(CmdEnv.LocalRoot, "Engine")); 
+
+
+		int ConfigHeapSize; 
+		if (!ConfigCache.GetInt32 ("BuildSettings","HeapSize"+Params.ClientConfigsToBuild[0].ToString(), out ConfigHeapSize)) // in Megs.
+		{ 	 
+			// we couldn't find a per config heap size, look for a common one.
+			if(!ConfigCache.GetInt32 ("BuildSettings","HeapSize", out ConfigHeapSize))
+				{
+					ConfigHeapSize = Params.IsCodeBasedProject ? 1024 : 512;
+					Log ("Could not find Heap Size setting in .ini for Client config {0}" ,Params.ClientConfigsToBuild[0].ToString()); 
+				}
+		}
+
+		HeapSize = (ulong)ConfigHeapSize * 1024L * 1024L; // convert to bytes. 
+		Log ("Setting Heap size to {0} Mb ", ConfigHeapSize); 
+
+
+		GenerateFileFromTemplate(TemplateFile, OutputFile, Params.ShortProjectName, Params.ClientConfigsToBuild[0].ToString(), Params.StageCommandline, !Params.IsCodeBasedProject, HeapSize);
 
 		// copy the jstorage files to the binaries directory
 		string JSDir = Path.Combine(CombinePaths(CmdEnv.LocalRoot, "Engine"), "Build", "HTML5");
@@ -51,7 +73,7 @@ public class HTML5Platform : Platform
 		PrintRunTime();
 	}
 
-	protected void GenerateFileFromTemplate(string InTemplateFile, string InOutputFile, string InGameName, string InGameConfiguration, string InArguments, bool IsContentOnly)
+	protected void GenerateFileFromTemplate(string InTemplateFile, string InOutputFile, string InGameName, string InGameConfiguration, string InArguments, bool IsContentOnly, ulong HeapSize)
 	{
 		StringBuilder outputContents = new StringBuilder();
 		using (StreamReader reader = new StreamReader(InTemplateFile))
@@ -65,6 +87,11 @@ public class HTML5Platform : Platform
 					LineStr = LineStr.Replace("%GAME%", InGameName);
 				}
 
+				if (LineStr.Contains ("%HEAPSIZE%")) 
+				{
+					LineStr = LineStr.Replace ("%HEAPSIZE%", HeapSize.ToString ());
+				}
+
 				if (LineStr.Contains("%CONFIG%"))
 				{
                     if (IsContentOnly)
