build script, add promotable label suffix, _withnightlys now working, shared promotables now work correctly with custom workflows

diff --git a/Engine/Source/Programs/AutomationTool/Scripts/GUBP.Automation.cs b/Engine/Source/Programs/AutomationTool/Scripts/GUBP.Automation.cs
--- a/Engine/Source/Programs/AutomationTool/Scripts/GUBP.Automation.cs
+++ b/Engine/Source/Programs/AutomationTool/Scripts/GUBP.Automation.cs
@@ -1612,11 +1612,13 @@ public class GUBP : BuildCommand
     public class WaitForPromotionUserInput : WaitForUserInput
     {
         string PromotionLabelPrefix;
+        string PromotionLabelSuffix;
         protected bool bLabelPromoted; // true if this is the promoted version
 
-        public WaitForPromotionUserInput(string InPromotionLabelPrefix, bool bInLabelPromoted)
+        public WaitForPromotionUserInput(string InPromotionLabelPrefix, string InPromotionLabelSuffix, bool bInLabelPromoted)
         {
             PromotionLabelPrefix = InPromotionLabelPrefix;
+            PromotionLabelSuffix = InPromotionLabelSuffix;
             bLabelPromoted = bInLabelPromoted;
             if (bLabelPromoted)
             {
@@ -1627,23 +1629,22 @@ public class GUBP : BuildCommand
                 AddDependency(AggregatePromotableNode.StaticGetFullName(PromotionLabelPrefix));
             }
         }
-        public static string StaticGetFullName(string InPromotionLabelPrefix, bool bInLabelPromoted)
+        public static string StaticGetFullName(string InPromotionLabelPrefix, string InPromotionLabelSuffix, bool bInLabelPromoted)
         {
-            return InPromotionLabelPrefix + (bInLabelPromoted ? "_WaitForPromotion" : "_WaitForPromotable");
+            return InPromotionLabelPrefix + (bInLabelPromoted ? "_WaitForPromotion" : "_WaitForPromotable") + InPromotionLabelSuffix;
         }
         public override string GetFullName()
         {
-            return StaticGetFullName(PromotionLabelPrefix, bLabelPromoted);
+            return StaticGetFullName(PromotionLabelPrefix, PromotionLabelSuffix, bLabelPromoted);
         }
     }
 
     public class WaitForGamePromotionUserInput : WaitForPromotionUserInput
     {
         BranchInfo.BranchUProject GameProj;
-        bool bCustomWorkflow;
-
+        bool bCustomWorkflow;        
         public WaitForGamePromotionUserInput(GUBP bp, BranchInfo.BranchUProject InGameProj, bool bInLabelPromoted)
-            : base(InGameProj.GameName, bInLabelPromoted)
+            : base(InGameProj.GameName, "", bInLabelPromoted)
         {
             GameProj = InGameProj;
             var Options = InGameProj.Options(UnrealTargetPlatform.Win64);
@@ -1651,7 +1652,7 @@ public class GUBP : BuildCommand
         }
         public static string StaticGetFullName(BranchInfo.BranchUProject InGameProj, bool bInLabelPromoted)
         {
-            return WaitForPromotionUserInput.StaticGetFullName(InGameProj.GameName, bInLabelPromoted);
+            return WaitForPromotionUserInput.StaticGetFullName(InGameProj.GameName, "", bInLabelPromoted);
         }
         public override string GameNameIfAnyForTempStorage()
         {
@@ -1693,9 +1694,9 @@ public class GUBP : BuildCommand
 
     public class WaitForSharedPromotionUserInput : WaitForPromotionUserInput
     {
-
+        
         public WaitForSharedPromotionUserInput(GUBP bp, bool bInLabelPromoted)
-            : base("Shared", bInLabelPromoted)
+            : base("Shared", IsMainBranch(), bInLabelPromoted)
         {
         }
         public override string GetTriggerDescText()
@@ -1716,11 +1717,7 @@ public class GUBP : BuildCommand
         }
         public static string StaticGetFullName(bool bInLabelPromoted)
         {            
-            return WaitForPromotionUserInput.StaticGetFullName("Shared", bInLabelPromoted) + IsMainBranch();
-        }
-        public override string GetFullName()
-        {
-            return base.GetFullName();// + IsMainBranch();
+            return WaitForPromotionUserInput.StaticGetFullName("Shared", IsMainBranch(), bInLabelPromoted);
         }
         public static string IsMainBranch()
         {
@@ -1738,23 +1735,27 @@ public class GUBP : BuildCommand
         public override string GetTriggerStateName()
         {
             return GetFullName();
-        } 
+        }
+        public override bool TriggerRequiresRecursiveWorkflow()
+        {
+            return !bLabelPromoted;
+        }
     }    
 
     public class LabelPromotableNode : GUBPNode
     {
-        string PromotionLabelPrefix;
+        string PromotionLabelPrefix;        
         protected bool bLabelPromoted; // true if this is the promoted version
 
-        public LabelPromotableNode(string InPromotionLabelPrefix, bool bInLabelPromoted)
+        public LabelPromotableNode(string InPromotionLabelPrefix, string InPromotionLabelSuffix, bool bInLabelPromoted)
         {
-             PromotionLabelPrefix = InPromotionLabelPrefix;
-             bLabelPromoted = bInLabelPromoted;
-             AddDependency(WaitForPromotionUserInput.StaticGetFullName(PromotionLabelPrefix, bLabelPromoted));
+            PromotionLabelPrefix = InPromotionLabelPrefix;            
+            bLabelPromoted = bInLabelPromoted;
+            AddDependency(WaitForPromotionUserInput.StaticGetFullName(PromotionLabelPrefix, InPromotionLabelSuffix, bLabelPromoted));
         }
         string LabelName(bool bLocalLabelPromoted)
         {
-            string LabelPrefix = PromotionLabelPrefix;
+            string LabelPrefix = PromotionLabelPrefix;            
             string CompleteLabelPrefix = (bLocalLabelPromoted ? "Promoted-" : "Promotable-") + LabelPrefix;
             if (LabelPrefix == "Shared" && bLocalLabelPromoted)
             {
@@ -1845,10 +1846,9 @@ public class GUBP : BuildCommand
 
     public class GameLabelPromotableNode : LabelPromotableNode
     {
-        BranchInfo.BranchUProject GameProj;
-
+        BranchInfo.BranchUProject GameProj;        
         public GameLabelPromotableNode(GUBP bp, BranchInfo.BranchUProject InGameProj, bool bInLabelPromoted)
-            : base(InGameProj.GameName, bInLabelPromoted)
+            : base(InGameProj.GameName, "", bInLabelPromoted)
         {
             GameProj = InGameProj;
         }
@@ -1864,9 +1864,9 @@ public class GUBP : BuildCommand
     }
 
     public class SharedLabelPromotableNode : LabelPromotableNode
-    {
+    {        
         public SharedLabelPromotableNode(GUBP bp, bool bInLabelPromoted)
-            : base("Shared", bInLabelPromoted)
+            : base("Shared", IsMainBranch(), bInLabelPromoted)
         {
         }
 
@@ -1874,6 +1874,19 @@ public class GUBP : BuildCommand
         {
             return LabelPromotableNode.StaticGetFullName("Shared", bInLabelPromoted);
         }
+        public static string IsMainBranch()
+        {
+            string isMain = "";
+            if (P4Enabled)
+            {
+                string CurrentBranch = P4Env.BuildRootP4;
+                if (CurrentBranch == "//depot/UE4")
+                {
+                    isMain = "_WithNightlys";
+                }
+            }
+            return isMain;
+        }
     }
 
     public class WaitForTestShared : AggregateNode
