build script, some retries on temp storage

diff --git a/Engine/Source/Programs/AutomationTool/Scripts/GUBP.Automation.cs b/Engine/Source/Programs/AutomationTool/Scripts/GUBP.Automation.cs
--- a/Engine/Source/Programs/AutomationTool/Scripts/GUBP.Automation.cs
+++ b/Engine/Source/Programs/AutomationTool/Scripts/GUBP.Automation.cs
@@ -4058,19 +4058,19 @@ public class GUBP : BuildCommand
         if (bFakeEC)
         {
             LocalOnly = true;
-        }
+        } 
         if (bSaveSharedTempStorage)
         {
-            if (!HaveSharedTempStorage())
+            if (!HaveSharedTempStorage(true))
             {
                 throw new AutomationException("Request to save to temp storage, but {0} is unavailable.", UE4TempStorageDirectory());
             }
             bSignBuildProducts = true;
         }
-        else if (!LocalOnly && !HaveSharedTempStorage())
+        else if (!LocalOnly && !HaveSharedTempStorage(false))
         {
             Log("Looks like we want to use shared temp storage, but since we don't have it, we won't use it.");
-            LocalOnly = false;
+            LocalOnly = true;
         }
 
         bool CommanderSetup = ParseParam("CommanderJobSetupOnly");
diff --git a/Engine/Source/Programs/AutomationTool/TempStorage.cs b/Engine/Source/Programs/AutomationTool/TempStorage.cs
--- a/Engine/Source/Programs/AutomationTool/TempStorage.cs
+++ b/Engine/Source/Programs/AutomationTool/TempStorage.cs
@@ -557,13 +557,27 @@ namespace AutomationTool
             string Root = RootSharedTempStorageDirectory();
             return CombinePaths(Root, "UE4", SharedSubdir);
         }
-        public static bool HaveSharedTempStorage()
+        public static bool HaveSharedTempStorage(bool ForSaving)
         {
-            if (!DirectoryExists_NoExceptions(UE4TempStorageDirectory()))
+            var Dir = UE4TempStorageDirectory();
+            if (ForSaving)
             {
-                return false;
+                int Retries = 24;
+                while (--Retries > 0)
+                {
+                    if (DirectoryExistsAndIsWritable_NoExceptions(Dir))
+                    {
+                        return true;
+                    }
+                    FindDirectories_NoExceptions(false, "*", false, Dir); // there is some internet evidence to suggest this might perk up the mac share
+                    System.Threading.Thread.Sleep(5000);
+                }
             }
-            return true;
+            else if (DirectoryExists_NoExceptions(Dir))
+            {
+                return true;
+            }
+            return false;
         }
 
         public static string ResolveSharedTempStorageDirectory(string GameFolder)
