Added property names and tool tips to targets for localization automation script.

diff --git a/Engine/Source/Programs/AutomationTool/Scripts/Localisation.Automation.cs b/Engine/Source/Programs/AutomationTool/Scripts/Localisation.Automation.cs
--- a/Engine/Source/Programs/AutomationTool/Scripts/Localisation.Automation.cs
+++ b/Engine/Source/Programs/AutomationTool/Scripts/Localisation.Automation.cs
@@ -42,6 +42,8 @@ class Localise : BuildCommand
 				{
                     String.Format("-config={0}", @"../../../Engine/Config/Localization/Engine.ini") + (string.IsNullOrEmpty(CommandletSCCArguments) ? "" : " " + CommandletSCCArguments),
                     String.Format("-config={0}", @"../../../Engine/Config/Localization/Editor.ini") + (string.IsNullOrEmpty(CommandletSCCArguments) ? "" : " " + CommandletSCCArguments),
+                    String.Format("-config={0}", @"../../../Engine/Config/Localization/PropertyNames.ini") + (string.IsNullOrEmpty(CommandletSCCArguments) ? "" : " " + CommandletSCCArguments),
+                    String.Format("-config={0}", @"../../../Engine/Config/Localization/ToolTips.ini") + (string.IsNullOrEmpty(CommandletSCCArguments) ? "" : " " + CommandletSCCArguments),
                     String.Format("-config={0}", @"../../../Engine/Config/Localization/WordCount.ini"),
 				};
 
