#UE4  HTML5
        - actually use python in the command line, rather than depend on .py extention  to execute.
        - change file types to writable.

diff --git a/Engine/Source/Programs/AutomationTool/HTML5/HTML5Platform.Automation.cs b/Engine/Source/Programs/AutomationTool/HTML5/HTML5Platform.Automation.cs
--- a/Engine/Source/Programs/AutomationTool/HTML5/HTML5Platform.Automation.cs
+++ b/Engine/Source/Programs/AutomationTool/HTML5/HTML5Platform.Automation.cs
@@ -28,12 +28,12 @@ public class HTML5Platform : Platform
 			if (Utils.IsRunningOnMono) 
 			{
 				string PackagerPath = BaseSDKPath + "/tools/file_packager.py";
-				string CmdLine = string.Format ("-c \" python {0} {1} --preload . --pre-run --js-output={1}.js \" ", PackagerPath, FinalDataLocation);
+				string CmdLine = string.Format ("-c \" python {0} {1} --preload . --js-output={1}.js \" ", PackagerPath, FinalDataLocation);
 				RunAndLog (CmdEnv, "/bin/bash", CmdLine);	
 			} else 
 			{ 
 				string PackagerPath = "\"" + BaseSDKPath + "\\tools\\file_packager.py\"";
-				string CmdLine = string.Format ("/c {0} {1} --preload . --pre-run --js-output={1}.js", PackagerPath, FinalDataLocation);
+				string CmdLine = string.Format ("/c python {0} {1} --preload . --js-output={1}.js", PackagerPath, FinalDataLocation);
 				RunAndLog (CmdEnv, CommandUtils.CombinePaths (Environment.SystemDirectory, "cmd.exe"), CmdLine);	
 			}
 		}
@@ -67,9 +67,9 @@ public class HTML5Platform : Platform
 		// copy the jstorage files to the binaries directory
 		string JSDir = Path.Combine(CombinePaths(CmdEnv.LocalRoot, "Engine"), "Build", "HTML5");
 		string OutDir = Path.Combine(Path.GetDirectoryName(Path.GetFullPath(Params.ProjectGameExeFilename)), Params.ProjectBinariesFolder);
-		File.Copy(JSDir + "/json2.js", OutDir + "/json2.js", true);
-		File.Copy(JSDir + "/jstorage.js", OutDir + "/jstorage.js", true);
-		File.Copy(JSDir + "/moz_binarystring.js", OutDir + "/moz_binarystring.js", true);
+        File.Copy(JSDir + "/json2.js", OutDir + "/json2.js", true);
+        File.Copy(JSDir + "/jstorage.js", OutDir + "/jstorage.js", true);
+        File.Copy(JSDir + "/moz_binarystring.js", OutDir + "/moz_binarystring.js", true);
 		PrintRunTime();
 	}
 
