#UE4 #HTML5
     - fix an issue where packaging will fail - if the target binary directory doesn't exist. ( Happens on a clean spec, first time run ).

diff --git a/Engine/Source/Programs/AutomationTool/HTML5/HTML5Platform.Automation.cs b/Engine/Source/Programs/AutomationTool/HTML5/HTML5Platform.Automation.cs
--- a/Engine/Source/Programs/AutomationTool/HTML5/HTML5Platform.Automation.cs
+++ b/Engine/Source/Programs/AutomationTool/HTML5/HTML5Platform.Automation.cs
@@ -19,6 +19,12 @@ public class HTML5Platform : Platform
 		Log("Package {0}", Params.RawProjectPath);
 
 		string FinalDataLocation = Path.Combine(Path.GetDirectoryName(Path.GetFullPath(Params.ProjectGameExeFilename)), Params.ProjectBinariesFolder, Params.ShortProjectName) +".data";
+        string OutDir = Path.Combine(Path.GetDirectoryName(Path.GetFullPath(Params.ProjectGameExeFilename)), Params.ProjectBinariesFolder);
+
+        bool FolderExists = Directory.Exists(OutDir);
+        if (!FolderExists)
+            Directory.CreateDirectory(OutDir);
+
 
 		// we need to operate in the root
 		using (new PushedDirectory(Path.Combine(Params.BaseStageDirectory, "HTML5")))
@@ -66,7 +72,7 @@ public class HTML5Platform : Platform
 
 		// copy the jstorage files to the binaries directory
 		string JSDir = Path.Combine(CombinePaths(CmdEnv.LocalRoot, "Engine"), "Build", "HTML5");
-		string OutDir = Path.Combine(Path.GetDirectoryName(Path.GetFullPath(Params.ProjectGameExeFilename)), Params.ProjectBinariesFolder);
+	
         File.Copy(JSDir + "/json2.js", OutDir + "/json2.js", true);
         File.Copy(JSDir + "/jstorage.js", OutDir + "/jstorage.js", true);
         File.Copy(JSDir + "/moz_binarystring.js", OutDir + "/moz_binarystring.js", true);
