build script, merge fix for p4 changes

diff --git a/Engine/Source/Programs/AutomationTool/P4Utils.cs b/Engine/Source/Programs/AutomationTool/P4Utils.cs
--- a/Engine/Source/Programs/AutomationTool/P4Utils.cs
+++ b/Engine/Source/Programs/AutomationTool/P4Utils.cs
@@ -616,7 +616,9 @@ namespace AutomationTool
         public bool Changes(out List<ChangeRecord> ChangeRecords, string CommandLine, bool AllowSpew = true, bool UseCaching = false, bool LongComment = false)
         {
             // If the user specified '-l' or '-L', the summary will appear on subsequent lines (no quotes) instead of the same line (surrounded by single quotes)
-            bool bSummaryIsOnSameLine = CommandLine.IndexOf("-L", StringComparison.InvariantCultureIgnoreCase) == -1;
+            bool ContainsDashL = CommandLine.StartsWith("-L ", StringComparison.InvariantCultureIgnoreCase) ||
+                CommandLine.IndexOf(" -L ", StringComparison.InvariantCultureIgnoreCase) > 0;
+            bool bSummaryIsOnSameLine = !ContainsDashL;
             if (bSummaryIsOnSameLine && LongComment)
             {
                 CommandLine = "-L " + CommandLine;
@@ -689,7 +691,7 @@ namespace AutomationTool
 							Line = Lines[ LineIndex ];
 							if( !String.IsNullOrEmpty( Line ) )
 							{
-								throw new AutomationException("Was expecting blank line after Change header output from P4");
+                                throw new AutomationException("Was expecting blank line after Change header output from P4, got {0}", Line);
 							}
 
 							++LineIndex;
@@ -704,6 +706,7 @@ namespace AutomationTool
 								if (SummaryChangeAt == 0 && SummaryOnAt > SummaryChangeAt && SummaryByAt > SummaryOnAt)
 								{
 									// OK, we found a new change. This isn't part of our summary.  We're done with the summary.  Back we go.
+                                    //CommandUtils.Log("Next summary is {0}", Line);
 									--LineIndex;
 									break;
 								}
diff --git a/Engine/Source/Programs/AutomationTool/Scripts/Tests.Automation.cs b/Engine/Source/Programs/AutomationTool/Scripts/Tests.Automation.cs
--- a/Engine/Source/Programs/AutomationTool/Scripts/Tests.Automation.cs
+++ b/Engine/Source/Programs/AutomationTool/Scripts/Tests.Automation.cs
@@ -1889,16 +1889,29 @@ class TestChanges : BuildCommand
 {
     public override void ExecuteBuild()
     {
-        var CommandParam = ParseParamValue("CommandParam", "//depot/UE4-LauncherReleases/*/Source/...@2061085,2061287 //depot/UE4-LauncherReleases/*/Build/...@2061085,2061287");
-        
-        List<P4Connection.ChangeRecord> ChangeRecords;
-        if (!P4.Changes(out ChangeRecords, CommandParam, true, true, LongComment: true))
+        var CommandParam = ParseParamValue("CommandParam", "//depot/UE4-LauncherReleases/*/Source/...@2091742,2091950 //depot/UE4-LauncherReleases/*/Build/...@2091742,2091950");
+
         {
-            throw new AutomationException("failed");
+            List<P4Connection.ChangeRecord> ChangeRecords;
+            if (!P4.Changes(out ChangeRecords, CommandParam, true, true, LongComment: true))
+            {
+                throw new AutomationException("failed");
+            }
+            foreach (var Record in ChangeRecords)
+            {
+                Log("{0} {1} {2}", Record.CL, Record.UserEmail, Record.Summary);
+            }
         }
-        foreach (var Record in ChangeRecords)
         {
-            Log("{0} {1} {2}", Record.CL, Record.UserEmail, Record.Summary);
+            List<P4Connection.ChangeRecord> ChangeRecords;
+            if (!P4.Changes(out ChangeRecords, "-L " + CommandParam, true, true, LongComment: false))
+            {
+                throw new AutomationException("failed");
+            }
+            foreach (var Record in ChangeRecords)
+            {
+                Log("{0} {1} {2}", Record.CL, Record.UserEmail, Record.Summary);
+            }
         }
     }
 }
