build script, add formal build testing to the GUBP

diff --git a/Engine/Source/Programs/AutomationTool/Scripts/GUBP.Automation.cs b/Engine/Source/Programs/AutomationTool/Scripts/GUBP.Automation.cs
--- a/Engine/Source/Programs/AutomationTool/Scripts/GUBP.Automation.cs
+++ b/Engine/Source/Programs/AutomationTool/Scripts/GUBP.Automation.cs
@@ -2271,6 +2271,8 @@ public class GUBP : BuildCommand
     public class FormalBuildNode : HostPlatformNode
     {
         BranchInfo.BranchUProject GameProj;
+
+        //CAUTION, these are lists, but it isn't clear that lists really work on all platforms, so we stick to one node per build
         List<UnrealTargetPlatform> ClientTargetPlatforms;
         List<UnrealTargetPlatform> ServerTargetPlatforms;
         List<UnrealTargetConfiguration> ClientConfigs;
@@ -2442,7 +2444,13 @@ public class GUBP : BuildCommand
         {
             return base.CISFrequencyQuantumShift(bp) + 3;
         }
-
+        public static string GetArchiveDirectory(BranchInfo.BranchUProject InGameProj, UnrealTargetPlatform InHostPlatform, List<UnrealTargetPlatform> InClientTargetPlatforms = null, List<UnrealTargetConfiguration> InClientConfigs = null, List<UnrealTargetPlatform> InServerTargetPlatforms = null, List<UnrealTargetConfiguration> InServerConfigs = null, bool InClientNotGame = false)
+        {
+            string BaseDir = ResolveSharedBuildDirectory(InGameProj.GameName);
+            string NodeName = StaticGetFullName(InGameProj, InHostPlatform, InClientTargetPlatforms, InClientConfigs, InServerTargetPlatforms, InServerConfigs, InClientNotGame);
+            string ArchiveDirectory = CombinePaths(BaseDir, NodeName, P4Env.BuildRootEscaped + "-CL-" + P4Env.ChangelistString);
+            return ArchiveDirectory;
+        }
         public override void DoBuild(GUBP bp)
         {
             BuildProducts = new List<string>();
@@ -2523,8 +2531,7 @@ public class GUBP : BuildCommand
 
             if (P4Enabled)
             {
-                string BaseDir = ResolveSharedBuildDirectory(GameProj.GameName);
-                string ArchiveDirectory = CombinePaths(BaseDir, GetFullName(), P4Env.BuildRootEscaped + "-CL-" + P4Env.ChangelistString);
+                string ArchiveDirectory = GetArchiveDirectory(GameProj, HostPlatform, ClientTargetPlatforms, ClientConfigs, ServerTargetPlatforms, ServerConfigs, ClientNotGame);
                 CleanFormalBuilds(ArchiveDirectory);
                 if (DirectoryExists_NoExceptions(ArchiveDirectory))
                 {
@@ -2534,7 +2541,7 @@ public class GUBP : BuildCommand
                     }
                     DeleteDirectory_NoExceptions(ArchiveDirectory);
                 }
-                Args += String.Format(" -Archive -archivedirectory={0}", ArchiveDirectory);
+                Args += String.Format(" -Archive -archivedirectory={0}", CommandUtils.MakePathSafeToUseWithCommandLine(ArchiveDirectory));
             }
 
             string LogFile = CommandUtils.RunUAT(CommandUtils.CmdEnv, Args);
@@ -2558,12 +2565,112 @@ public class GUBP : BuildCommand
         }
         public override void DoBuild(GUBP bp)
         {
-            //aggregates don't do anything, they just merge dependencies
             BuildProducts = new List<string>();
             DoTest(bp);
         }
     }
 
+    public class FormalBuildTestNode : TestNode
+    {
+        BranchInfo.BranchUProject GameProj;
+        UnrealTargetPlatform ClientTargetPlatform;
+        UnrealTargetConfiguration ClientConfig;
+        UnrealBuildTool.TargetRules.TargetType GameOrClient;
+
+        public FormalBuildTestNode(GUBP bp,
+            BranchInfo.BranchUProject InGameProj,
+            UnrealTargetPlatform InHostPlatform,
+            UnrealTargetPlatform InClientTargetPlatform,
+            UnrealTargetConfiguration InClientConfig
+            )
+            : base(InHostPlatform)
+        {
+            GameProj = InGameProj;
+            ClientTargetPlatform = InClientTargetPlatform;
+            ClientConfig = InClientConfig;
+            GameOrClient = TargetRules.TargetType.Game;
+
+            // verify we actually built these
+            var WorkingGameProject = InGameProj;
+            if (!WorkingGameProject.Properties.Targets.ContainsKey(TargetRules.TargetType.Editor))
+            {
+                // this is a codeless project, use the base project
+                WorkingGameProject = bp.Branch.BaseEngineProject;
+            }
+
+            if (!WorkingGameProject.Properties.Targets.ContainsKey(GameOrClient))
+            {
+                throw new AutomationException("Can't make a game build for {0} because it doesn't have a {1} target.", WorkingGameProject.GameName, GameOrClient.ToString());
+            }
+
+            if (!WorkingGameProject.Properties.Targets[GameOrClient].Rules.GUBP_GetPlatforms_MonolithicOnly(HostPlatform).Contains(ClientTargetPlatform))
+            {
+                throw new AutomationException("Can't make a game/client build for {0} because we didn't build platform {1}.", WorkingGameProject.GameName, ClientTargetPlatform.ToString());
+            }
+            if (!WorkingGameProject.Properties.Targets[GameOrClient].Rules.GUBP_GetConfigs_MonolithicOnly(HostPlatform, ClientTargetPlatform).Contains(ClientConfig))
+            {
+                throw new AutomationException("Can't make a game/client build for {0} because we didn't build platform {1} config {2}.", WorkingGameProject.GameName, ClientTargetPlatform.ToString(), ClientConfig.ToString());
+            }
+            AddDependency(FormalBuildNode.StaticGetFullName(GameProj, HostPlatform, new List<UnrealTargetPlatform>() { ClientTargetPlatform }, InClientConfigs: new List<UnrealTargetConfiguration>() { ClientConfig }, InClientNotGame: GameOrClient == TargetRules.TargetType.Client));
+        }
+
+        public static string StaticGetFullName(BranchInfo.BranchUProject InGameProj, UnrealTargetPlatform InHostPlatform, UnrealTargetPlatform InClientTargetPlatform, UnrealTargetConfiguration InClientConfig)
+        {
+            string Infix = "_" + InClientTargetPlatform.ToString();
+            Infix += "_" + InClientConfig.ToString();
+            return InGameProj.GameName + Infix + "_TestBuild" + HostPlatformNode.StaticGetHostPlatformSuffix(InHostPlatform);
+        }
+        public override string GetFullName()
+        {
+            return StaticGetFullName(GameProj, HostPlatform, ClientTargetPlatform, ClientConfig);
+        }
+        public override string GameNameIfAnyForTempStorage()
+        {
+            return GameProj.GameName;
+        }
+        public override float Priority()
+        {
+            return base.Priority() - 20.0f;
+        }
+        public override int CISFrequencyQuantumShift(GUBP bp)
+        {
+            return base.CISFrequencyQuantumShift(bp) + 3;
+        }
+        public override void DoTest(GUBP bp)
+        {
+            string ProjectArg = "";
+            if (!String.IsNullOrEmpty(GameProj.FilePath))
+            {
+                ProjectArg = " -project=\"" + GameProj.FilePath + "\"";
+            }
+
+            string ArchiveDirectory = FormalBuildNode.GetArchiveDirectory(GameProj, HostPlatform, new List<UnrealTargetPlatform>() { ClientTargetPlatform }, InClientConfigs: new List<UnrealTargetConfiguration>() { ClientConfig }, InClientNotGame: GameOrClient == TargetRules.TargetType.Client);
+            if (!DirectoryExists_NoExceptions(ArchiveDirectory))
+            {
+                throw new AutomationException("Archive directory does not exist {0}, so we can't test the build.", ArchiveDirectory);
+            }
+
+            string WorkingCommandline = String.Format("TestFormalBuild {0} -Archive -archivedirectory={1} -platform={2} -clientconfig={3}",
+                ProjectArg, CommandUtils.MakePathSafeToUseWithCommandLine(ArchiveDirectory), ClientTargetPlatform.ToString(), ClientConfig.ToString());
+
+            if (WorkingCommandline.Contains("-project=\"\""))
+            {
+                throw new AutomationException("Command line {0} contains -project=\"\" which is doomed to fail", WorkingCommandline);
+            }
+            string LogFile = RunUAT(CommandUtils.CmdEnv, WorkingCommandline);
+            SaveRecordOfSuccessAndAddToBuildProducts(CommandUtils.ReadAllText(LogFile));
+        }
+        public override string ECAgentString()
+        {
+            string Result = base.ECAgentString();
+            if (ClientTargetPlatform != HostPlatform)
+            {
+                Result = MergeSpaceStrings(Result, ClientTargetPlatform.ToString());
+            }
+            return Result;
+        }
+    }
+
     public class NonUnityTestNode : TestNode
     {
         public NonUnityTestNode(UnrealTargetPlatform InHostPlatform)
@@ -4345,7 +4452,7 @@ public class GUBP : BuildCommand
             }
 
             var NonCodeProjectNames = new Dictionary<string, List<UnrealTargetPlatform>>();
-            var NonCodeFormalBuilds = new Dictionary<string, List<KeyValuePair<UnrealTargetPlatform, UnrealTargetConfiguration>>>();
+            var NonCodeFormalBuilds = new Dictionary<string, List<TargetRules.GUBPFormalBuild>>();
             {
                 var Target = Branch.BaseEngineProject.Properties.Targets[TargetRules.TargetType.Editor];
 
@@ -4366,7 +4473,7 @@ public class GUBP : BuildCommand
                     }
                 }
 
-                var TempNonCodeFormalBuilds = Target.Rules.GUBP_NonCodeFormalBuilds_BaseEditorTypeOnly();
+                var TempNonCodeFormalBuilds = Target.Rules.GUBP_GetNonCodeFormalBuilds_BaseEditorTypeOnly();
                 var HostMonos = GetMonolithicPlatformsForUProject(HostPlatform, Branch.BaseEngineProject, true);
 
                 foreach (var Codeless in TempNonCodeFormalBuilds)
@@ -4374,10 +4481,10 @@ public class GUBP : BuildCommand
                     if (NonCodeProjectNames.ContainsKey(Codeless.Key))
                     {
                         var PlatList = Codeless.Value;
-                        var NewPlatList = new List<KeyValuePair<UnrealTargetPlatform, UnrealTargetConfiguration>>();
+                        var NewPlatList = new List<TargetRules.GUBPFormalBuild>();
                         foreach (var PlatPair in PlatList)
                         {
-                            if (HostMonos.Contains(PlatPair.Key))
+                            if (HostMonos.Contains(PlatPair.TargetPlatform))
                             {
                                 NewPlatList.Add(PlatPair);
                             }
@@ -4389,7 +4496,7 @@ public class GUBP : BuildCommand
                     }
                     else
                     {
-                        Log(System.Diagnostics.TraceEventType.Information, "{0} was listed as a codeless formal build GUBP_NonCodeFormalBuilds_BaseEditorTypeOnly, however it does not exist in this branch.", Codeless.Key);
+                        Log(System.Diagnostics.TraceEventType.Information, "{0} was listed as a codeless formal build GUBP_GetNonCodeFormalBuilds_BaseEditorTypeOnly, however it does not exist in this branch.", Codeless.Key);
                     }
                 }
             }
@@ -4562,9 +4669,9 @@ public class GUBP : BuildCommand
                                             var PlatList = NonCodeFormalBuilds[NonCodeProject.GameName];
                                             foreach (var PlatPair in PlatList)
                                             {
-                                                if (PlatPair.Key == Plat)
+                                                if (PlatPair.TargetPlatform == Plat)
                                                 {
-                                                    var NodeName = AddNode(new FormalBuildNode(this, NonCodeProject, HostPlatform, new List<UnrealTargetPlatform>() { Plat }, new List<UnrealTargetConfiguration>() {PlatPair.Value}));
+                                                    var NodeName = AddNode(new FormalBuildNode(this, NonCodeProject, HostPlatform, new List<UnrealTargetPlatform>() { Plat }, new List<UnrealTargetConfiguration>() {PlatPair.TargetConfig}));
                                                     // we don't want this delayed
                                                     // this would normally wait for the testing phase, we just want to build it right away
                                                     RemovePseudodependencyFromNode(
@@ -4575,7 +4682,7 @@ public class GUBP : BuildCommand
                                                         CookNode.StaticGetFullName(HostPlatform, Branch.BaseEngineProject, CookedPlatform));
 
                                                     string BuildAgentSharingGroup = NonCodeProject.GameName + "_MakeFormalBuild_" + Plat.ToString() + HostPlatformNode.StaticGetHostPlatformSuffix(HostPlatform);
-                                                    if (Plat == UnrealTargetPlatform.IOS) // Mac-IOS trashes build products, so we need to use different agents
+                                                    if (Plat == UnrealTargetPlatform.IOS || Plat == UnrealTargetPlatform.Android) // These trash build products, so we need to use different agents
                                                     {
                                                         BuildAgentSharingGroup = "";
                                                     }
@@ -4583,6 +4690,10 @@ public class GUBP : BuildCommand
                                                     GUBPNodes[GamePlatformCookedAndCompiledNode.StaticGetFullName(HostPlatform, NonCodeProject, Plat)].AgentSharingGroup = BuildAgentSharingGroup;
                                                     GUBPNodes[NodeName].AgentSharingGroup = BuildAgentSharingGroup;
 
+                                                    if (PlatPair.bTest)
+                                                    {
+                                                        AddNode(new FormalBuildTestNode(this, NonCodeProject, HostPlatform, Plat, PlatPair.TargetConfig));
+                                                    }
                                                 }
                                             }
                                     }
@@ -4721,23 +4832,22 @@ public class GUBP : BuildCommand
                                 {
                                     AddNode(new GamePlatformCookedAndCompiledNode(this, HostPlatform, CodeProj, Plat, true));
                                 }
-#if true
-                                var FormalBuildConfigs = Target.Rules.GUBP_GetConfigsForFormalBuilds_MonolithicOnly(HostPlatform, Plat);
+                                var FormalBuildConfigs = Target.Rules.GUBP_GetConfigsForFormalBuilds_MonolithicOnly(HostPlatform);
 
                                 foreach (var Config in FormalBuildConfigs)
                                 {
                                     string FormalNodeName = null;
                                     if (Kind == TargetRules.TargetType.Client)
                                     {
-                                        FormalNodeName = AddNode(new FormalBuildNode(this, CodeProj, HostPlatform, InClientTargetPlatforms: new List<UnrealTargetPlatform>() { Plat }, InClientConfigs: new List<UnrealTargetConfiguration>() { Config }, InClientNotGame : true));
+                                        FormalNodeName = AddNode(new FormalBuildNode(this, CodeProj, HostPlatform, InClientTargetPlatforms: new List<UnrealTargetPlatform>() { Plat }, InClientConfigs: new List<UnrealTargetConfiguration>() { Config.TargetConfig }, InClientNotGame : true));
                                     }
                                     else if (Kind == TargetRules.TargetType.Server)
                                     {
-                                        FormalNodeName = AddNode(new FormalBuildNode(this, CodeProj, HostPlatform, InServerTargetPlatforms : new List<UnrealTargetPlatform>() { Plat }, InServerConfigs : new List<UnrealTargetConfiguration>() { Config }));
+                                        FormalNodeName = AddNode(new FormalBuildNode(this, CodeProj, HostPlatform, InServerTargetPlatforms: new List<UnrealTargetPlatform>() { Plat }, InServerConfigs: new List<UnrealTargetConfiguration>() { Config.TargetConfig }));
                                     }
                                     else if (Kind == TargetRules.TargetType.Game)
                                     {
-                                        FormalNodeName = AddNode(new FormalBuildNode(this, CodeProj, HostPlatform, InClientTargetPlatforms: new List<UnrealTargetPlatform>() { Plat }, InClientConfigs: new List<UnrealTargetConfiguration>() { Config }));
+                                        FormalNodeName = AddNode(new FormalBuildNode(this, CodeProj, HostPlatform, InClientTargetPlatforms: new List<UnrealTargetPlatform>() { Plat }, InClientConfigs: new List<UnrealTargetConfiguration>() { Config.TargetConfig }));
                                     }
                                     // we don't want this delayed
                                     // this would normally wait for the testing phase, we just want to build it right away
@@ -4749,7 +4859,7 @@ public class GUBP : BuildCommand
                                         CookNode.StaticGetFullName(HostPlatform, Branch.BaseEngineProject, CookedPlatform));
 
                                     string BuildAgentSharingGroup = CodeProj.GameName + "_MakeFormalBuild_" + Plat.ToString() + HostPlatformNode.StaticGetHostPlatformSuffix(HostPlatform);
-                                    if (Plat == UnrealTargetPlatform.IOS) // Mac-IOS trashes build products, so we need to use different agents
+                                    if (Plat == UnrealTargetPlatform.IOS || Plat == UnrealTargetPlatform.Android) // These trash build products, so we need to use different agents
                                     {
                                         BuildAgentSharingGroup = "";
                                     }
@@ -4757,8 +4867,11 @@ public class GUBP : BuildCommand
                                     GUBPNodes[GamePlatformCookedAndCompiledNode.StaticGetFullName(HostPlatform, CodeProj, Plat)].AgentSharingGroup = BuildAgentSharingGroup;
                                     GUBPNodes[FormalNodeName].AgentSharingGroup = BuildAgentSharingGroup;
 
+                                    if (Config.bTest)
+                                    {
+                                        AddNode(new FormalBuildTestNode(this, CodeProj, HostPlatform, Plat, Config.TargetConfig));
+                                    }
                                 }
-#endif
                                 if (bAutomatedTesting)
                                 {
 if (HostPlatform == UnrealTargetPlatform.Mac) continue; //temp hack till mac automated testing works
@@ -5061,6 +5174,23 @@ if (HostPlatform == UnrealTargetPlatform.Mac) continue; //temp hack till mac aut
                         }
                     }
                     if (GameName != null)
+                    {
+                        foreach (var GameProj in Branch.NonCodeProjects)
+                        {
+                            if (GameProj.GameName.Equals(GameName, StringComparison.InvariantCultureIgnoreCase))
+                            {
+                                foreach (var Node in GUBPNodes)
+                                {
+                                    if (Node.Value.GameNameIfAnyForTempStorage() == GameProj.GameName)
+                                    {
+                                        NodesToDo.Add(Node.Key);
+                                    }
+                                }
+                                GameName = null;
+                            }
+                        }
+                    }
+                    if (GameName != null)
                     {
                         throw new AutomationException("Could not find game named {0}", GameName);
                     }
diff --git a/Engine/Source/Programs/UnrealBuildTool/System/RulesCompiler.cs b/Engine/Source/Programs/UnrealBuildTool/System/RulesCompiler.cs
--- a/Engine/Source/Programs/UnrealBuildTool/System/RulesCompiler.cs
+++ b/Engine/Source/Programs/UnrealBuildTool/System/RulesCompiler.cs
@@ -634,11 +634,34 @@ namespace UnrealBuildTool
         /// Return a list of configs for target platforms for formal builds
         /// </summary>
         /// <returns>a list of configs for a target platforms for the monolithic</returns>        
+        [Obsolete]
         public virtual List<UnrealTargetConfiguration> GUBP_GetConfigsForFormalBuilds_MonolithicOnly(UnrealTargetPlatform HostPlatform, UnrealTargetPlatform Platform)
         {
             return new List<UnrealTargetConfiguration>();
         }
 
+        public class GUBPFormalBuild
+        {
+            public UnrealTargetPlatform TargetPlatform = UnrealTargetPlatform.Unknown;
+            public UnrealTargetConfiguration TargetConfig = UnrealTargetConfiguration.Unknown;
+            public bool bTest = false;
+            public GUBPFormalBuild(UnrealTargetPlatform InTargetPlatform, UnrealTargetConfiguration InTargetConfig, bool bInTest = false)
+            {
+                TargetPlatform = InTargetPlatform;
+                TargetConfig = InTargetConfig;
+                bTest = bInTest;
+            }
+        }
+        /// <summary>
+        /// Return a list of formal builds
+        /// </summary>
+        /// <returns>a list of formal builds</returns>        
+        public virtual List<GUBPFormalBuild> GUBP_GetConfigsForFormalBuilds_MonolithicOnly(UnrealTargetPlatform HostPlatform)
+        {
+            return new List<GUBPFormalBuild>();
+        }
+
+
         /// <summary>
         /// Return true if this target should be included in a promotion and indicate shared or not
         /// </summary>
@@ -673,10 +696,20 @@ namespace UnrealBuildTool
         /// Return a list of the non-code projects to make formal builds for
         /// </summary>
         /// <returns>a list of the non-code projects to build cook and test</returns>
+        [Obsolete]
         public virtual Dictionary<string, List<KeyValuePair<UnrealTargetPlatform, UnrealTargetConfiguration>>> GUBP_NonCodeFormalBuilds_BaseEditorTypeOnly()
         {
             return new Dictionary<string, List<KeyValuePair<UnrealTargetPlatform, UnrealTargetConfiguration>>>();
         }
+        /// <summary>
+        /// Return a list of the non-code projects to make formal builds for
+        /// </summary>
+        /// <returns>a list of the non-code projects to build cook and test</returns>
+        public virtual Dictionary<string, List<GUBPFormalBuild>> GUBP_GetNonCodeFormalBuilds_BaseEditorTypeOnly()
+        {
+            return new Dictionary<string, List<GUBPFormalBuild>>();
+        }
+
         /// <summary>
         /// Return a list of "test name", "UAT command" pairs for testing the editor
         /// </summary>
diff --git a/Engine/Source/UE4Editor.Target.cs b/Engine/Source/UE4Editor.Target.cs
--- a/Engine/Source/UE4Editor.Target.cs
+++ b/Engine/Source/UE4Editor.Target.cs
@@ -122,16 +122,16 @@ public class UE4EditorTarget : TargetRules
 
         return NonCodeProjectNames;
     }
-    public override Dictionary<string, List<KeyValuePair<UnrealTargetPlatform, UnrealTargetConfiguration>>> GUBP_NonCodeFormalBuilds_BaseEditorTypeOnly()
+    public override Dictionary<string, List<GUBPFormalBuild>> GUBP_GetNonCodeFormalBuilds_BaseEditorTypeOnly()
     {
-        var NonCodeProjectNames = new Dictionary<string, List<KeyValuePair<UnrealTargetPlatform, UnrealTargetConfiguration>>>();
-        NonCodeProjectNames.Add("TappyChicken", 
-            new List<KeyValuePair<UnrealTargetPlatform, UnrealTargetConfiguration>>
+        var NonCodeProjectNames = new Dictionary<string, List<GUBPFormalBuild>>();
+        NonCodeProjectNames.Add("TappyChicken",
+            new List<GUBPFormalBuild>
             {
-                    new KeyValuePair<UnrealTargetPlatform, UnrealTargetConfiguration>(UnrealTargetPlatform.IOS, UnrealTargetConfiguration.Shipping),
-                    new KeyValuePair<UnrealTargetPlatform, UnrealTargetConfiguration>(UnrealTargetPlatform.IOS, UnrealTargetConfiguration.Test),
-                    new KeyValuePair<UnrealTargetPlatform, UnrealTargetConfiguration>(UnrealTargetPlatform.Android, UnrealTargetConfiguration.Shipping),
-                    new KeyValuePair<UnrealTargetPlatform, UnrealTargetConfiguration>(UnrealTargetPlatform.Android, UnrealTargetConfiguration.Test)
+                    new GUBPFormalBuild(UnrealTargetPlatform.IOS, UnrealTargetConfiguration.Shipping),
+                    new GUBPFormalBuild(UnrealTargetPlatform.IOS, UnrealTargetConfiguration.Test),
+                    new GUBPFormalBuild(UnrealTargetPlatform.Android, UnrealTargetConfiguration.Shipping /*, true*/),
+                    new GUBPFormalBuild(UnrealTargetPlatform.Android, UnrealTargetConfiguration.Test)
             }
         );
         return NonCodeProjectNames;
