Fix paths with spaces using git PR script.

diff --git a/Engine/Source/Programs/AutomationTool/Scripts/Tests.Automation.cs b/Engine/Source/Programs/AutomationTool/Scripts/Tests.Automation.cs
--- a/Engine/Source/Programs/AutomationTool/Scripts/Tests.Automation.cs
+++ b/Engine/Source/Programs/AutomationTool/Scripts/Tests.Automation.cs
@@ -276,7 +276,7 @@ class GitPullRequest : BuildCommand
         P4Sub.Sync(String.Format("-f -k -q {0}/...@{1}", Depot, CL));
 
         var Change = P4Sub.CreateChange(null, String.Format("GitHub pull request #{0}", PRNum));
-        P4Sub.ReconcileNoDeletes(Change, CombinePaths(Dir, bDoingUT ? "UnrealTournament" : "Engine", "..."));
+        P4Sub.ReconcileNoDeletes(Change, CommandUtils.MakePathSafeToUseWithCommandLine(CombinePaths(Dir, bDoingUT ? "UnrealTournament" : "Engine", "...")));
         P4Sub.Shelve(Change);
         P4Sub.Revert(Change, "-k //...");
     }
