Add UnrealVersionSelector into the promotable

diff --git a/Engine/Source/Programs/AutomationTool/Scripts/GUBP.Automation.cs b/Engine/Source/Programs/AutomationTool/Scripts/GUBP.Automation.cs
--- a/Engine/Source/Programs/AutomationTool/Scripts/GUBP.Automation.cs
+++ b/Engine/Source/Programs/AutomationTool/Scripts/GUBP.Automation.cs
@@ -1452,12 +1452,21 @@ public class GUBP : BuildCommand
 
     public class EditorAndToolsNode : HostPlatformAggregateNode
     {
-        public EditorAndToolsNode(UnrealTargetPlatform InHostPlatform)
+        public EditorAndToolsNode(GUBP bp, UnrealTargetPlatform InHostPlatform)
             : base(InHostPlatform)
         {
             AddDependency(RootEditorNode.StaticGetFullName(HostPlatform));
             AddDependency(ToolsNode.StaticGetFullName(HostPlatform));
             AddDependency(InternalToolsNode.StaticGetFullName(HostPlatform));
+
+			if(InHostPlatform == UnrealTargetPlatform.Win64)
+			{
+				var VersionSelector = bp.Branch.FindProgram("UnrealVersionSelector"); 
+				if (VersionSelector.Rules != null) 
+				{ 
+					AddDependency(SingleInternalToolsNode.StaticGetFullName(HostPlatform, VersionSelector));
+				}
+			}
         }
         public static string StaticGetFullName(UnrealTargetPlatform InHostPlatform)
         {
@@ -4422,7 +4431,7 @@ public class GUBP : BuildCommand
                 }
             }
 
-            AddNode(new EditorAndToolsNode(HostPlatform));
+            AddNode(new EditorAndToolsNode(this, HostPlatform));
 
             if (bOrthogonalizeEditorPlatforms)
             {
