Added new automation tool to support for deploying prebuilt apks to several devices at once.

diff --git a/Engine/Source/Programs/AutomationTool/Android/AndroidPlatform.Automation.cs b/Engine/Source/Programs/AutomationTool/Android/AndroidPlatform.Automation.cs
--- a/Engine/Source/Programs/AutomationTool/Android/AndroidPlatform.Automation.cs
+++ b/Engine/Source/Programs/AutomationTool/Android/AndroidPlatform.Automation.cs
@@ -6,6 +6,7 @@ using System.Text;
 using System.IO;
 using System.Diagnostics;
 using System.Net.NetworkInformation;
+using System.Threading;
 using AutomationTool;
 using UnrealBuildTool;
 
@@ -29,9 +30,19 @@ public class AndroidPlatform : Platform
 
 	private static string GetFinalApkName(ProjectParams Params, string DecoratedExeName, bool bRenameUE4Game)
 	{
-		string ProjectDir = Path.GetDirectoryName(Path.GetFullPath(Params.RawProjectPath));
+        string ProjectDir = Path.Combine( Path.GetDirectoryName(Path.GetFullPath(Params.RawProjectPath)),  "Binaries/Android" );
+        
+        if (Params.Prebuilt)
+        {
+            ProjectDir = Path.Combine( Params.BaseStageDirectory, "Android" );
+        }
+
+        // Apk's go to project location, not necessarily where the .so is (content only packages need to output to their directory)
+        string ApkName = Path.Combine(ProjectDir, DecoratedExeName) + GetArchitecture(Params) + ".apk";
+
+		/*string ProjectDir = Path.GetDirectoryName(Path.GetFullPath(Params.RawProjectPath));
 		// Apk's go to project location, not necessarily where the .so is (content only packages need to output to their directory)
-		string ApkName = Path.Combine(ProjectDir, "Binaries/Android", DecoratedExeName) + GetArchitecture(Params) + ".apk";
+		string ApkName = Path.Combine(ProjectDir, "Binaries/Android", DecoratedExeName) + GetArchitecture(Params) + ".apk";*/
 		
 		// if the source binary was UE4Game, handle using it or switching to project name
 		if (Path.GetFileNameWithoutExtension(Params.ProjectGameExeFilename) == "UE4Game")
@@ -103,8 +114,11 @@ public class AndroidPlatform : Platform
 		// packaging just takes a pak file and makes it the .obb
         UEBuildConfiguration.bOBBinAPK = Params.OBBinAPK; // Make sure this setting is sync'd pre-build
         var Deploy = UEBuildDeploy.GetBuildDeploy(UnrealTargetPlatform.Android);
-		string CookFlavor = SC.CookPlatform.IndexOf("_") > 0 ? SC.CookPlatform.Substring(SC.CookPlatform.IndexOf("_")) : "";
-		Deploy.PrepForUATPackageOrDeploy(Params.ShortProjectName, SC.ProjectRoot, SOName, SC.LocalRoot + "/Engine", Params.Distribution, CookFlavor);
+        if (!Params.Prebuilt)
+        {
+            string CookFlavor = SC.CookPlatform.IndexOf("_") > 0 ? SC.CookPlatform.Substring(SC.CookPlatform.IndexOf("_")) : "";
+            Deploy.PrepForUATPackageOrDeploy(Params.ShortProjectName, SC.ProjectRoot, SOName, SC.LocalRoot + "/Engine", Params.Distribution, CookFlavor);
+        }
         
 		// first, look for a .pak file in the staged directory
 		string[] PakFiles = Directory.GetFiles(SC.StageDirectory, "*.pak", SearchOption.AllDirectories);
@@ -194,6 +208,31 @@ public class AndroidPlatform : Platform
 		return Environment.ExpandEnvironmentVariables("/c %ANDROID_HOME%/platform-tools/adb.exe" + SerialNumber + " ");
 	}
 
+    public override void GetConnectedDevices(ProjectParams Params, out List<string> Devices)
+    {
+        Devices = new List<string>();
+        string AdbCommand = GetAdbCommand(Params);
+        ProcessResult Result = Run(CmdEnv.CmdExe, AdbCommand + "devices");
+
+        if (Result.Output.Length > 0)
+        {
+            string[] LogLines = Result.Output.Split(new char [] { '\n', '\r', '\t' });
+            for ( int i = 1; i < LogLines.Length; ++i )
+            {
+                if (LogLines[i].Length == 0)
+                {
+                    continue;
+                }
+                if (LogLines[i] == "device")
+                {
+                    continue;
+                }
+
+                Devices.Add("@"+LogLines[i]);
+            }
+        }
+    }
+
 	public override void Deploy(ProjectParams Params, DeploymentContext SC)
 	{
 		string SOName = GetFinalSOName(Params, SC.StageExecutables[0]);
@@ -202,8 +241,11 @@ public class AndroidPlatform : Platform
 
 		// make sure APK is up to date (this is fast if so)
 		var Deploy = UEBuildDeploy.GetBuildDeploy(UnrealTargetPlatform.Android);
-		string CookFlavor = SC.CookPlatform.IndexOf("_") > 0 ? SC.CookPlatform.Substring(SC.CookPlatform.IndexOf("_")) : "";
-		Deploy.PrepForUATPackageOrDeploy(Params.ShortProjectName, SC.ProjectRoot, SOName, SC.LocalRoot + "/Engine", Params.Distribution, CookFlavor);
+        if (!Params.Prebuilt)
+        {
+            string CookFlavor = SC.CookPlatform.IndexOf("_") > 0 ? SC.CookPlatform.Substring(SC.CookPlatform.IndexOf("_")) : "";
+            Deploy.PrepForUATPackageOrDeploy(Params.ShortProjectName, SC.ProjectRoot, SOName, SC.LocalRoot + "/Engine", Params.Distribution, CookFlavor);
+        }
 
 		// now we can use the apk to get more info
 		string DeviceObbName = GetDeviceObbName(ApkName);
@@ -216,6 +258,20 @@ public class AndroidPlatform : Platform
 		string InstallCommandline = AdbCommand + "install \"" + ApkName + "\"";
 		RunAndLog(CmdEnv, CmdEnv.CmdExe, InstallCommandline);
 
+        if (Params.Prebuilt)
+        {
+            // deploy the obb if there is one
+            string ObbPath = Path.Combine(SC.StageDirectory, GetFinalObbName(ApkName));
+            if (File.Exists(ObbPath))
+            {
+                // cache some strings
+                string BaseCommandline = AdbCommand + "push";
+
+                string Commandline = string.Format("{0} \"{1}\" \"{2}\"", BaseCommandline, ObbPath, DeviceObbName );
+                Run(CmdEnv.CmdExe, Commandline);
+            }
+        }
+
 		// update the ue4commandline.txt
 		// update and deploy ue4commandline.txt
 		// always delete the existing commandline text file, so it doesn't reuse an old one
@@ -224,7 +280,7 @@ public class AndroidPlatform : Platform
 
 
 		// copy files to device if we were staging
-		if (SC.Stage)
+		if (SC.Stage && !Params.SkipStage)
 		{
 			// cache some strings
 			string BaseCommandline = AdbCommand + "push";
@@ -282,12 +338,15 @@ public class AndroidPlatform : Platform
 
 	/** Internal usage for GetPackageName */
 	private static string PackageLine = null;
+    private static Mutex PackageInfoMutex = new Mutex();
 
 	/** Run an external exe (and capture the output), given the exe path and the commandline. */
 	private static string GetPackageInfo(string ApkName, bool bRetrieveVersionCode)
 	{
 		// we expect there to be one, so use the first one
 		string AaptPath = GetAaptPath();
+        
+        PackageInfoMutex.WaitOne();
 
 		var ExeInfo = new ProcessStartInfo(AaptPath, "dump badging \"" + ApkName + "\"");
 		ExeInfo.UseShellExecute = false;
@@ -299,6 +358,9 @@ public class AndroidPlatform : Platform
 			GameProcess.OutputDataReceived += ParsePackageName;
 			GameProcess.WaitForExit();
 		}
+
+        PackageInfoMutex.ReleaseMutex();
+
 		string ReturnValue = null;
 		if (PackageLine != null)
 		{
@@ -343,7 +405,7 @@ public class AndroidPlatform : Platform
 		return Path.Combine(Subdirs[0], "aapt.exe");
 	}
 
-	public override ProcessResult RunClient(ERunOptions ClientRunFlags, string ClientApp, string ClientCmdLine, ProjectParams Params)
+    public override ProcessResult RunClient(ERunOptions ClientRunFlags, string ClientApp, string ClientCmdLine, ProjectParams Params)
 	{
 		string ApkName = ClientApp + GetArchitecture(Params) + ".apk";
 		if (!File.Exists(ApkName))
@@ -363,9 +425,57 @@ public class AndroidPlatform : Platform
 		string AdbCommand = GetAdbCommand(Params);
 		string CommandLine = "shell am start -n " + PackageName + "/com.epicgames.ue4.GameActivity";
 
+        if (Params.Prebuilt)
+        {
+            // clear the log
+            Run(CmdEnv.CmdExe, AdbCommand + "logcat -c");
+        }
+        
+
 		// start the app on device!
 		ProcessResult ClientProcess = Run(CmdEnv.CmdExe, AdbCommand + CommandLine, null, ClientRunFlags);
-		 
+
+
+        if (Params.Prebuilt)
+        {
+            // save the output to the staging directory
+            string LogPath = Path.Combine(Params.BaseStageDirectory, "Android\\logs");
+            string LogFilename = Path.Combine(LogPath, "devicelog" + Params.Device + ".log");
+            string ServerLogFilename = Path.Combine( CmdEnv.LogFolder, "devicelog" + Params.Device + ".log");
+
+            Directory.CreateDirectory(LogPath);
+
+            // check if the game is still running 
+            // time out if it takes to long
+            DateTime StartTime = DateTime.Now;
+            int TimeOutSeconds = Params.RunTimeoutSeconds;
+
+            while (true)
+            {
+                ProcessResult ProcessesResult = Run(CmdEnv.CmdExe, AdbCommand + "shell ps", null, ERunOptions.SpewIsVerbose);
+
+                string RunningProcessList = ProcessesResult.Output;
+                if (!RunningProcessList.Contains(PackageName))
+                {
+                    break;
+                }
+                Thread.Sleep(10);
+
+                TimeSpan DeltaRunTime = DateTime.Now - StartTime;
+                if ((DeltaRunTime.TotalSeconds > TimeOutSeconds) && (TimeOutSeconds != 0))
+                {
+                    Log("Device: " + Params.Device + " timed out while waiting for run to finish");
+                    break;
+                }
+            }
+            
+            ProcessResult LogFileProcess = Run(CmdEnv.CmdExe, AdbCommand + "logcat -d");
+            
+            File.WriteAllText(LogFilename, LogFileProcess.Output);
+            File.WriteAllText(ServerLogFilename, LogFileProcess.Output);
+        }
+
+
 		return ClientProcess;
 	}
 
diff --git a/Engine/Source/Programs/AutomationTool/Platform.cs b/Engine/Source/Programs/AutomationTool/Platform.cs
--- a/Engine/Source/Programs/AutomationTool/Platform.cs
+++ b/Engine/Source/Programs/AutomationTool/Platform.cs
@@ -99,6 +99,19 @@ namespace AutomationTool
 		{
 		}
 
+        /// <summary>
+        /// Get all connected device names for this platform
+        /// </summary>
+        /// <param name="Params"></param>
+        /// <param name="SC"></param>
+        public virtual void GetConnectedDevices(ProjectParams Params, out List<string> Devices)
+        {
+            Devices = null;
+            Log(System.Diagnostics.TraceEventType.Warning, "{0} does not implement GetConnectedDevices", PlatformType);
+        }
+
+
+
 		/// <summary>
 		/// Deploy the application on the current platform
 		/// </summary>
diff --git a/Engine/Source/Programs/AutomationTool/ProjectParams.cs b/Engine/Source/Programs/AutomationTool/ProjectParams.cs
--- a/Engine/Source/Programs/AutomationTool/ProjectParams.cs
+++ b/Engine/Source/Programs/AutomationTool/ProjectParams.cs
@@ -192,6 +192,8 @@ namespace AutomationTool
 			this.CookFlavor = InParams.CookFlavor;
 			this.SkipCook = InParams.SkipCook;
 			this.SkipCookOnTheFly = InParams.SkipCookOnTheFly;
+            this.Prebuilt = InParams.Prebuilt;
+            this.RunTimeoutSeconds = InParams.RunTimeoutSeconds;
 			this.Clean = InParams.Clean;
 			this.Pak = InParams.Pak;
 			this.SignPak = InParams.SignPak;
@@ -249,6 +251,8 @@ namespace AutomationTool
 			this.ArchiveDirectoryParam = InParams.ArchiveDirectoryParam;
 			this.Distribution = InParams.Distribution;
             this.OBBinAPK = InParams.OBBinAPK;
+            this.Prebuilt = InParams.Prebuilt;
+            this.RunTimeoutSeconds = InParams.RunTimeoutSeconds;
 		}
 
 		/// <summary>
@@ -323,7 +327,9 @@ namespace AutomationTool
 			string ArchiveDirectoryParam = null,
 			ParamList<string> ProgramTargets = null,
 			bool? Distribution = null,
-            bool? OBBinAPK = null
+            bool? OBBinAPK = null,
+            bool? Prebuilt = null,
+            int? RunTimeoutSeconds = null
 			)
 		{
 			//
@@ -407,7 +413,20 @@ namespace AutomationTool
 			this.ArchiveDirectoryParam = ParseParamValueIfNotSpecified(Command, ArchiveDirectoryParam, "archivedirectory", String.Empty);
 			this.Distribution = GetParamValueIfNotSpecified(Command, Distribution, this.Distribution, "distribution");
             this.OBBinAPK = GetParamValueIfNotSpecified(Command, OBBinAPK, this.OBBinAPK, "obbinapk");
-			this.NoDebugInfo = GetParamValueIfNotSpecified(Command, NoDebugInfo, this.NoDebugInfo, "nodebuginfo");
+            this.Prebuilt = GetParamValueIfNotSpecified(Command, Prebuilt, this.Prebuilt, "prebuilt");
+            if (this.Prebuilt)
+            {
+                this.SkipCook = true;
+                this.SkipPak = true;
+                this.SkipStage = true;
+                this.Cook = true;
+                this.Stage = true;
+                this.Pak = true;
+                this.Deploy = true;
+                this.Run = true;
+                //this.StageDirectoryParam = this.PrebuiltDir;
+            }
+            this.NoDebugInfo = GetParamValueIfNotSpecified(Command, NoDebugInfo, this.NoDebugInfo, "nodebuginfo");
 			this.NoCleanStage = GetParamValueIfNotSpecified(Command, NoCleanStage, this.NoCleanStage, "nocleanstage");
 			this.MapToRun = ParseParamValueIfNotSpecified(Command, MapToRun, "map", String.Empty);
 			this.AdditionalServerMapParams = ParseParamValueIfNotSpecified(Command, AdditionalServerMapParams, "AdditionalServerMapParams", String.Empty);
@@ -489,6 +508,14 @@ namespace AutomationTool
             {
                 this.CrashIndex = Command.ParseParamInt("CrashIndex");
             }
+            if (RunTimeoutSeconds.HasValue)
+            {
+                this.RunTimeoutSeconds = RunTimeoutSeconds.Value;
+            }
+            else if (Command != null)
+            {
+                this.RunTimeoutSeconds = Command.ParseParamInt("runtimeoutseconds");
+            }
 
 			AutodetectSettings(false);
 			ValidateAndLog();
@@ -990,6 +1017,13 @@ namespace AutomationTool
         [Help("obbinapk", "package with OBB data in APK assets directory")]
         public bool OBBinAPK {get; set; }
 
+        [Help("Prebuilt", "this is a prebuilt cooked and packaged build")]
+        public bool Prebuilt { get; private set; }
+
+        [Help("RunTimeoutSeconds", "timeout to wait after we lunch the game")]
+        public int RunTimeoutSeconds;
+
+
 		#endregion
 
 		#region Deploy
@@ -1195,7 +1229,7 @@ namespace AutomationTool
 				}
 			}
 
-            if (ClientCookedTargetsList == null && !NoClient && (Cook || CookOnTheFly))
+            if (ClientCookedTargetsList == null && !NoClient && (Cook || CookOnTheFly || Prebuilt))
 			{
                 if (String.IsNullOrEmpty(GameTarget))
 				{
@@ -1548,6 +1582,7 @@ namespace AutomationTool
                 CommandUtils.Log("EditorTest={0}", EditorTest);
                 CommandUtils.Log("RunAutomationTests={0}", RunAutomationTests); 
                 CommandUtils.Log("RunAutomationTest={0}", RunAutomationTest);
+                CommandUtils.Log("RunTimeoutSeconds={0}", RunTimeoutSeconds);
                 CommandUtils.Log("CrashIndex={0}", CrashIndex);
 				CommandUtils.Log("ProgramTargets={0}", ProgramTargets.ToString());
                 CommandUtils.Log("ProjectBinariesFolder={0}", ProjectBinariesFolder);
@@ -1556,6 +1591,7 @@ namespace AutomationTool
 				CommandUtils.Log("ProjectGameExePath={0}", ProjectGameExePath);
 				CommandUtils.Log("Distribution={0}", Distribution);
                 CommandUtils.Log("OBBinAPK={0}", OBBinAPK);
+                CommandUtils.Log("Prebuilt={0}", Prebuilt);
 				CommandUtils.Log("RawProjectPath={0}", RawProjectPath);
 				CommandUtils.Log("Rocket={0}", Rocket);
 				CommandUtils.Log("Run={0}", Run);
