Add EmailHint Param that allows us to direct all emails to a particular group for manually triggered builds

diff --git a/Engine/Source/Programs/AutomationTool/Scripts/GUBP.Automation.cs b/Engine/Source/Programs/AutomationTool/Scripts/GUBP.Automation.cs
--- a/Engine/Source/Programs/AutomationTool/Scripts/GUBP.Automation.cs
+++ b/Engine/Source/Programs/AutomationTool/Scripts/GUBP.Automation.cs
@@ -24,6 +24,7 @@ public class GUBP : BuildCommand
     public static bool bBuildRocket = false;
     public static bool bForceIncrementalCompile = false;
     public static string ECProject = null;
+    public string EmailHint;
 
     Dictionary<string, GUBPNode> GUBPNodes;
     Dictionary<string, bool> GUBPNodesCompleted;
@@ -133,7 +134,7 @@ public class GUBP : BuildCommand
     }
 
     public abstract class GUBPEmailHacker
-    {
+    {        
         public virtual List<string> AddEmails(GUBP bp, string Branch, string NodeName)
         {
             return new List<string>();
@@ -151,6 +152,7 @@ public class GUBP : BuildCommand
     private static List<GUBPEmailHacker> EmailHackers;
     private string HackEmails(string Emails, string Branch, string NodeName)
     {
+        EmailHint = ParseParamValue("EmailHint");
         if (EmailHackers == null)
         {
             EmailHackers = new List<GUBPEmailHacker>();
@@ -4123,8 +4125,7 @@ public class GUBP : BuildCommand
 
         bBuildRocket = ParseParam("BuildRocket");
         bForceIncrementalCompile = ParseParam("ForceIncrementalCompile");
-        bool bAutomatedTesting = ParseParam("AutomatedTesting");
-
+        bool bAutomatedTesting = ParseParam("AutomatedTesting");        
         StoreName = ParseParamValue("Store");
         string StoreSuffix = ParseParamValue("StoreSuffix", "");
         if (bBuildRocket)
