Fixing build

diff --git a/Engine/Source/Programs/AutomationTool/Scripts/GUBP.Automation.cs b/Engine/Source/Programs/AutomationTool/Scripts/GUBP.Automation.cs
--- a/Engine/Source/Programs/AutomationTool/Scripts/GUBP.Automation.cs
+++ b/Engine/Source/Programs/AutomationTool/Scripts/GUBP.Automation.cs
@@ -24,7 +24,7 @@ public class GUBP : BuildCommand
     public static bool bBuildRocket = false;
     public static bool bForceIncrementalCompile = false;
     public static string ECProject = null;
-    public string EmailHint;
+    public string EmailHint = "";
 
     Dictionary<string, GUBPNode> GUBPNodes;
     Dictionary<string, bool> GUBPNodesCompleted;
