Temp backing out email hint while I figure out why this is succeeding locally

diff --git a/Engine/Source/Programs/AutomationTool/Scripts/GUBP.Automation.cs b/Engine/Source/Programs/AutomationTool/Scripts/GUBP.Automation.cs
--- a/Engine/Source/Programs/AutomationTool/Scripts/GUBP.Automation.cs
+++ b/Engine/Source/Programs/AutomationTool/Scripts/GUBP.Automation.cs
@@ -24,7 +24,7 @@ public class GUBP : BuildCommand
     public static bool bBuildRocket = false;
     public static bool bForceIncrementalCompile = false;
     public static string ECProject = null;
-    public string EmailHint = "";
+    //public string EmailHint = "";
 
     Dictionary<string, GUBPNode> GUBPNodes;
     Dictionary<string, bool> GUBPNodesCompleted;
@@ -152,7 +152,7 @@ public class GUBP : BuildCommand
     private static List<GUBPEmailHacker> EmailHackers;
     private string HackEmails(string Emails, string Branch, string NodeName)
     {
-        EmailHint = ParseParamValue("EmailHint");
+        //EmailHint = ParseParamValue("EmailHint");
         if (EmailHackers == null)
         {
             EmailHackers = new List<GUBPEmailHacker>();
