Adding a working emailhint param that shouldn't break the build

diff --git a/Engine/Source/Programs/AutomationTool/Scripts/GUBP.Automation.cs b/Engine/Source/Programs/AutomationTool/Scripts/GUBP.Automation.cs
--- a/Engine/Source/Programs/AutomationTool/Scripts/GUBP.Automation.cs
+++ b/Engine/Source/Programs/AutomationTool/Scripts/GUBP.Automation.cs
@@ -24,7 +24,7 @@ public class GUBP : BuildCommand
     public static bool bBuildRocket = false;
     public static bool bForceIncrementalCompile = false;
     public static string ECProject = null;
-    //public string EmailHint = "";
+    public string EmailHint;
 
     Dictionary<string, GUBPNode> GUBPNodes;
     Dictionary<string, bool> GUBPNodesCompleted;
@@ -151,8 +151,12 @@ public class GUBP : BuildCommand
 
     private static List<GUBPEmailHacker> EmailHackers;
     private string HackEmails(string Emails, string Branch, string NodeName)
-    {
-        //EmailHint = ParseParamValue("EmailHint");
+    {        
+        EmailHint = ParseParamValue("EmailHint");
+        if(EmailHint == null)
+        {
+            EmailHint = "";
+        }
         if (EmailHackers == null)
         {
             EmailHackers = new List<GUBPEmailHacker>();
