Remove Mac penalty from CIS quantum

diff --git a/Engine/Source/Programs/AutomationTool/Scripts/GUBP.Automation.cs b/Engine/Source/Programs/AutomationTool/Scripts/GUBP.Automation.cs
--- a/Engine/Source/Programs/AutomationTool/Scripts/GUBP.Automation.cs
+++ b/Engine/Source/Programs/AutomationTool/Scripts/GUBP.Automation.cs
@@ -491,11 +491,7 @@ public class GUBP : BuildCommand
         }
         public override int CISFrequencyQuantumShift(GUBP bp)
         {
-            int Result = base.CISFrequencyQuantumShift(bp);
-            if (HostPlatform == UnrealTargetPlatform.Mac)
-            {
-                Result += 1; // mac only runs half as much
-            }
+            int Result = base.CISFrequencyQuantumShift(bp);            
             return Result;
         }
     }
@@ -2049,8 +2045,7 @@ public class GUBP : BuildCommand
         }
         public override int CISFrequencyQuantumShift(GUBP bp)
         {
-            return base.CISFrequencyQuantumShift(bp) + 4 + (bIsMassive ? 1 : 0) 
-                - (HostPlatform == UnrealTargetPlatform.Mac ? 1 : 0); // we don't want the mac penalty here
+            return base.CISFrequencyQuantumShift(bp) + 4 + (bIsMassive ? 1 : 0);                
         }
         public override float Priority()
         {
