GUBP: RetrieveFromTempStorage will now not only add executable bit to Mac and iOS executables, but also fix read/write permissions for other users for all files.

diff --git a/Engine/Source/Programs/AutomationTool/BuildUtils.cs b/Engine/Source/Programs/AutomationTool/BuildUtils.cs
--- a/Engine/Source/Programs/AutomationTool/BuildUtils.cs
+++ b/Engine/Source/Programs/AutomationTool/BuildUtils.cs
@@ -80,19 +80,6 @@ namespace AutomationTool
 			MsBuild(Env, ProjectFile, CmdLine, LogName);
 		}
 
-        /// <summary>
-        /// Sets an executable bit for a Unix file.
-        /// <param name="Filename">Filename</param>
-        /// </summary>
-        public static void SetExecutableBit(string Filename)
-        {
-            var Result = CommandUtils.Run("sh", string.Format("-c 'chmod +x \"{0}\"'", Filename));
-            if (Result.ExitCode != 0)
-            {
-                throw new AutomationException(String.Format("Failed to chmod \"{0}\"", Filename));
-            }
-        }
-
         /// <summary>
         /// returns true if this is a mac executable using some awful conventions
         /// <param name="Filename">Filename</param>
@@ -107,6 +94,20 @@ namespace AutomationTool
                     CommandUtils.CombinePaths(Filename).ToLower().Contains(CommandUtils.CombinePaths("Binaries", "IOS").ToLower())
                 ) && (Path.GetExtension(Filename) == "" || Path.GetExtension(Filename) == "."));
         }
+
+		/// <summary>
+		/// Sets an executable bit for Unix executables and adds read permission for all users plus write permission for owner.
+		/// <param name="Filename">Filename</param>
+		/// </summary>
+		public static void FixUnixFilePermissions(string Filename)
+		{
+			string Permissions = IsProbablyAMacOrIOSExe(Filename) ? "0755" : "0644";
+			var Result = CommandUtils.Run("sh", string.Format("-c 'chmod {0} \"{1}\"'", Permissions, Filename));
+			if (Result.ExitCode != 0)
+			{
+				throw new AutomationException(String.Format("Failed to chmod \"{0}\"", Filename));
+			}
+		}
 	}
 
 }
diff --git a/Engine/Source/Programs/AutomationTool/TempStorage.cs b/Engine/Source/Programs/AutomationTool/TempStorage.cs
--- a/Engine/Source/Programs/AutomationTool/TempStorage.cs
+++ b/Engine/Source/Programs/AutomationTool/TempStorage.cs
@@ -883,9 +883,9 @@ namespace AutomationTool
 
                     Robust_FileExists_NoExceptions(true, DestFile, "Could not copy to {0}");
 
-                    if (UnrealBuildTool.Utils.IsRunningOnMono && IsProbablyAMacOrIOSExe(DestFile))
+                    if (UnrealBuildTool.Utils.IsRunningOnMono)
                     {
-                        SetExecutableBit(DestFile);
+						FixUnixFilePermissions(DestFile);
                     }
 
                     FileInfo Info = new FileInfo(DestFile);
@@ -919,9 +919,9 @@ namespace AutomationTool
                 foreach (string DestFile in DestFiles)
                 {
                     Robust_FileExists_NoExceptions(true, DestFile, "Could not copy to {0}");
-                    if (UnrealBuildTool.Utils.IsRunningOnMono && IsProbablyAMacOrIOSExe(DestFile))
+                    if (UnrealBuildTool.Utils.IsRunningOnMono)
                     {
-                        SetExecutableBit(DestFile);
+						FixUnixFilePermissions(DestFile);
                     }
                     FileInfo Info = new FileInfo(DestFile);
                     NewDestFiles.Add(Info.FullName);
