build script, turn automated testing back on

diff --git a/Engine/Source/Programs/AutomationTool/Scripts/GUBP.Automation.cs b/Engine/Source/Programs/AutomationTool/Scripts/GUBP.Automation.cs
--- a/Engine/Source/Programs/AutomationTool/Scripts/GUBP.Automation.cs
+++ b/Engine/Source/Programs/AutomationTool/Scripts/GUBP.Automation.cs
@@ -2865,6 +2865,10 @@ public class GUBP : BuildCommand
         {
             return ECPriority;
         }
+        public override int CISFrequencyQuantumShift(GUBP bp)
+        {
+            return base.CISFrequencyQuantumShift(bp) + 5 /*++ (bIsMassive ? 1 : 0);*/ ;
+        }
         public override string ECAgentString()
         {
             string Result = base.ECAgentString();
@@ -4124,7 +4128,7 @@ public class GUBP : BuildCommand
 
         bBuildRocket = ParseParam("BuildRocket");
         bForceIncrementalCompile = ParseParam("ForceIncrementalCompile");
-        bool bAutomatedTesting = ParseParam("AutomatedTesting");        
+        bool bNoAutomatedTesting = ParseParam("NoAutomatedTesting");        
         StoreName = ParseParamValue("Store");
         string StoreSuffix = ParseParamValue("StoreSuffix", "");
         if (bBuildRocket)
@@ -4538,10 +4542,10 @@ public class GUBP : BuildCommand
                     var EditorTestNodes = new List<string>();
                     foreach (var Test in EditorTests)
                     {
-                        if (bAutomatedTesting)
+                        if (!bNoAutomatedTesting)
                         {
                             EditorTestNodes.Add(AddNode(new UATTestNode(this, HostPlatform, Branch.BaseEngineProject, Test.Key, Test.Value, AgentSharingGroup)));
-
+                        
                             foreach (var NonCodeProject in Branch.NonCodeProjects)
                             {
                                 if (!NonCodeProjectNames.ContainsKey(NonCodeProject.GameName))
@@ -4646,7 +4650,7 @@ public class GUBP : BuildCommand
                                 }
                                 var GameTests = Target.Rules.GUBP_GetGameTests_MonolithicOnly(HostPlatform, GetAltHostPlatform(HostPlatform), Plat);
                                 var RequiredPlatforms = new List<UnrealTargetPlatform> { Plat };
-                                if (bAutomatedTesting)
+                                if (!bNoAutomatedTesting)
                                 {
                                     var ThisMonoGameTestNodes = new List<string>();
                                     foreach (var Test in GameTests)
@@ -4707,7 +4711,7 @@ public class GUBP : BuildCommand
                                             }
                                     }
                                     }
-                                    if (bAutomatedTesting)
+                                    if (!bNoAutomatedTesting)
                                     {
                                     if (HostPlatform == UnrealTargetPlatform.Mac) continue; //temp hack till mac automated testing works
                                     var ThisMonoGameTestNodes = new List<string>();
@@ -4778,7 +4782,7 @@ public class GUBP : BuildCommand
                 }
 
                 AddNode(new EditorGameNode(this, HostPlatform, CodeProj));
-                if (bAutomatedTesting && HostPlatform != UnrealTargetPlatform.Mac) //temp hack till mac automated testing works
+                if (!bNoAutomatedTesting && HostPlatform != UnrealTargetPlatform.Mac) //temp hack till mac automated testing works
                 {
                     var EditorTests = CodeProj.Properties.Targets[TargetRules.TargetType.Editor].Rules.GUBP_GetEditorTests_EditorTypeOnly(HostPlatform);
                     var EditorTestNodes = new List<string>();
@@ -4881,7 +4885,7 @@ public class GUBP : BuildCommand
                                         AddNode(new FormalBuildTestNode(this, CodeProj, HostPlatform, Plat, Config.TargetConfig));
                                     }
                                 }
-                                if (bAutomatedTesting)
+                                if (!bNoAutomatedTesting)
                                 {
 if (HostPlatform == UnrealTargetPlatform.Mac) continue; //temp hack till mac automated testing works
                                 var GameTests = Target.Rules.GUBP_GetGameTests_MonolithicOnly(HostPlatform, GetAltHostPlatform(HostPlatform), Plat);
@@ -4903,7 +4907,7 @@ if (HostPlatform == UnrealTargetPlatform.Mac) continue; //temp hack till mac aut
                         }
                     }
                 }
-                if (bAutomatedTesting)
+                if (!bNoAutomatedTesting)
                 {
                 foreach (var ServerPlatform in ServerPlatforms)
                 {
