build script, make sure we didn't delete cis counter after syncing it to head

diff --git a/Engine/Source/Programs/AutomationTool/Scripts/GUBP.Automation.cs b/Engine/Source/Programs/AutomationTool/Scripts/GUBP.Automation.cs
--- a/Engine/Source/Programs/AutomationTool/Scripts/GUBP.Automation.cs
+++ b/Engine/Source/Programs/AutomationTool/Scripts/GUBP.Automation.cs
@@ -4243,6 +4243,27 @@ public class GUBP : BuildCommand
                     }
                 }
 				P4.Sync("-f " + P4IndexFileP4 + "#head");
+                if (!FileExists_NoExceptions(P4IndexFileLocal))
+                {
+                    Log("{0} doesn't exist, checking in a new one", P4IndexFileP4);
+                    WriteAllText(P4IndexFileLocal, "-1 0");
+                    int WorkingCL = -1;
+                    try
+                    {
+                        WorkingCL = P4.CreateChange(P4Env.Client, "Adding new CIS Counter");
+                        P4.Add(WorkingCL, P4IndexFileP4);
+                        int SubmittedCL;
+                        P4.Submit(WorkingCL, out SubmittedCL);
+                    }
+                    catch (Exception)
+                    {
+                        Log("Add of CIS counter failed, assuming it now exists.");
+                        if (WorkingCL > 0)
+                        {
+                            P4.DeleteChange(WorkingCL);
+                        }
+                    }
+                }
                 var Data = ReadAllText(P4IndexFileLocal);
                 var Parts = Data.Split(" ".ToCharArray());
                 int Index = int.Parse(Parts[0]);
