Fixed deploy error for launch on for android.

diff --git a/Engine/Source/Programs/AutomationTool/Android/AndroidPlatform.Automation.cs b/Engine/Source/Programs/AutomationTool/Android/AndroidPlatform.Automation.cs
--- a/Engine/Source/Programs/AutomationTool/Android/AndroidPlatform.Automation.cs
+++ b/Engine/Source/Programs/AutomationTool/Android/AndroidPlatform.Automation.cs
@@ -280,7 +280,7 @@ public class AndroidPlatform : Platform
 
 
 		// copy files to device if we were staging
-		if (SC.Stage && !Params.SkipStage)
+		if (SC.Stage)
 		{
 			// cache some strings
 			string BaseCommandline = AdbCommand + "push";
@@ -317,23 +317,23 @@ public class AndroidPlatform : Platform
 			// delete the .obb file, since it will cause nothing we just deployed to be used
 			Run(CmdEnv.CmdExe, AdbCommand + "shell rm " + DeviceObbName);
 		}
-		else
-		{
-			// cache some strings
-			string BaseCommandline = AdbCommand + "push";
-			string RemoteDir = "/mnt/sdcard/" + Params.ShortProjectName;
-
-			string FinalRemoteDir = RemoteDir;
-/*			// handle the special case of the UE4Commandline.txt when using content only game (UE4Game)
-			if (!Params.IsCodeBasedProject)
-			{
-				FinalRemoteDir = "/mnt/sdcard/UE4Game";
-			}*/
-
-			string RemoteFilename = IntermediateCmdLineFile.Replace(SC.StageDirectory, FinalRemoteDir).Replace("\\", "/");
-			string Commandline = string.Format("{0} \"{1}\" \"{2}\"", BaseCommandline, IntermediateCmdLineFile, RemoteFilename);
-			Run(CmdEnv.CmdExe, Commandline);
-		}
+        else
+        {
+            // cache some strings
+            string BaseCommandline = AdbCommand + "push";
+            string RemoteDir = "/mnt/sdcard/" + Params.ShortProjectName;
+
+            string FinalRemoteDir = RemoteDir;
+            /*			// handle the special case of the UE4Commandline.txt when using content only game (UE4Game)
+                        if (!Params.IsCodeBasedProject)
+                        {
+                            FinalRemoteDir = "/mnt/sdcard/UE4Game";
+                        }*/
+
+            string RemoteFilename = IntermediateCmdLineFile.Replace(SC.StageDirectory, FinalRemoteDir).Replace("\\", "/");
+            string Commandline = string.Format("{0} \"{1}\" \"{2}\"", BaseCommandline, IntermediateCmdLineFile, RemoteFilename);
+            Run(CmdEnv.CmdExe, Commandline);
+        }
 	}
 
 	/** Internal usage for GetPackageName */
