fix for launch reporting success even though it failed to launch
TTP337260
#ue4
#ios

diff --git a/Engine/Source/Programs/AutomationTool/Scripts/RunProjectCommand.Automation.cs b/Engine/Source/Programs/AutomationTool/Scripts/RunProjectCommand.Automation.cs
--- a/Engine/Source/Programs/AutomationTool/Scripts/RunProjectCommand.Automation.cs
+++ b/Engine/Source/Programs/AutomationTool/Scripts/RunProjectCommand.Automation.cs
@@ -302,6 +302,11 @@ public partial class Project : CommandUtils
 					Thread.Sleep(100);
 				}
 				while (ClientProcess.HasExited == false);
+
+				if (ClientProcess.ExitCode != 0)
+				{
+					throw new AutomationException("Client exited with error code: " + ClientProcess.ExitCode);
+				}
 			}
 		}
 	}
