Build script, half of the solution for "red after red" spam reduction

diff --git a/Engine/Source/Programs/AutomationTool/ProcessUtils.cs b/Engine/Source/Programs/AutomationTool/ProcessUtils.cs
--- a/Engine/Source/Programs/AutomationTool/ProcessUtils.cs
+++ b/Engine/Source/Programs/AutomationTool/ProcessUtils.cs
@@ -750,7 +750,7 @@ namespace AutomationTool
         /// <param name="Logfile">Full path to the logfile, where the application output should be written to.</param>
 		/// <param name="Input">Optional Input for the program (will be provided as stdin)</param>
 		/// <param name="Options">Defines the options how to run. See ERunOptions.</param>
-        public static void RunAndLog(string App, string CommandLine, string Logfile = null, int MaxSuccessCode = 0, string Input = null, ERunOptions Options = ERunOptions.Default)
+        public static string RunAndLog(string App, string CommandLine, string Logfile = null, int MaxSuccessCode = 0, string Input = null, ERunOptions Options = ERunOptions.Default)
         {
             ProcessResult Result = Run(App, CommandLine, Input, Options);
             if (Result.Output.Length > 0 && Logfile != null)
@@ -771,6 +771,11 @@ namespace AutomationTool
                 throw new AutomationException(String.Format("Command failed (Result:{3}): {0} {1}. See logfile for details: '{2}' ",
                                                 App, CommandLine, Path.GetFileName(Logfile), Result.ExitCode));
             }
+            if (Result.Output.Length > 0)
+            {
+                return Result.Output;
+            }
+            return "";
         }
 
         /// <summary>
diff --git a/Engine/Source/Programs/AutomationTool/Scripts/GUBP.Automation.cs b/Engine/Source/Programs/AutomationTool/Scripts/GUBP.Automation.cs
--- a/Engine/Source/Programs/AutomationTool/Scripts/GUBP.Automation.cs
+++ b/Engine/Source/Programs/AutomationTool/Scripts/GUBP.Automation.cs
@@ -3243,11 +3243,12 @@ public class GUBP : BuildCommand
         GUBPNodesCompleted.Add(NodeToDo, Result);
         return Result;
     }
-    void RunECTool(string Args, bool bQuiet = false)
+    string RunECTool(string Args, bool bQuiet = false)
     {
         if (ParseParam("FakeEC"))
         {
             Log("***** Would have ran ectool {0}", Args);
+            return "We didn't actually run ectool";
         }
         else
         {
@@ -3256,7 +3257,7 @@ public class GUBP : BuildCommand
             {
                 Opts = (Opts & ~ERunOptions.AllowSpew) | ERunOptions.NoLoggingOfRunCommand;
             }
-            RunAndLog("ectool", Args, Options: Opts);
+            return RunAndLog("ectool", Args, Options: Opts);
         }
     }
     string GetEMailListForNode(GUBP bp, string NodeToDo, string Causers)
@@ -3602,9 +3603,21 @@ public class GUBP : BuildCommand
         Result.Sort();
         return Result;
     }
-    void SaveStatus(string NodeToDo, string Suffix, string NodeStoreName, bool bSaveSharedTempStorage, string GameNameIfAny)
+    void SaveStatus(string NodeToDo, string Suffix, string NodeStoreName, bool bSaveSharedTempStorage, string GameNameIfAny, string JobStepIDForFailure = null)
     {
         string Contents = "Just a status record: " + Suffix;
+        if (!String.IsNullOrEmpty(JobStepIDForFailure) && IsBuildMachine)
+        {
+            try
+            {
+                Contents = RunECTool(String.Format("getProperties --jobStepId {0} --recurse 1", JobStepIDForFailure), true);
+            }
+            catch (Exception Ex)
+            {
+                Log(System.Diagnostics.TraceEventType.Warning, "Failed to get properties for jobstep to save them.");
+                Log(System.Diagnostics.TraceEventType.Warning, LogUtils.FormatException(Ex));
+            }
+        }
         string RecordOfSuccess = CombinePaths(CommandUtils.CmdEnv.LocalRoot, "Engine", "Saved", "Logs", NodeToDo + Suffix +".log");
         CreateDirectory(Path.GetDirectoryName(RecordOfSuccess));
         WriteAllText(RecordOfSuccess, Contents);
@@ -6004,7 +6017,7 @@ if (HostPlatform == UnrealTargetPlatform.Mac) continue; //temp hack till mac aut
                     if (SaveSuccessRecords)
                     {
                         UpdateHistoryAndEC(NodeToDo, CLString);
-                        SaveStatus(NodeToDo, FailedTempStorageSuffix, NodeStoreName, bSaveSharedTempStorage, GameNameIfAny);
+                        SaveStatus(NodeToDo, FailedTempStorageSuffix, NodeStoreName, bSaveSharedTempStorage, GameNameIfAny, ParseParamValue("MyJobStepId"));
                     }
 
                     Log("{0}", ExceptionToString(Ex));
