build script, fix automated testing

diff --git a/Engine/Source/Programs/AutomationTool/Scripts/RunProjectCommand.Automation.cs b/Engine/Source/Programs/AutomationTool/Scripts/RunProjectCommand.Automation.cs
--- a/Engine/Source/Programs/AutomationTool/Scripts/RunProjectCommand.Automation.cs
+++ b/Engine/Source/Programs/AutomationTool/Scripts/RunProjectCommand.Automation.cs
@@ -219,58 +219,58 @@ public partial class Project : CommandUtils
 					Log("Waiting for client logging process to start...{0}", ClientLogFile);
 					Thread.Sleep(2000);
 				}
-				if (!ClientProcess.HasExited)
+                if (FileExists(ClientLogFile))
 				{
 					Thread.Sleep(2000);
 					Log("Client logging process started...{0}", ClientLogFile);
 					ClientProcessLog = File.Open(ClientLogFile, FileMode.Open, FileAccess.Read, FileShare.ReadWrite);
 					ClientLogReader = new StreamReader(ClientProcessLog);
 				}
-				if (ClientProcess.HasExited || ClientLogReader == null)
+				if (ClientLogReader == null)
 				{
-					throw new AutomationException("Client exited before we asked it to.");
+					throw new AutomationException("Client exited without creating a log file.");
 				}
 				bool bKeepReading = true;
 				bool WelcomedCorrectly = false;
-				while (!ClientProcess.HasExited && bKeepReading)
+				while (!ClientLogReader.EndOfStream && bKeepReading)
 				{
-					while (!ClientLogReader.EndOfStream && bKeepReading && !ClientProcess.HasExited)
+					string ClientOutput = ClientLogReader.ReadToEnd();
+					if (!String.IsNullOrEmpty(ClientOutput))
 					{
-						string ClientOutput = ClientLogReader.ReadToEnd();
-						if (!String.IsNullOrEmpty(ClientOutput))
-						{
-                            AllClientOutput += ClientOutput;
-							Console.Write(ClientOutput);
+                        AllClientOutput += ClientOutput;
+						Console.Write(ClientOutput);
 
-                            if (AllClientOutput.LastIndexOf(LookFor) > AllClientOutput.IndexOf(LookFor))
-							{
-								Log("Client loaded, lets wait 30 seconds...");
-								Thread.Sleep(30000);
-								WelcomedCorrectly = true;
-								Log("Test complete...");
-								bKeepReading = false;
-							}
-                            else if (Params.RunAutomationTests)
+                        if (AllClientOutput.LastIndexOf(LookFor) > AllClientOutput.IndexOf(LookFor))
+						{
+							WelcomedCorrectly = true;
+							Log("Test complete...");
+							bKeepReading = false;
+						}
+                        else if (Params.RunAutomationTests)
+                        {
+                            int FailIndex = AllClientOutput.LastIndexOf("Automation Test Failed");
+                            int ParenIndex = AllClientOutput.LastIndexOf(")");
+                            if (FailIndex >= 0 && ParenIndex > FailIndex && FailIndex > LastAutoFailIndex)
                             {
-                                int FailIndex = AllClientOutput.LastIndexOf("Automation Test Failed");
-                                int ParenIndex = AllClientOutput.LastIndexOf(")");
-                                if (FailIndex >= 0 && ParenIndex > FailIndex && FailIndex > LastAutoFailIndex)
+                                string Tail = AllClientOutput.Substring(FailIndex);
+                                int CloseParenIndex = Tail.IndexOf(")");
+                                int OpenParenIndex = Tail.IndexOf("(");
+                                string Test = "";
+                                if (OpenParenIndex >= 0 && CloseParenIndex > OpenParenIndex)
                                 {
-                                    string Tail = AllClientOutput.Substring(FailIndex);
-                                    int CloseParenIndex = Tail.IndexOf(")");
-                                    int OpenParenIndex = Tail.IndexOf("(");
-                                    string Test = "";
-                                    if (OpenParenIndex >= 0 && CloseParenIndex > OpenParenIndex)
-                                    {
-                                        Test = Tail.Substring(OpenParenIndex + 1, CloseParenIndex - OpenParenIndex - 1);
-                                        Log(System.Diagnostics.TraceEventType.Error, "Automated test failed ({0}).", Test);
-                                        LastAutoFailIndex = FailIndex;
-                                    }
+                                    Test = Tail.Substring(OpenParenIndex + 1, CloseParenIndex - OpenParenIndex - 1);
+                                    Log(System.Diagnostics.TraceEventType.Error, "Automated test failed ({0}).", Test);
+                                    LastAutoFailIndex = FailIndex;
                                 }
                             }
                         }
-					}
+                    }
 				}
+                if (ClientProcess != null && !ClientProcess.HasExited)
+                {
+                    Log("Client is supposed to exit, lets wait 30 seconds...");
+                    Thread.Sleep(30000);
+                }
 				if (ClientProcess != null && !ClientProcess.HasExited)
 				{
 					Log("Stopping client...");
