build script, preflight builds

diff --git a/Engine/Source/Programs/AutomationTool/P4Utils.cs b/Engine/Source/Programs/AutomationTool/P4Utils.cs
--- a/Engine/Source/Programs/AutomationTool/P4Utils.cs
+++ b/Engine/Source/Programs/AutomationTool/P4Utils.cs
@@ -1517,6 +1517,21 @@ namespace AutomationTool
 			return Result;
 		}
 
+        /// <summary>
+        /// Returns the output from p4 opened
+        /// </summary>
+        /// <param name="CL">Changelist to get the specification from.</param>
+        /// <returns>Specification of the changelist.</returns>
+        public string OpenedOutput()
+        {
+            CheckP4Enabled();
+            string CmdOutput;
+            if (LogP4Output(out CmdOutput, "opened"))
+            {
+                return CmdOutput;
+            }
+            throw new P4Exception("OpenedOutput failed, output follows\n{0}", CmdOutput);
+        }
 		/// <summary>
 		/// Deletes the specified label.
 		/// </summary>
diff --git a/Engine/Source/Programs/AutomationTool/Scripts/GUBP.Automation.cs b/Engine/Source/Programs/AutomationTool/Scripts/GUBP.Automation.cs
--- a/Engine/Source/Programs/AutomationTool/Scripts/GUBP.Automation.cs
+++ b/Engine/Source/Programs/AutomationTool/Scripts/GUBP.Automation.cs
@@ -25,6 +25,9 @@ public class GUBP : BuildCommand
     public static bool bForceIncrementalCompile = false;
     public static string ECProject = null;
     public string EmailHint;
+    static public bool bPreflightBuild = false;
+    public int PreflightShelveCL = 0;
+    static public string PreflightMangleSuffix = "";
 
     Dictionary<string, GUBPNode> GUBPNodes;
     Dictionary<string, bool> GUBPNodesCompleted;
@@ -151,7 +154,12 @@ public class GUBP : BuildCommand
 
     private static List<GUBPEmailHacker> EmailHackers;
     private string HackEmails(string Emails, string Branch, string NodeName)
-    {        
+    {
+        string OnlyEmail = ParseParamValue("OnlyEmail");
+        if (!String.IsNullOrEmpty(OnlyEmail))
+        {
+            return OnlyEmail;
+        }
         EmailHint = ParseParamValue("EmailHint");
         if(EmailHint == null)
         {
@@ -1408,6 +1416,23 @@ public class GUBP : BuildCommand
         }
     }
 
+    public class GeneralSuccessNode : GUBP.SuccessNode
+    {
+        string MyName;
+        public GeneralSuccessNode(string InMyName)
+        {
+            MyName = InMyName;
+        }
+        public static string StaticGetFullName(string InMyName)
+        {
+            return InMyName + "_Success";
+        }
+        public override string GetFullName()
+        {
+            return StaticGetFullName(MyName);
+        }
+    }
+
     public class AggregateNode : GUBPNode
     {
         public AggregateNode()
@@ -1823,6 +1848,10 @@ public class GUBP : BuildCommand
             {
                 CompleteLabelPrefix = "Rocket-" + CompleteLabelPrefix;
             }
+            if (GUBP.bPreflightBuild)
+            {
+                CompleteLabelPrefix = CompleteLabelPrefix + PreflightMangleSuffix;
+            }
             return CompleteLabelPrefix;
         }
 
@@ -4112,6 +4141,27 @@ public class GUBP : BuildCommand
     public override void ExecuteBuild()
     {
         Log("************************* GUBP");
+        string PreflightShelveCLString = GetEnvVar("uebp_PreflightShelveCL");
+        if ((!String.IsNullOrEmpty(PreflightShelveCLString) && IsBuildMachine) || ParseParam("PreflightTest"))
+        {
+            Log("**** Preflight shelve {0}", PreflightShelveCLString);
+            if (!String.IsNullOrEmpty(PreflightShelveCLString))
+            {
+                PreflightShelveCL = int.Parse(PreflightShelveCLString);
+                if (PreflightShelveCL < 2000000)
+                {
+                    throw new AutomationException("{0} does not look like a CL");
+                }
+            }
+            bPreflightBuild = true;
+
+            var OpenFiles = P4.OpenedOutput();
+            if (OpenFiles.Contains("File(s) not opened on this client."))
+            {
+                throw new AutomationException("Apparently, nothing was unshelved for this preflight.");
+            }
+        }
+        
         ECProject = ParseParamValue("ECProject");
         if (ECProject == null)
         {
@@ -4148,6 +4198,13 @@ public class GUBP : BuildCommand
         {
             StoreSuffix = StoreSuffix + "-Rkt";
         }
+        if (bPreflightBuild)
+        {
+            int TotalSeconds = (int)((DateTime.UtcNow - new DateTime(2014, 6, 1)).TotalSeconds);
+
+            PreflightMangleSuffix = String.Format("-PF-{0}-{1}", PreflightShelveCL, TotalSeconds);
+            StoreSuffix = StoreSuffix + PreflightMangleSuffix;
+        }
         CL = ParseParamInt("CL", 0);
         bool bCleanLocalTempStorage = ParseParam("CleanLocal");
         bool bChanges = ParseParam("Changes") || ParseParam("AllChanges");
@@ -5019,14 +5076,30 @@ if (HostPlatform == UnrealTargetPlatform.Mac) continue; //temp hack till mac aut
         {
             AddCustomNodes(HostPlatform);
         }
-
+        
         if (HasNode(ToolsForCompileNode.StaticGetFullName(UnrealTargetPlatform.Win64)))
         {
             if (HasNode(GamePlatformMonolithicsNode.StaticGetFullName(UnrealTargetPlatform.Mac, Branch.BaseEngineProject, UnrealTargetPlatform.IOS)) && HasNode(ToolsNode.StaticGetFullName(UnrealTargetPlatform.Win64)))
             {
                 AddNode(new IOSOnPCTestNode(this));
             }
-            AddNode(new CleanSharedTempStorageNode(this));
+            if (!bPreflightBuild)
+            {
+                AddNode(new CleanSharedTempStorageNode(this));
+            }
+        }
+
+        if (bPreflightBuild)
+        {
+            GeneralSuccessNode PreflightSuccessNode = new GeneralSuccessNode("Preflight");
+            foreach (var NodeToDo in GUBPNodes)
+            {
+                if (NodeToDo.Value.RunInEC())
+                {
+                    PreflightSuccessNode.AddPseudodependency(NodeToDo.Key);
+                }
+            }
+            AddNode(PreflightSuccessNode);
         }
 
         foreach (var NodeToDo in GUBPNodes)
@@ -5525,6 +5598,11 @@ if (HostPlatform == UnrealTargetPlatform.Mac) continue; //temp hack till mac aut
                             continue; // this wasn't on the chain related to the trigger we are triggering, so it is not relevant
                         }
                     }
+                    if (bPreflightBuild && !bSkipTriggers && GUBPNodes[NodeToDo].TriggerNode())
+                    {
+                        // in preflight builds, we are either skipping triggers (and running things downstream) or we just stop at triggers and don't make them available for triggering.
+                        continue;
+                    }
                     FilteredOrdereredToDo.Add(NodeToDo);
                 }
             }
