Android generated install batch file now takes an optional device serial number to pass to ADB
#android

diff --git a/Engine/Source/Programs/AutomationTool/Android/AndroidPlatform.Automation.cs b/Engine/Source/Programs/AutomationTool/Android/AndroidPlatform.Automation.cs
--- a/Engine/Source/Programs/AutomationTool/Android/AndroidPlatform.Automation.cs
+++ b/Engine/Source/Programs/AutomationTool/Android/AndroidPlatform.Automation.cs
@@ -150,12 +150,14 @@ public class AndroidPlatform : Platform
 		string[] BatchLines = new string[] {
 			"setlocal",
 			"set ADB=%ANDROID_HOME%\\platform-tools\\adb.exe",
-			"%ADB% uninstall " + PackageName,
-			"%ADB% install " + Path.GetFileName(ApkName),
-			"%ADB% shell rm -rf /mnt/sdcard/" + Params.ShortProjectName,
-			"%ADB% shell rm -rf /mnt/sdcard/UE4Game/UE4CommandLine.txt", // we need to delete the commandline in UE4Game or it will mess up loading
-			"%ADB% shell rm -rf /mnt/sdcard/obb/" + PackageName,
-			Params.OBBinAPK ? "" : "%ADB% push " + Path.GetFileName(LocalObbName) + " " + DeviceObbName,
+            "set DEVICE=",
+            "if not \"%1\"==\"\" set DEVICE=-s %1",
+			"%ADB% %DEVICE% uninstall " + PackageName,
+			"%ADB% %DEVICE% install " + Path.GetFileName(ApkName),
+			"%ADB% %DEVICE% shell rm -rf /mnt/sdcard/" + Params.ShortProjectName,
+			"%ADB% %DEVICE% shell rm -rf /mnt/sdcard/UE4Game/UE4CommandLine.txt", // we need to delete the commandline in UE4Game or it will mess up loading
+			"%ADB% %DEVICE% shell rm -rf /mnt/sdcard/obb/" + PackageName,
+			Params.OBBinAPK ? "" : "%ADB% %DEVICE% push " + Path.GetFileName(LocalObbName) + " " + DeviceObbName,
 		};
 		File.WriteAllLines(BatchName, BatchLines);
 
