Changed "rm -rf" to "rm -r" in Android generated install batch file; -f not supported on all devices
#android
#UAT

diff --git a/Engine/Source/Programs/AutomationTool/Android/AndroidPlatform.Automation.cs b/Engine/Source/Programs/AutomationTool/Android/AndroidPlatform.Automation.cs
--- a/Engine/Source/Programs/AutomationTool/Android/AndroidPlatform.Automation.cs
+++ b/Engine/Source/Programs/AutomationTool/Android/AndroidPlatform.Automation.cs
@@ -154,9 +154,9 @@ public class AndroidPlatform : Platform
             "if not \"%1\"==\"\" set DEVICE=-s %1",
 			"%ADB% %DEVICE% uninstall " + PackageName,
 			"%ADB% %DEVICE% install " + Path.GetFileName(ApkName),
-			"%ADB% %DEVICE% shell rm -rf /mnt/sdcard/" + Params.ShortProjectName,
-			"%ADB% %DEVICE% shell rm -rf /mnt/sdcard/UE4Game/UE4CommandLine.txt", // we need to delete the commandline in UE4Game or it will mess up loading
-			"%ADB% %DEVICE% shell rm -rf /mnt/sdcard/obb/" + PackageName,
+			"%ADB% %DEVICE% shell rm -r /mnt/sdcard/" + Params.ShortProjectName,
+			"%ADB% %DEVICE% shell rm -r /mnt/sdcard/UE4Game/UE4CommandLine.txt", // we need to delete the commandline in UE4Game or it will mess up loading
+			"%ADB% %DEVICE% shell rm -r /mnt/sdcard/obb/" + PackageName,
 			Params.OBBinAPK ? "" : "%ADB% %DEVICE% push " + Path.GetFileName(LocalObbName) + " " + DeviceObbName,
 		};
 		File.WriteAllLines(BatchName, BatchLines);
