Try switching chmod command line to use an escaped string rather than a quoted string. Should fix Rocket builds in main.

diff --git a/Engine/Source/Programs/AutomationTool/BuildUtils.cs b/Engine/Source/Programs/AutomationTool/BuildUtils.cs
--- a/Engine/Source/Programs/AutomationTool/BuildUtils.cs
+++ b/Engine/Source/Programs/AutomationTool/BuildUtils.cs
@@ -102,7 +102,7 @@ namespace AutomationTool
 		public static void FixUnixFilePermissions(string Filename)
 		{
 			string Permissions = IsProbablyAMacOrIOSExe(Filename) ? "0755" : "0644";
-			var Result = CommandUtils.Run("sh", string.Format("-c 'chmod {0} \"{1}\"'", Permissions, Filename));
+			var Result = CommandUtils.Run("sh", string.Format("-c 'chmod {0} {1}'", Permissions, Filename.Replace(" ", "\\ ")));
 			if (Result.ExitCode != 0)
 			{
 				throw new AutomationException(String.Format("Failed to chmod \"{0}\"", Filename));
