build script, fix race condition in automated testing

diff --git a/Engine/Source/Programs/AutomationTool/Scripts/RunProjectCommand.Automation.cs b/Engine/Source/Programs/AutomationTool/Scripts/RunProjectCommand.Automation.cs
--- a/Engine/Source/Programs/AutomationTool/Scripts/RunProjectCommand.Automation.cs
+++ b/Engine/Source/Programs/AutomationTool/Scripts/RunProjectCommand.Automation.cs
@@ -232,11 +232,23 @@ public partial class Project : CommandUtils
 				}
 				bool bKeepReading = true;
 				bool WelcomedCorrectly = false;
-				while (!ClientLogReader.EndOfStream && bKeepReading)
+                bool bClientExited = false;
+                DateTime ExitTime = DateTime.UtcNow;
+
+                while (bKeepReading)
 				{
+                    if (!bClientExited && ClientProcess.HasExited)
+                    {
+                        ExitTime = DateTime.UtcNow;
+                        bClientExited = true;
+                    }
 					string ClientOutput = ClientLogReader.ReadToEnd();
 					if (!String.IsNullOrEmpty(ClientOutput))
 					{
+                        if (bClientExited)
+                        {
+                            ExitTime = DateTime.UtcNow; // as long as it is spewing, we reset the timer
+                        }
                         AllClientOutput += ClientOutput;
 						Console.Write(ClientOutput);
 
@@ -264,18 +276,35 @@ public partial class Project : CommandUtils
                                 }
                             }
                         }
+                    }
+                    else if (bClientExited && (DateTime.UtcNow - ExitTime).TotalSeconds > 30)
+                    {
+                        Log("Client exited and has been quiet for 30 seconds...exiting");
+                        bKeepReading = false;
                     }
 				}
                 if (ClientProcess != null && !ClientProcess.HasExited)
                 {
-                    Log("Client is supposed to exit, lets wait 30 seconds...");
-                    Thread.Sleep(30000);
+                    Log("Client is supposed to exit, lets wait a while for it to exit naturally...");
+                    for (int i = 0 ; i < 120 && !ClientProcess.HasExited; i++)
+                    {
+                        Thread.Sleep(1000);
+                    }
                 }
 				if (ClientProcess != null && !ClientProcess.HasExited)
 				{
 					Log("Stopping client...");
 					ClientProcess.StopProcess();
-				}
+                    Thread.Sleep(10000);
+                }
+                while (ClientLogReader != null && !ClientLogReader.EndOfStream)
+                {
+                    string ClientOutput = ClientLogReader.ReadToEnd();
+                    if (!String.IsNullOrEmpty(ClientOutput))
+                    {
+                        Console.Write(ClientOutput);
+                    }
+                }
 
 				if (!WelcomedCorrectly)
 				{
