build script, red-after-red spam reduction

diff --git a/Engine/Source/Programs/AutomationTool/CommandUtils.cs b/Engine/Source/Programs/AutomationTool/CommandUtils.cs
--- a/Engine/Source/Programs/AutomationTool/CommandUtils.cs
+++ b/Engine/Source/Programs/AutomationTool/CommandUtils.cs
@@ -376,6 +376,23 @@ namespace AutomationTool
 				}
 			}
 		}
+        /// <summary>
+        /// Deletes a file(s). 
+        /// If the file does not exist, silently succeeds.
+        /// If the deletion of the file fails, this function throws an Exception.
+        /// </summary>
+        /// <param name="Filename">Filename</param>
+        public static void DeleteFile(bool bQuiet, params string[] Filenames)
+        {
+            foreach (var Filename in Filenames)
+            {
+                var NormalizedFilename = ConvertSeparators(PathSeparator.Default, Filename);
+                if (!InternalUtils.SafeDeleteFile(NormalizedFilename, bQuiet))
+                {
+                    throw new AutomationException(String.Format("Failed to delete file '{0}'", NormalizedFilename));
+                }
+            }
+        }
 
 		/// <summary>
 		/// Deletes a file(s). 
diff --git a/Engine/Source/Programs/AutomationTool/Scripts/GUBP.Automation.cs b/Engine/Source/Programs/AutomationTool/Scripts/GUBP.Automation.cs
--- a/Engine/Source/Programs/AutomationTool/Scripts/GUBP.Automation.cs
+++ b/Engine/Source/Programs/AutomationTool/Scripts/GUBP.Automation.cs
@@ -6,8 +6,92 @@ using System.IO;
 using AutomationTool;
 using UnrealBuildTool;
 using System.Reflection;
+using System.Xml;
+
+public class ECJobPropsUtils
+{
+    public static HashSet<string> ErrorsFromProps(string Filename)
+    {
+        var Result = new HashSet<string>();
+        XmlDocument Doc = new XmlDocument();
+        Doc.Load(Filename);
+        foreach (XmlElement ChildNode in Doc.FirstChild.ChildNodes)
+        {
+            if (ChildNode.Name == "propertySheet")
+            {
+                foreach (XmlElement PropertySheetChild in ChildNode.ChildNodes)
+                {
+                    if (PropertySheetChild.Name == "property")
+                    {
+                        bool IsDiag = false;
+                        foreach (XmlElement PropertySheetChildDiag in PropertySheetChild.ChildNodes)
+                        {
+                            if (PropertySheetChildDiag.Name == "propertyName" && PropertySheetChildDiag.InnerText == "ec_diagnostics")
+                            {
+                                IsDiag = true;
+                            }
+                            if (IsDiag && PropertySheetChildDiag.Name == "propertySheet")
+                            {
+                                foreach (XmlElement PropertySheetChildDiagSheet in PropertySheetChildDiag.ChildNodes)
+                                {
+                                    if (PropertySheetChildDiagSheet.Name == "property")
+                                    {
+                                        bool IsError = false;
+                                        foreach (XmlElement PropertySheetChildDiagSheetElem in PropertySheetChildDiagSheet.ChildNodes)
+                                        {
+                                            if (PropertySheetChildDiagSheetElem.Name == "propertyName" && PropertySheetChildDiagSheetElem.InnerText.StartsWith("error-"))
+                                            {
+                                                IsError = true;
+                                            }
+                                            if (IsError && PropertySheetChildDiagSheetElem.Name == "propertySheet")
+                                            {
+                                                foreach (XmlElement PropertySheetChildDiagSheetElemInner in PropertySheetChildDiagSheetElem.ChildNodes)
+                                                {
+                                                    if (PropertySheetChildDiagSheetElemInner.Name == "property")
+                                                    {
+                                                        bool IsMessage = false;
+                                                        foreach (XmlElement PropertySheetChildDiagSheetElemInner2 in PropertySheetChildDiagSheetElemInner.ChildNodes)
+                                                        {
+                                                            if (PropertySheetChildDiagSheetElemInner2.Name == "propertyName" && PropertySheetChildDiagSheetElemInner2.InnerText == "message")
+                                                            {
+                                                                IsMessage = true;
+                                                            }
+                                                            if (IsMessage && PropertySheetChildDiagSheetElemInner2.Name == "value")
+                                                            {
+                                                                Result.Add(PropertySheetChildDiagSheetElemInner2.InnerText);
+                                                            }
+                                                        }
+                                                    }
+                                                }
+                                            }
+                                        }
+                                    }
+                                }
+                            }
+                        }
+                    }
+                }
+            }
+        }
+        return Result;
+    }
+}
+
+public class TestECJobErrorParse : BuildCommand
+{
+    public override void ExecuteBuild()
+    {
+        Log("*********************** TestECJobErrorParse");
+
+        string Filename = CombinePaths(@"P:\Builds\UE4\GUBP\++depot+UE4-2104401-RootEditor_Failed\Engine\Saved\Logs", "RootEditor_Failed.log");
+        var Errors = ECJobPropsUtils.ErrorsFromProps(Filename);
+        foreach (var ThisError in Errors)
+        {
+            Log("Error: {0}", ThisError);
+        }
+    }
+}
 
-//ProjectUtil.cs
 
 public class GUBP : BuildCommand
 {
@@ -3967,7 +4051,78 @@ public class GUBP : BuildCommand
         }
     }
 
-    List<string> GetECPropsForNode(string NodeToDo, out string EMails, bool OnlyLateUpdates = false)
+    bool HashSetEqual(HashSet<string> A, HashSet<string> B)
+    {
+        if (A.Count != B.Count)
+        {
+            return false;
+        }
+        foreach (var Elem in A)
+        {
+            if (!B.Contains(Elem))
+            {
+                return false;
+            }
+        }
+        foreach (var Elem in B)
+        {
+            if (!A.Contains(Elem))
+            {
+                return false;
+            }
+        }
+        return true;
+    }
+
+    int FindLastNonDuplicateFail(string NodeToDo, string CLString)
+    {
+        var History = GUBPNodesHistory[NodeToDo];
+        int Result = P4Env.Changelist;
+
+        string GameNameIfAny = GUBPNodes[NodeToDo].GameNameIfAnyForTempStorage();
+        string NodeStore = StoreName + "-" + GUBPNodes[NodeToDo].GetFullName() + FailedTempStorageSuffix;
+
+        var BackwardsFails = new List<int>(History.AllFailed);
+        BackwardsFails.Add(P4Env.Changelist);
+        BackwardsFails.Sort();
+        BackwardsFails.Reverse();
+        HashSet<string> CurrentErrors = null;
+        foreach (var CL in BackwardsFails)
+        {
+            if (CL > P4Env.Changelist)
+            {
+                continue;
+            }
+            if (CL <= History.LastSucceeded)
+            {
+                break;
+            }
+            var ThisNodeStore = NodeStore.Replace(CLString, String.Format("{0}", CL));
+            DeleteLocalTempStorage(CmdEnv, ThisNodeStore, true); // these all clash locally, which is fine we just retrieve them from shared
+            bool WasLocal;
+            var Files = RetrieveFromTempStorage(CmdEnv, ThisNodeStore, out WasLocal, GameNameIfAny);
+            if (Files.Count != 1)
+            {
+                throw new AutomationException("Unexpected number of files for fail record {0}", Files.Count);
+            }
+            string ErrorFile = Files[0];
+            var ThisErrors = ECJobPropsUtils.ErrorsFromProps(ErrorFile);
+            if (CurrentErrors == null)
+            {
+                CurrentErrors = ThisErrors;
+            }
+            else
+            {
+                if (!HashSetEqual(CurrentErrors, ThisErrors))
+                {
+                    break;
+                }
+                Result = CL;
+            }
+        }
+        return Result;
+    }
+    List<string> GetECPropsForNode(string NodeToDo, string CLString, out string EMails, bool OnlyLateUpdates = false)
     {
         var StartTime = DateTime.UtcNow;
 
@@ -3988,7 +4143,23 @@ public class GUBP : BuildCommand
 
             if (History.LastSucceeded > 0 && History.LastSucceeded < P4Env.Changelist)
             {
-                var ChangeRecords = GetChanges(History.LastSucceeded, P4Env.Changelist, History.LastSucceeded);
+                int LastNonDuplicateFail = P4Env.Changelist;
+                try
+                {
+                    LastNonDuplicateFail = FindLastNonDuplicateFail(NodeToDo, CLString);
+                    if (LastNonDuplicateFail < P4Env.Changelist)
+                    {
+                        Log("*** Red-after-red spam reduction, changed CL {0} to CL {1} because the errors didn't change.", P4Env.Changelist, LastNonDuplicateFail);
+                    }
+                }
+                catch (Exception Ex)
+                {
+                    LastNonDuplicateFail = P4Env.Changelist;
+                    Log(System.Diagnostics.TraceEventType.Warning, "Failed to FindLastNonDuplicateFail.");
+                    Log(System.Diagnostics.TraceEventType.Warning, LogUtils.FormatException(Ex));
+                }
+
+                var ChangeRecords = GetChanges(History.LastSucceeded, LastNonDuplicateFail, History.LastSucceeded);
                 foreach (var Record in ChangeRecords)
                 {
                     FailCauserEMails = GUBPNode.MergeSpaceStrings(FailCauserEMails, Record.UserEmail);
@@ -4077,14 +4248,13 @@ public class GUBP : BuildCommand
         return ECProps;
     }
 
-    void UpdateHistoryAndEC(string NodeToDo, string CLString)
+    void UpdateECProps(string NodeToDo, string CLString)
     {
         try
         {
             Log("Updating node props for node {0}", NodeToDo);
-            UpdateNodeHistory(NodeToDo, CLString);
             string EMails;
-            var Props = GetECPropsForNode(NodeToDo, out EMails, true);
+            var Props = GetECPropsForNode(NodeToDo, CLString, out EMails, true);
             foreach (var Prop in Props)
             {
                 var Parts = Prop.Split("=".ToCharArray());
@@ -4093,7 +4263,7 @@ public class GUBP : BuildCommand
         }
         catch (Exception Ex)
         {
-            Log(System.Diagnostics.TraceEventType.Warning, "Failed to UpdateHistoryAndEC.");
+            Log(System.Diagnostics.TraceEventType.Warning, "Failed to UpdateECProps.");
             Log(System.Diagnostics.TraceEventType.Warning, LogUtils.FormatException(Ex));
         }
     }
@@ -5689,7 +5859,7 @@ if (HostPlatform == UnrealTargetPlatform.Mac) continue; //temp hack till mac aut
                 if (GUBPNodes[NodeToDo].RunInEC() && !NodeIsAlreadyComplete(NodeToDo, LocalOnly)) // if something is already finished, we don't put it into EC  
                 {
                     string EMails;
-                    var NodeProps = GetECPropsForNode(NodeToDo, out EMails);
+                    var NodeProps = GetECPropsForNode(NodeToDo, CLString, out EMails);
                     ECProps.AddRange(NodeProps);
 
                     bool Sticky = GUBPNodes[NodeToDo].IsSticky();
@@ -6094,8 +6264,9 @@ if (HostPlatform == UnrealTargetPlatform.Mac) continue; //temp hack till mac aut
                 {
                     if (SaveSuccessRecords)
                     {
-                        UpdateHistoryAndEC(NodeToDo, CLString);
+                        UpdateNodeHistory(NodeToDo, CLString);
                         SaveStatus(NodeToDo, FailedTempStorageSuffix, NodeStoreName, bSaveSharedTempStorage, GameNameIfAny, ParseParamValue("MyJobStepId"));
+                        UpdateECProps(NodeToDo, CLString);
                     }
 
                     Log("{0}", ExceptionToString(Ex));
@@ -6145,8 +6316,9 @@ if (HostPlatform == UnrealTargetPlatform.Mac) continue; //temp hack till mac aut
                 }
                 if (SaveSuccessRecords) 
                 {
-                    UpdateHistoryAndEC(NodeToDo, CLString);
+                    UpdateNodeHistory(NodeToDo, CLString);
                     SaveStatus(NodeToDo, SucceededTempStorageSuffix, NodeStoreName, bSaveSharedTempStorage, GameNameIfAny);
+                    UpdateECProps(NodeToDo, CLString);
                 }
             }
             foreach (var Product in GUBPNodes[NodeToDo].BuildProducts)
diff --git a/Engine/Source/Programs/AutomationTool/TempStorage.cs b/Engine/Source/Programs/AutomationTool/TempStorage.cs
--- a/Engine/Source/Programs/AutomationTool/TempStorage.cs
+++ b/Engine/Source/Programs/AutomationTool/TempStorage.cs
@@ -692,6 +692,11 @@ namespace AutomationTool
             }
             return false;
         }
+        public static void DeleteLocalTempStorage(CommandEnvironment Env, string StorageBlockName, bool bQuiet = false)
+        {
+            var LocalManifest = LocalTempStorageManifestFilename(Env, StorageBlockName);
+            DeleteFile(bQuiet, LocalManifest);
+        }
         public static bool SharedTempStorageExists(CommandEnvironment Env, string StorageBlockName, string GameFolder = "", bool bQuiet = false)
         {
             var SharedManifest = SharedTempStorageManifestFilename(Env, StorageBlockName, GameFolder);
