build script, remove a certain error from red-after-red spam reduction because it contains the build CL and jobid

diff --git a/Engine/Source/Programs/AutomationTool/Scripts/GUBP.Automation.cs b/Engine/Source/Programs/AutomationTool/Scripts/GUBP.Automation.cs
--- a/Engine/Source/Programs/AutomationTool/Scripts/GUBP.Automation.cs
+++ b/Engine/Source/Programs/AutomationTool/Scripts/GUBP.Automation.cs
@@ -58,7 +58,10 @@ public class ECJobPropsUtils
                                                             }
                                                             if (IsMessage && PropertySheetChildDiagSheetElemInner2.Name == "value")
                                                             {
-                                                                Result.Add(PropertySheetChildDiagSheetElemInner2.InnerText);
+                                                                if (!PropertySheetChildDiagSheetElemInner2.InnerText.Contains("LogTailsAndChanges"))
+                                                                {
+                                                                    Result.Add(PropertySheetChildDiagSheetElemInner2.InnerText);
+                                                                }
                                                             }
                                                         }
                                                     }
