build script, fix red-after-red spam reduction

diff --git a/Engine/Source/Programs/AutomationTool/Scripts/GUBP.Automation.cs b/Engine/Source/Programs/AutomationTool/Scripts/GUBP.Automation.cs
--- a/Engine/Source/Programs/AutomationTool/Scripts/GUBP.Automation.cs
+++ b/Engine/Source/Programs/AutomationTool/Scripts/GUBP.Automation.cs
@@ -4102,8 +4102,20 @@ public class GUBP : BuildCommand
             }
             var ThisNodeStore = NodeStore.Replace(CLString, String.Format("{0}", CL));
             DeleteLocalTempStorage(CmdEnv, ThisNodeStore, true); // these all clash locally, which is fine we just retrieve them from shared
-            bool WasLocal;
-            var Files = RetrieveFromTempStorage(CmdEnv, ThisNodeStore, out WasLocal, GameNameIfAny);
+
+            List<string> Files = null;
+            try
+            {
+                bool WasLocal;
+                Files = RetrieveFromTempStorage(CmdEnv, ThisNodeStore, out WasLocal, GameNameIfAny); // this will fail on our CL if we didn't fail or we are just setting up the branch
+            }
+            catch (Exception)
+            {
+            }
+            if (Files == null)
+            {
+                continue;
+            }
             if (Files.Count != 1)
             {
                 throw new AutomationException("Unexpected number of files for fail record {0}", Files.Count);
@@ -4149,10 +4161,13 @@ public class GUBP : BuildCommand
                 int LastNonDuplicateFail = P4Env.Changelist;
                 try
                 {
-                    LastNonDuplicateFail = FindLastNonDuplicateFail(NodeToDo, CLString);
-                    if (LastNonDuplicateFail < P4Env.Changelist)
+                    if (OnlyLateUpdates)
                     {
-                        Log("*** Red-after-red spam reduction, changed CL {0} to CL {1} because the errors didn't change.", P4Env.Changelist, LastNonDuplicateFail);
+                        LastNonDuplicateFail = FindLastNonDuplicateFail(NodeToDo, CLString);
+                        if (LastNonDuplicateFail < P4Env.Changelist)
+                        {
+                            Log("*** Red-after-red spam reduction, changed CL {0} to CL {1} because the errors didn't change.", P4Env.Changelist, LastNonDuplicateFail);
+                        }
                     }
                 }
                 catch (Exception Ex)
