Reworked deploy process for android automation testing.
Was previously using staging build when should have been using archive builds.

diff --git a/Engine/Source/Programs/AutomationTool/Android/AndroidPlatform.Automation.cs b/Engine/Source/Programs/AutomationTool/Android/AndroidPlatform.Automation.cs
--- a/Engine/Source/Programs/AutomationTool/Android/AndroidPlatform.Automation.cs
+++ b/Engine/Source/Programs/AutomationTool/Android/AndroidPlatform.Automation.cs
@@ -260,20 +260,6 @@ public class AndroidPlatform : Platform
 		string InstallCommandline = AdbCommand + "install \"" + ApkName + "\"";
 		RunAndLog(CmdEnv, CmdEnv.CmdExe, InstallCommandline);
 
-        if (Params.Prebuilt)
-        {
-            // deploy the obb if there is one
-            string ObbPath = Path.Combine(SC.StageDirectory, GetFinalObbName(ApkName));
-            if (File.Exists(ObbPath))
-            {
-                // cache some strings
-                string BaseCommandline = AdbCommand + "push";
-
-                string Commandline = string.Format("{0} \"{1}\" \"{2}\"", BaseCommandline, ObbPath, DeviceObbName );
-                Run(CmdEnv.CmdExe, Commandline);
-            }
-        }
-
 		// update the ue4commandline.txt
 		// update and deploy ue4commandline.txt
 		// always delete the existing commandline text file, so it doesn't reuse an old one
@@ -319,6 +305,19 @@ public class AndroidPlatform : Platform
 			// delete the .obb file, since it will cause nothing we just deployed to be used
 			Run(CmdEnv.CmdExe, AdbCommand + "shell rm " + DeviceObbName);
 		}
+        else if (SC.Archive)
+        {
+            // deploy the obb if there is one
+            string ObbPath = Path.Combine(SC.StageDirectory, GetFinalObbName(ApkName));
+            if (File.Exists(ObbPath))
+            {
+                // cache some strings
+                string BaseCommandline = AdbCommand + "push";
+
+                string Commandline = string.Format("{0} \"{1}\" \"{2}\"", BaseCommandline, ObbPath, DeviceObbName);
+                Run(CmdEnv.CmdExe, Commandline);
+            }
+        }
         else
         {
             // cache some strings
diff --git a/Engine/Source/Programs/AutomationTool/ProjectParams.cs b/Engine/Source/Programs/AutomationTool/ProjectParams.cs
--- a/Engine/Source/Programs/AutomationTool/ProjectParams.cs
+++ b/Engine/Source/Programs/AutomationTool/ProjectParams.cs
@@ -417,11 +417,13 @@ namespace AutomationTool
             if (this.Prebuilt)
             {
                 this.SkipCook = true;
-                this.SkipPak = true;
+                /*this.SkipPak = true;
                 this.SkipStage = true;
-                this.Cook = true;
-                this.Stage = true;
                 this.Pak = true;
+                this.Stage = true;*/
+                this.Cook = true;
+                this.Archive = true;
+                
                 this.Deploy = true;
                 this.Run = true;
                 //this.StageDirectoryParam = this.PrebuiltDir;
