build minor tweak to red-after-red criteria, if there are no errors, then we consider that non-reducible

diff --git a/Engine/Source/Programs/AutomationTool/Scripts/GUBP.Automation.cs b/Engine/Source/Programs/AutomationTool/Scripts/GUBP.Automation.cs
--- a/Engine/Source/Programs/AutomationTool/Scripts/GUBP.Automation.cs
+++ b/Engine/Source/Programs/AutomationTool/Scripts/GUBP.Automation.cs
@@ -4128,7 +4128,7 @@ public class GUBP : BuildCommand
             }
             else
             {
-                if (!HashSetEqual(CurrentErrors, ThisErrors))
+                if (CurrentErrors.Count == 0 || !HashSetEqual(CurrentErrors, ThisErrors))
                 {
                     break;
                 }
