Added validation of -compressed command with UAT

diff --git a/Engine/Source/Programs/AutomationTool/ProjectParams.cs b/Engine/Source/Programs/AutomationTool/ProjectParams.cs
--- a/Engine/Source/Programs/AutomationTool/ProjectParams.cs
+++ b/Engine/Source/Programs/AutomationTool/ProjectParams.cs
@@ -1536,6 +1536,11 @@ namespace AutomationTool
 			{
 				throw new AutomationException("Don't use both -cook and -cookonthefly.");
 			}
+
+            if (Compressed && !Pak)
+            {
+                throw new AutomationException("-compressed can only be used with -pak");
+            }
 		}
 
 		protected bool bLogged = false;
