add support to HTML5 to build to the Game/Binaries/HTML5 directory
#ue4
#html5

diff --git a/Engine/Source/Programs/AutomationTool/HTML5/HTML5Platform.Automation.cs b/Engine/Source/Programs/AutomationTool/HTML5/HTML5Platform.Automation.cs
--- a/Engine/Source/Programs/AutomationTool/HTML5/HTML5Platform.Automation.cs
+++ b/Engine/Source/Programs/AutomationTool/HTML5/HTML5Platform.Automation.cs
@@ -18,12 +18,12 @@ public class HTML5Platform : Platform
 	{
 		Log("Package {0}", Params.RawProjectPath);
 
-		string FinalDataLocation = Path.Combine(Path.GetDirectoryName(Path.GetFullPath(Params.ProjectGameExeFilename)), Params.ProjectBinariesFolder, Params.ShortProjectName) +".data";
-        string OutDir = Path.Combine(Path.GetDirectoryName(Path.GetFullPath(Params.ProjectGameExeFilename)), Params.ProjectBinariesFolder);
-
-        bool FolderExists = Directory.Exists(OutDir);
-        if (!FolderExists)
-            Directory.CreateDirectory(OutDir);
+		string PackagePath = Path.Combine(Path.GetDirectoryName(Params.RawProjectPath), "Binaries", "HTML5");
+		if (!Directory.Exists(PackagePath))
+		{
+			Directory.CreateDirectory(PackagePath);
+		}
+		string FinalDataLocation = Path.Combine(PackagePath, Params.ShortProjectName) +".data";
 
 
 		// we need to operate in the root
@@ -34,25 +34,39 @@ public class HTML5Platform : Platform
 			if (Utils.IsRunningOnMono) 
 			{
 				string PackagerPath = BaseSDKPath + "/tools/file_packager.py";
-				string CmdLine = string.Format ("-c \" python {0} {1} --preload . --js-output={1}.js \" ", PackagerPath, FinalDataLocation);
+				string CmdLine = string.Format ("-c \" python {0} '{1}' --preload . --js-output='{1}.js' \" ", PackagerPath, FinalDataLocation);
 				RunAndLog (CmdEnv, "/bin/bash", CmdLine);	
 			} else 
 			{ 
 				string PackagerPath = "\"" + BaseSDKPath + "\\tools\\file_packager.py\"";
-				string CmdLine = string.Format ("/c python {0} {1} --preload . --js-output={1}.js", PackagerPath, FinalDataLocation);
+				string CmdLine = string.Format ("/c python {0} \"{1}\" --preload . --js-output=\"{1}.js\"", PackagerPath, FinalDataLocation);
 				RunAndLog (CmdEnv, CommandUtils.CombinePaths (Environment.SystemDirectory, "cmd.exe"), CmdLine);	
 			}
 		}
-		// put the HTML file to the binaries directory
+
+		// copy the "Executable" to the package directory
+		string GameExe = Path.GetFileNameWithoutExtension(Params.ProjectGameExeFilename);
+		if (Params.ClientConfigsToBuild[0].ToString() != "Development")
+		{
+			GameExe += "-HTML5-" + Params.ClientConfigsToBuild[0].ToString();
+		}
+		GameExe += ".js";
+		if (Path.Combine(Path.GetDirectoryName(Params.ProjectGameExeFilename), GameExe) != Path.Combine(PackagePath, GameExe))
+		{
+			File.Copy(Path.Combine(Path.GetDirectoryName(Params.ProjectGameExeFilename), GameExe), Path.Combine(PackagePath, GameExe), true);
+			File.Copy(Path.Combine(Path.GetDirectoryName(Params.ProjectGameExeFilename), GameExe) + ".mem", Path.Combine(PackagePath, GameExe) + ".mem", true);
+		}
+		File.SetAttributes(Path.Combine(PackagePath, GameExe), FileAttributes.Normal);
+		File.SetAttributes(Path.Combine(PackagePath, GameExe) + ".mem", FileAttributes.Normal);
+		
+		// put the HTML file to the package directory
 		string TemplateFile = Path.Combine(CombinePaths(CmdEnv.LocalRoot, "Engine"), "Build", "HTML5", "Game.html.template");
-		string OutputFile = Path.Combine(Path.GetDirectoryName(Path.GetFullPath(Params.ProjectGameExeFilename)), Params.ProjectBinariesFolder, (Params.ClientConfigsToBuild[0].ToString() != "Development" ? (Params.ShortProjectName + "-HTML5-" + Params.ClientConfigsToBuild[0].ToString()) : Params.ShortProjectName)) + ".html";
+		string OutputFile = Path.Combine(PackagePath, (Params.ClientConfigsToBuild[0].ToString() != "Development" ? (Params.ShortProjectName + "-HTML5-" + Params.ClientConfigsToBuild[0].ToString()) : Params.ShortProjectName)) + ".html";
     	
-
 		// find Heap Size.
 		ulong HeapSize;
 		var ConfigCache = new UnrealBuildTool.ConfigCacheIni(UnrealTargetPlatform.HTML5, "Engine", Path.GetDirectoryName(Params.RawProjectPath),CombinePaths(CmdEnv.LocalRoot, "Engine")); 
 
-
 		int ConfigHeapSize; 
 		if (!ConfigCache.GetInt32 ("BuildSettings","HeapSize"+Params.ClientConfigsToBuild[0].ToString(), out ConfigHeapSize)) // in Megs.
 		{ 	 
@@ -72,10 +86,13 @@ public class HTML5Platform : Platform
 
 		// copy the jstorage files to the binaries directory
 		string JSDir = Path.Combine(CombinePaths(CmdEnv.LocalRoot, "Engine"), "Build", "HTML5");
-	
-        File.Copy(JSDir + "/json2.js", OutDir + "/json2.js", true);
-        File.Copy(JSDir + "/jstorage.js", OutDir + "/jstorage.js", true);
-        File.Copy(JSDir + "/moz_binarystring.js", OutDir + "/moz_binarystring.js", true);
+		string OutDir = PackagePath;
+		File.Copy(JSDir + "/json2.js", OutDir + "/json2.js", true);
+		File.SetAttributes(OutDir + "/json2.js", FileAttributes.Normal);
+		File.Copy(JSDir + "/jstorage.js", OutDir + "/jstorage.js", true);
+		File.SetAttributes(OutDir + "/jstorage.js", FileAttributes.Normal);
+		File.Copy(JSDir + "/moz_binarystring.js", OutDir + "/moz_binarystring.js", true);
+		File.SetAttributes(OutDir + "/moz_binarystring.js", FileAttributes.Normal);
 		PrintRunTime();
 	}
 
@@ -100,8 +117,8 @@ public class HTML5Platform : Platform
 
 				if (LineStr.Contains("%CONFIG%"))
 				{
-                    if (IsContentOnly)
-                        InGameName = "UE4Game";
+					if (IsContentOnly)
+						InGameName = "UE4Game";
 					LineStr = LineStr.Replace("%CONFIG%", (InGameConfiguration != "Development" ? (InGameName + "-HTML5-" + InGameConfiguration) : InGameName));
 				}
 
@@ -187,7 +204,7 @@ public class HTML5Platform : Platform
         if  ( !File.Exists(HTMLPath) )  // its probably a content only game - then it exists in the UE4 directory and not in the game directory. 
             HTMLPath = Path.Combine(CombinePaths(CmdEnv.LocalRoot, "Engine"), "Binaries", "HTML5", Path.GetFileName(HTMLPath)); 
 
-		ProcessResult ClientProcess = Run(FirefoxPath, FirefoxProfileCommand + HTMLPath , null, ClientRunFlags | ERunOptions.NoWaitForExit);
+		ProcessResult ClientProcess = Run(FirefoxPath, FirefoxProfileCommand + "\"" + HTMLPath + "\"", null, ClientRunFlags | ERunOptions.NoWaitForExit);
 
 		return ClientProcess;
 	}
