Allow SpeedTree for Rocket builds.

diff --git a/Engine/Source/Programs/AutomationTool/Scripts/GUBP.Automation.cs b/Engine/Source/Programs/AutomationTool/Scripts/GUBP.Automation.cs
--- a/Engine/Source/Programs/AutomationTool/Scripts/GUBP.Automation.cs
+++ b/Engine/Source/Programs/AutomationTool/Scripts/GUBP.Automation.cs
@@ -140,7 +140,7 @@ public class GUBP : BuildCommand
     {
         if (bBuildRocket)
         {
-            return " -NoSimplygon -NoSpeedTree " + (bUseRocketInsteadOfBuildRocket ? "-Rocket" : "-BuildRocket");
+            return " -NoSimplygon " + (bUseRocketInsteadOfBuildRocket ? "-Rocket" : "-BuildRocket");
         }
         return "";
     }
