Faster Android deployment. TTP#326554

diff --git a/Engine/Source/Programs/AutomationTool/Android/AndroidPlatform.Automation.cs b/Engine/Source/Programs/AutomationTool/Android/AndroidPlatform.Automation.cs
--- a/Engine/Source/Programs/AutomationTool/Android/AndroidPlatform.Automation.cs
+++ b/Engine/Source/Programs/AutomationTool/Android/AndroidPlatform.Automation.cs
@@ -12,6 +12,8 @@ using UnrealBuildTool;
 
 public class AndroidPlatform : Platform
 {
+    private const int DeployMaxParallelCommands = 6;
+
 	public AndroidPlatform()
 		:base(UnrealTargetPlatform.Android)
 	{
@@ -235,6 +237,40 @@ public class AndroidPlatform : Platform
         }
     }
 
+    /*
+    private class TimeRegion : System.IDisposable
+    {
+        private System.DateTime StartTime { get; set; }
+
+        private string Format { get; set; }
+
+        private System.Collections.Generic.List<object> FormatArgs { get; set; }
+
+        public TimeRegion(string format, params object[] format_args)
+        {
+            Format = format;
+            FormatArgs = new List<object>(format_args);
+            StartTime = DateTime.UtcNow;
+        }
+
+        public void Dispose()
+        {
+            double total_time = (DateTime.UtcNow - StartTime).TotalMilliseconds / 1000.0;
+            FormatArgs.Insert(0, total_time);
+            CommandUtils.Log(Format, FormatArgs.ToArray());
+        }
+    }
+    */
+
+    internal class LongestFirst : IComparer<string>
+    {
+        public int Compare(string a, string b)
+        {
+            if (a.Length == b.Length) return a.CompareTo(b);
+            else return b.Length - a.Length;
+        }
+    }
+
 	public override void Deploy(ProjectParams Params, DeploymentContext SC)
 	{
 		string SOName = GetFinalSOName(Params, SC.StageExecutables[0]);
@@ -242,7 +278,7 @@ public class AndroidPlatform : Platform
 		string ApkName = GetFinalApkName(Params, SC.StageExecutables[0], true);
 
 		// make sure APK is up to date (this is fast if so)
-		var Deploy = UEBuildDeploy.GetBuildDeploy(UnrealTargetPlatform.Android);
+        var Deploy = UEBuildDeploy.GetBuildDeploy(UnrealTargetPlatform.Android);
         if (!Params.Prebuilt)
         {
             string CookFlavor = SC.CookPlatform.IndexOf("_") > 0 ? SC.CookPlatform.Substring(SC.CookPlatform.IndexOf("_")) : "";
@@ -254,8 +290,8 @@ public class AndroidPlatform : Platform
 		string PackageName = GetPackageInfo(ApkName, false);
 
 		// install the apk
-		string UninstallCommandline = AdbCommand + "uninstall " + PackageName;
-		RunAndLog(CmdEnv, CmdEnv.CmdExe, UninstallCommandline);
+        string UninstallCommandline = AdbCommand + "uninstall " + PackageName;
+        RunAndLog(CmdEnv, CmdEnv.CmdExe, UninstallCommandline);
 
 		string InstallCommandline = AdbCommand + "install \"" + ApkName + "\"";
 		RunAndLog(CmdEnv, CmdEnv.CmdExe, InstallCommandline);
@@ -279,32 +315,109 @@ public class AndroidPlatform : Platform
 			Run(CmdEnv.CmdExe, AdbCommand + "shell rm -r " + RemoteDir);
 			Run(CmdEnv.CmdExe, AdbCommand + "shell rm -r " + UE4GameRemoteDir);
 
+            // Copy UFS files..
 			string[] Files = Directory.GetFiles(SC.StageDirectory, "*", SearchOption.AllDirectories);
-			// copy each UFS file
-			foreach (string Filename in Files)
-			{
-				// don't push the apk, we install it
-				if (Path.GetExtension(Filename).Equals(".apk", StringComparison.InvariantCultureIgnoreCase))
-				{
-					continue;
-				}
+            System.Array.Sort(Files);
+
+            // Find all the files we exclude from copying. And include
+            // the directories we need to individually copy.
+            HashSet<string> ExcludedFiles = new HashSet<string>();
+            SortedSet<string> IndividualCopyDirectories
+                = new SortedSet<string>((IComparer<string>)new LongestFirst());
+            foreach (string Filename in Files)
+            {
+                bool Exclude = false;
+                // Don't push the apk, we install it
+                Exclude |= Path.GetExtension(Filename).Equals(".apk", StringComparison.InvariantCultureIgnoreCase);
+                // For excluded files we add the parent dirs to our
+                // tracking of stuff to individually copy.
+                if (Exclude)
+                {
+                    ExcludedFiles.Add(Filename);
+                    // We include all directories up to the stage root in having
+                    // to individually copy the files.
+                    for (string FileDirectory = Path.GetDirectoryName(Filename);
+                        !FileDirectory.Equals(SC.StageDirectory);
+                        FileDirectory = Path.GetDirectoryName(FileDirectory))
+                    {
+                        if (!IndividualCopyDirectories.Contains(FileDirectory))
+                        {
+                            IndividualCopyDirectories.Add(FileDirectory);
+                        }
+                    }
+                    if (!IndividualCopyDirectories.Contains(SC.StageDirectory))
+                    {
+                        IndividualCopyDirectories.Add(SC.StageDirectory);
+                    }
+                }
+            }
 
-				string FinalRemoteDir = RemoteDir;
-/*				// handle the special case of the UE4Commandline.txt when using content only game (UE4Game)
-				if (!Params.IsCodeBasedProject &&
-					Path.GetFileName(Filename).Equals("UE4CommandLine.txt", StringComparison.InvariantCultureIgnoreCase))
-				{
-					FinalRemoteDir = "/mnt/sdcard/UE4Game";
-				}*/
+            // The directories are sorted above in "deepest" first. We can
+            // therefore start copying those individual dirs which will
+            // recreate the tree. As the subtrees will get copied at each
+            // possible individual level.
+            HashSet<string> EntriesToDeploy = new HashSet<string>();
+            foreach (string DirectoryName in IndividualCopyDirectories)
+            {
+                string[] Entries
+                    = Directory.GetFileSystemEntries(DirectoryName, "*", SearchOption.TopDirectoryOnly);
+                foreach (string Entry in Entries)
+                {
+                    // We avoid excluded files and the individual copy dirs
+                    // (the individual copy dirs will get handled as we iterate).
+                    if (ExcludedFiles.Contains(Entry) || IndividualCopyDirectories.Contains(Entry))
+                    {
+                        continue;
+                    }
+                    else
+                    {
+                        EntriesToDeploy.Add(Entry);
+                    }
+                }
+            }
+            if (EntriesToDeploy.Count == 0)
+            {
+                EntriesToDeploy.Add(SC.StageDirectory);
+            }
 
-				string RemoteFilename = Filename.Replace(SC.StageDirectory, FinalRemoteDir).Replace("\\", "/");
- 				string Commandline = string.Format("{0} \"{1}\" \"{2}\"", BaseCommandline, Filename, RemoteFilename);
- 				Run(CmdEnv.CmdExe, Commandline);
-			}
+            // We now have a minimal set of file & dir entries we need
+            // to deploy. Files we deploy will get individually copied
+            // and dirs will get the tree copies by default (that's
+            // what ADB does).
+            HashSet<ProcessResult> DeployCommands = new HashSet<ProcessResult>();
+            foreach (string Entry in EntriesToDeploy)
+            {
+				string FinalRemoteDir = RemoteDir;
+				string RemotePath = Entry.Replace(SC.StageDirectory, FinalRemoteDir).Replace("\\", "/");
+ 				string Commandline = string.Format("{0} \"{1}\" \"{2}\"", BaseCommandline, Entry, RemotePath);
+                // We run deploy commands in parallel to maximize the connection
+                // throughput.
+                DeployCommands.Add(
+                    Run(CmdEnv.CmdExe, Commandline, null,
+                        ERunOptions.Default | ERunOptions.NoWaitForExit));
+                // But we limit the parallel commands to avoid overwhelming
+                // memory resources.
+                if (DeployCommands.Count == DeployMaxParallelCommands)
+                {
+                    while (DeployCommands.Count > DeployMaxParallelCommands / 2)
+                    {
+                        Thread.Sleep(10);
+                        DeployCommands.RemoveWhere(
+                            delegate(ProcessResult r)
+                            {
+                                return r.HasExited;
+                            });
+                    }
+                }
+            }
+            foreach (ProcessResult deploy_result in DeployCommands)
+            {
+                deploy_result.WaitForExit();
+            }
 
 			// delete the .obb file, since it will cause nothing we just deployed to be used
 			Run(CmdEnv.CmdExe, AdbCommand + "shell rm " + DeviceObbName);
-		}
+        }
         else if (SC.Archive)
         {
             // deploy the obb if there is one
@@ -325,17 +438,19 @@ public class AndroidPlatform : Platform
             string RemoteDir = "/mnt/sdcard/" + Params.ShortProjectName;
 
             string FinalRemoteDir = RemoteDir;
-            /*			// handle the special case of the UE4Commandline.txt when using content only game (UE4Game)
-                        if (!Params.IsCodeBasedProject)
-                        {
-                            FinalRemoteDir = "/mnt/sdcard/UE4Game";
-                        }*/
+            /*
+            // handle the special case of the UE4Commandline.txt when using content only game (UE4Game)
+            if (!Params.IsCodeBasedProject)
+            {
+                FinalRemoteDir = "/mnt/sdcard/UE4Game";
+            }
+            */
 
             string RemoteFilename = IntermediateCmdLineFile.Replace(SC.StageDirectory, FinalRemoteDir).Replace("\\", "/");
             string Commandline = string.Format("{0} \"{1}\" \"{2}\"", BaseCommandline, IntermediateCmdLineFile, RemoteFilename);
             Run(CmdEnv.CmdExe, Commandline);
         }
-	}
+    }
 
 	/** Internal usage for GetPackageName */
 	private static string PackageLine = null;
