Added timeout to testformalbuild script call.
Now it will wait for 5 minutes to check if everything is ok.

diff --git a/Engine/Source/Programs/AutomationTool/Scripts/GUBP.Automation.cs b/Engine/Source/Programs/AutomationTool/Scripts/GUBP.Automation.cs
--- a/Engine/Source/Programs/AutomationTool/Scripts/GUBP.Automation.cs
+++ b/Engine/Source/Programs/AutomationTool/Scripts/GUBP.Automation.cs
@@ -2776,7 +2776,7 @@ public class GUBP : BuildCommand
                 throw new AutomationException("Archive directory does not exist {0}, so we can't test the build.", ArchiveDirectory);
             }
 
-            string WorkingCommandline = String.Format("TestFormalBuild {0} -Archive -archivedirectory={1} -platform={2} -clientconfig={3}",
+            string WorkingCommandline = String.Format("TestFormalBuild {0} -Archive -archivedirectory={1} -platform={2} -clientconfig={3} -timeoutseconds=300",
                 ProjectArg, CommandUtils.MakePathSafeToUseWithCommandLine(ArchiveDirectory), ClientTargetPlatform.ToString(), ClientConfig.ToString());
 
             if (WorkingCommandline.Contains("-project=\"\""))
