Testformalbuild outputs system to log file but only parses ue4 log for errors.

diff --git a/Engine/Source/Programs/AutomationTool/Android/AndroidPlatform.Automation.cs b/Engine/Source/Programs/AutomationTool/Android/AndroidPlatform.Automation.cs
--- a/Engine/Source/Programs/AutomationTool/Android/AndroidPlatform.Automation.cs
+++ b/Engine/Source/Programs/AutomationTool/Android/AndroidPlatform.Automation.cs
@@ -584,8 +584,12 @@ public class AndroidPlatform : Platform
                     break;
                 }
             }
-            
-            ProcessResult LogFileProcess = Run(CmdEnv.CmdExe, AdbCommand + "logcat -d");
+
+            // this is just to get the ue4 log to go to the output
+            Run(CmdEnv.CmdExe, AdbCommand + "logcat -d -s UE4");
+
+            // get the log we actually want to save
+            ProcessResult LogFileProcess = Run(CmdEnv.CmdExe, AdbCommand + "logcat -d", null, ERunOptions.AppMustExist);
             
             File.WriteAllText(LogFilename, LogFileProcess.Output);
             File.WriteAllText(ServerLogFilename, LogFileProcess.Output);
