Removed staging of ICU DLLs from Windows platform staging - unnecessary as games do not use ICU DLLs and packaging only applies to standalone games.

diff --git a/Engine/Source/Programs/AutomationTool/Win/WinPlatform.Automation.cs b/Engine/Source/Programs/AutomationTool/Win/WinPlatform.Automation.cs
--- a/Engine/Source/Programs/AutomationTool/Win/WinPlatform.Automation.cs
+++ b/Engine/Source/Programs/AutomationTool/Win/WinPlatform.Automation.cs
@@ -45,8 +45,6 @@ public abstract class BaseWinPlatform : Platform
             StageExecutable("exe", SC, CommandUtils.CombinePaths(SC.LocalRoot, "Engine/Binaries", SC.PlatformDir), "CrashReportClient.");
         }
 
-        StageExecutable("dll", SC, CommandUtils.CombinePaths(SC.LocalRoot, "Engine/Binaries/ThirdParty/ICU/icu4c-53_1", SC.PlatformDir, "VS" + WindowsPlatform.GetVisualStudioCompilerVersionName()), Params.bDebugBuildsActuallyUseDebugCRT ? "*d." : "*.", true, new [] {Params.bDebugBuildsActuallyUseDebugCRT ? "*." : "*d."});
-
 		//todo we need to support shipping and test executables
 		//todo this should all be partially based on UBT manifests and not hard coded
 		//monolithic assumption
