Reworked the way devices are found so that build machine can detect turned off devices and unauthorized devices.

diff --git a/Engine/Source/Programs/AutomationTool/Android/AndroidPlatform.Automation.cs b/Engine/Source/Programs/AutomationTool/Android/AndroidPlatform.Automation.cs
--- a/Engine/Source/Programs/AutomationTool/Android/AndroidPlatform.Automation.cs
+++ b/Engine/Source/Programs/AutomationTool/Android/AndroidPlatform.Automation.cs
@@ -220,19 +220,37 @@ public class AndroidPlatform : Platform
 
         if (Result.Output.Length > 0)
         {
-            string[] LogLines = Result.Output.Split(new char [] { '\n', '\r', '\t' });
-            for ( int i = 1; i < LogLines.Length; ++i )
+            string[] LogLines = Result.Output.Split(new char [] { '\n', '\r' });
+            bool FoundList = false;
+            for ( int i = 0; i < LogLines.Length; ++i )
             {
-                if (LogLines[i].Length == 0)
+                if (FoundList == false)
                 {
+                    if (LogLines[i].StartsWith("List of devices attached"))
+                    {
+                        FoundList = true;
+                    }
                     continue;
                 }
-                if (LogLines[i] == "device")
+
+                string[] DeviceLine = LogLines[i].Split(new char[] { '\t' });
+
+                if (DeviceLine.Length == 2)
                 {
-                    continue;
+                    // the second param should be "device"
+                    // if it's not setup correctly it might be "unattached" or "powered off" or something like that
+                    // warning in that case
+                    if (DeviceLine[1] == "device")
+                    {
+                        Devices.Add("@" + DeviceLine[0] );
+                    }
+                    else
+                    {
+                        CommandUtils.LogWarning("Device attached but in bad state {0}:{1}", DeviceLine[0], DeviceLine[1]);
+                    }
                 }
-
-                Devices.Add("@"+LogLines[i]);
+                
+                
             }
         }
     }
