build script, remove CIS heros

diff --git a/Engine/Source/Programs/AutomationTool/Scripts/GUBP.Automation.cs b/Engine/Source/Programs/AutomationTool/Scripts/GUBP.Automation.cs
--- a/Engine/Source/Programs/AutomationTool/Scripts/GUBP.Automation.cs
+++ b/Engine/Source/Programs/AutomationTool/Scripts/GUBP.Automation.cs
@@ -4229,10 +4229,12 @@ public class GUBP : BuildCommand
             var AdditonalEmails = "";
             if (ParseParam("CIS") && !GUBPNodes[NodeToDo].SendSuccessEmail() && !GUBPNodes[NodeToDo].TriggerNode())
             {
+#if false // cisheros was disabled because it was ignored and was just spam
                 if (NumPeople > 50 || NumPeople == 0)
                 {
                     FailCauserEMails = "[CISHeros]";
                 }
+#endif
                 AdditonalEmails = FailCauserEMails;
             }
             string AddEmails = ParseParamValue("AddEmails");
