Add FrequencyHacker plus initial frequency changes

diff --git a/Engine/Source/Programs/AutomationTool/Scripts/GUBP.Automation.cs b/Engine/Source/Programs/AutomationTool/Scripts/GUBP.Automation.cs
--- a/Engine/Source/Programs/AutomationTool/Scripts/GUBP.Automation.cs
+++ b/Engine/Source/Programs/AutomationTool/Scripts/GUBP.Automation.cs
@@ -300,6 +300,43 @@ public class GUBP : BuildCommand
         }
         return FinalEmails;
     }
+    public abstract class GUBPFrequencyHacker
+    {
+        public virtual int GetNodeFrequency(GUBP bp, string Branch, string NodeName, int BaseFrequency)
+        {
+            return new int();
+        }
+    }
+    private static List<GUBPFrequencyHacker> FrequencyHackers;
+    private int HackFrequency(GUBP bp, string Branch, string NodeName, int BaseFrequency)
+    {
+        int Frequency = BaseFrequency;
+        if (FrequencyHackers == null)
+        {
+            FrequencyHackers = new List<GUBPFrequencyHacker>();
+            Assembly[] LoadedAssemblies = AppDomain.CurrentDomain.GetAssemblies();
+            foreach (var Dll in LoadedAssemblies)
+            {
+                Type[] AllTypes = Dll.GetTypes();
+                foreach (var PotentialConfigType in AllTypes)
+                {
+                    if (PotentialConfigType != typeof(GUBPFrequencyHacker) && typeof(GUBPFrequencyHacker).IsAssignableFrom(PotentialConfigType))
+                    {
+                        GUBPFrequencyHacker Config = Activator.CreateInstance(PotentialConfigType) as GUBPFrequencyHacker;
+                        if (Config != null)
+                        {
+                            FrequencyHackers.Add(Config);
+                        }
+                    }
+                }
+            }
+        }
+        foreach(var FrequencyHacker in FrequencyHackers)
+        {
+           Frequency = FrequencyHacker.GetNodeFrequency(bp, Branch, NodeName, BaseFrequency);            
+        }
+        return Frequency;
+    }
     public abstract class GUBPNode
     {
         public List<string> FullNamesOfDependencies = new List<string>();
@@ -584,11 +621,6 @@ public class GUBP : BuildCommand
         {
             return base.TimeoutInMinutes() + ((HostPlatform == UnrealTargetPlatform.Mac) ? 30 : 0); // Mac is slower and more resource constrained
         }
-        public override int CISFrequencyQuantumShift(GUBP bp)
-        {
-            int Result = base.CISFrequencyQuantumShift(bp);            
-            return Result;
-        }
     }
 
     public class CompileNode : HostPlatformNode
@@ -978,7 +1010,7 @@ public class GUBP : BuildCommand
         }
         public override int CISFrequencyQuantumShift(GUBP bp)
         {
-            int Result = base.CISFrequencyQuantumShift(bp) + 1;
+            int Result = base.CISFrequencyQuantumShift(bp) + 1;                              
             return Result;
         }
         public override bool DeleteBuildProducts()
@@ -1067,7 +1099,7 @@ public class GUBP : BuildCommand
         }
         public override int CISFrequencyQuantumShift(GUBP bp)
         {
-            int Result = base.CISFrequencyQuantumShift(bp) + 1;
+            int Result = base.CISFrequencyQuantumShift(bp) + 1;                             
             return Result;
         }
         public override float Priority()
@@ -1120,6 +1152,11 @@ public class GUBP : BuildCommand
         {
             return base.Priority() + 1;
         }
+        public override int CISFrequencyQuantumShift(GUBP bp)
+        {
+            int Result = base.CISFrequencyQuantumShift(bp);            
+            return Result;
+        }
 
         public override UE4Build.BuildAgenda GetAgenda(GUBP bp)
         {
@@ -1180,10 +1217,10 @@ public class GUBP : BuildCommand
         public override int CISFrequencyQuantumShift(GUBP bp)
         {
             int Result = base.CISFrequencyQuantumShift(bp);
-            if (GameProj.Options(HostPlatform).bTestWithShared)
+            if(GameProj.Options(HostPlatform).bTestWithShared)
             {
                 Result += 3;
-            }
+            }                    
             return Result;
         }
         public override UE4Build.BuildAgenda GetAgenda(GUBP bp)
@@ -1283,18 +1320,19 @@ public class GUBP : BuildCommand
         public override bool DeleteBuildProducts()
         {
             return true;
-        }
+        }  
+
         public override int CISFrequencyQuantumShift(GUBP bp)
-        {
+        {            
             int Result = base.CISFrequencyQuantumShift(bp);
-            if (GameProj.GameName != bp.Branch.BaseEngineProject.GameName)
+            if(GameProj.GameName != bp.Branch.BaseEngineProject.GameName)
             {
-                Result += 3; // only every 80m
+                Result += 3; //only every 80m
             }
             else if (TargetPlatform != HostPlatform)
             {
-                Result += 2; // only every 40m
-            }
+                Result += 2; //only every 40m
+            }                 
             return Result;
         }
 
@@ -2161,7 +2199,7 @@ public class GUBP : BuildCommand
         }
         public override int CISFrequencyQuantumShift(GUBP bp)
         {
-            return base.CISFrequencyQuantumShift(bp) + 4 + (bIsMassive ? 1 : 0);                
+            return base.CISFrequencyQuantumShift(bp) + 3 + (bIsMassive ? 1 : 0);
         }
         public override float Priority()
         {
@@ -2983,7 +3021,7 @@ public class GUBP : BuildCommand
         }
         public override int CISFrequencyQuantumShift(GUBP bp)
         {
-            return base.CISFrequencyQuantumShift(bp) + 5 /*++ (bIsMassive ? 1 : 0);*/ ;
+            return base.CISFrequencyQuantumShift(bp) + 5;
         }
         public override string ECAgentString()
         {
@@ -3310,10 +3348,11 @@ public class GUBP : BuildCommand
 
     public int ComputeDependentCISFrequencyQuantumShift(string NodeToDo)
     {
-        int Result = GUBPNodes[NodeToDo].ComputedDependentCISFrequencyQuantumShift;
+        int Result = GUBPNodes[NodeToDo].ComputedDependentCISFrequencyQuantumShift;        
         if (Result < 0)
         {
             Result = GUBPNodes[NodeToDo].CISFrequencyQuantumShift(this);
+            Result = GetFrequencyForNode(this, GUBPNodes[NodeToDo].GetFullName(), Result);
             foreach (var Dep in GUBPNodes[NodeToDo].FullNamesOfDependencies)
             {
                 Result = Math.Max(ComputeDependentCISFrequencyQuantumShift(Dep), Result);
@@ -3385,6 +3424,15 @@ public class GUBP : BuildCommand
         }
         return HackEmails(Causers, BranchForEmail, NodeToDo);
     }
+    int GetFrequencyForNode(GUBP bp, string NodeToDo, int BaseFrequency)
+    {
+        var BranchForFrequency = "";
+        if (P4Enabled)
+        {
+            BranchForFrequency = P4Env.BuildRootP4;
+        }
+        return HackFrequency(bp, BranchForFrequency, NodeToDo, BaseFrequency);
+    }
 
     List<P4Connection.ChangeRecord> GetChanges(int LastOutputForChanges, int TopCL, int LastGreen)
     {
