Fix ICU deployment for Linux.

diff --git a/Engine/Source/Programs/AutomationTool/Linux/LinuxPlatform.Automation.cs b/Engine/Source/Programs/AutomationTool/Linux/LinuxPlatform.Automation.cs
--- a/Engine/Source/Programs/AutomationTool/Linux/LinuxPlatform.Automation.cs
+++ b/Engine/Source/Programs/AutomationTool/Linux/LinuxPlatform.Automation.cs
@@ -31,8 +31,7 @@ public abstract class BaseLinuxPlatform : Platform
         }
 
         {
-            //SC.StageFiles(StagedFileType.NonUFS, CombinePaths(SC.LocalRoot, "Engine/Binaries/ThirdParty/ICU/icu4c-53_1", SC.PlatformDir, "x86_64-unknown-linux-gnu"), Params.bDebugBuildsActuallyUseDebugCRT ? "*d." : "*.", false, Params.bDebugBuildsActuallyUseDebugCRT ? "*." : "*d.", CombinePaths("Engine/Binaries", SC.PlatformDir));
-            SC.StageFiles(StagedFileType.NonUFS, CombinePaths(SC.LocalRoot, "Engine/Binaries/ThirdParty/ICU/icu4c-53_1/", SC.PlatformDir, "x86_64-unknown-linux-gnu"), Params.bDebugBuildsActuallyUseDebugCRT ? "*d." : "*.", false, new[] { Params.bDebugBuildsActuallyUseDebugCRT ? "*." : "*d." }, CombinePaths("Engine/Binaries", SC.PlatformDir));
+            SC.StageFiles(StagedFileType.NonUFS, CombinePaths(SC.LocalRoot, "Engine/Binaries/ThirdParty/ICU/icu4c-53_1/", SC.PlatformDir, "x86_64-unknown-linux-gnu"), Params.bDebugBuildsActuallyUseDebugCRT ? "*d.so" : "*.so", false, new[] { Params.bDebugBuildsActuallyUseDebugCRT ? "*.so" : "*d.so" }, CombinePaths("Engine/Binaries", SC.PlatformDir));
         }
         
         // assume that we always have to deploy Steam (FIXME: should be automatic)
