Added interface methods for get build by label, and checking for a build manifest on the prod cdn

diff --git a/Engine/Source/Programs/AutomationTool/MCPPublic.cs b/Engine/Source/Programs/AutomationTool/MCPPublic.cs
--- a/Engine/Source/Programs/AutomationTool/MCPPublic.cs
+++ b/Engine/Source/Programs/AutomationTool/MCPPublic.cs
@@ -370,6 +370,15 @@ namespace EpicGames.MCP.Automation
 		/// <param name="Platform">Platform to add to base label.</param>
 		abstract public string GetBuildVersionWithPlatform(BuildPatchToolStagingInfo StagingInfo);
 
+		/// <summary>
+		/// 
+		/// </summary>
+		/// <param name="AppName">Application name to check the label in</param>
+		/// <param name="LabelName">Label name to get the build for</param>
+		/// <param name="McpConfigName">Which BuildInfo backend to label the build in.</param>
+		/// <returns></returns>
+		abstract public string GetLabeledBuildVersion(string AppName, string LabelName, string McpConfigName);
+
 		/// <summary>
 		/// Apply the requested label to the requested build in the BuildInfo backend for the requested MCP environment
 		/// </summary>
@@ -403,6 +412,12 @@ namespace EpicGames.MCP.Automation
         /// <returns>Path to the share (allows for override)</returns>
         abstract public string MountInternalCDNShare();
 
+		/// <summary>
+		/// Checks for a stagingInfo's manifest on the production CDN.
+		/// </summary>
+		/// <param name="stagingInfo">Staging info used to determine where the chunks are to copy.</param>
+		abstract public bool IsManifestOnProductionCDN(BuildPatchToolStagingInfo stagingInfo);
+
         /// <summary>
         /// Copies chunks from a staged location to the production CDN.
         /// </summary>
