Added detail to distill failure exception to inform the reader whether the distill search was recursive or not. Should help inform developer if they forgot to specify a recursive search when specifying a directory with no files.

diff --git a/Engine/Source/Programs/AutomationTool/Distiller.cs b/Engine/Source/Programs/AutomationTool/Distiller.cs
--- a/Engine/Source/Programs/AutomationTool/Distiller.cs
+++ b/Engine/Source/Programs/AutomationTool/Distiller.cs
@@ -256,7 +256,7 @@ namespace AutomationTool
             }
             if (FilesToCopy.Count < 1 && !bAllowMissing)
             {
-                throw new AutomationException("Distill {0} did not produce any files.", AbsFile);
+                throw new AutomationException("Distill {0} did not produce any files. Recursive option was {1}.", AbsFile, bRecursive.ToString());
             }
             return FilesToCopy;
         }
