Fix for bad exit code on successful launch of game via UFE command-line
TTP338848
#ue4
#uat

diff --git a/Engine/Source/Programs/AutomationTool/ProcessUtils.cs b/Engine/Source/Programs/AutomationTool/ProcessUtils.cs
--- a/Engine/Source/Programs/AutomationTool/ProcessUtils.cs
+++ b/Engine/Source/Programs/AutomationTool/ProcessUtils.cs
@@ -332,6 +332,10 @@ namespace AutomationTool
 					if (Proc != null)
 					{
 						bHasExited = Proc.HasExited;
+						if (bHasExited)
+						{
+							ExitCode = Proc.ExitCode;
+						}
 					}
 				}
 				return bHasExited; 
