ULD will now properly launch a game on device via the app id
#ue4
#ios

diff --git a/Engine/Source/Developer/LauncherServices/Private/Launcher/LauncherDeployCommands.h b/Engine/Source/Developer/LauncherServices/Private/Launcher/LauncherDeployCommands.h
--- a/Engine/Source/Developer/LauncherServices/Private/Launcher/LauncherDeployCommands.h
+++ b/Engine/Source/Developer/LauncherServices/Private/Launcher/LauncherDeployCommands.h
@@ -55,7 +55,9 @@ public:
 			*StagePath,
 			*InitialMap,
 			*TargetPlatform.PlatformName());
+#if !PLATFORM_WINDOWS
 		CommandLine += FString::Printf(TEXT(" -device=\"%s\""), *DeviceProxy->GetDeviceId());
+#endif
 
 		// cook dependency arguments
 		CommandLine += CookCommand.IsValid() ? CookCommand->GetDependencyArguments(ChainState) : TEXT(" -skipcook");
diff --git a/Engine/Source/Programs/AutomationTool/IOS/IOSPlatform.Automation.cs b/Engine/Source/Programs/AutomationTool/IOS/IOSPlatform.Automation.cs
--- a/Engine/Source/Programs/AutomationTool/IOS/IOSPlatform.Automation.cs
+++ b/Engine/Source/Programs/AutomationTool/IOS/IOSPlatform.Automation.cs
@@ -616,7 +616,7 @@ public class IOSPlatform : Platform
 			var IPPExe = CombinePaths(CmdEnv.LocalRoot, "Engine/Binaries/DotNET/IOS/IPhonePackager.exe");
 
 			// check for it in the stage directory
-			RunAndLog(CmdEnv, IPPExe, "Deploy \"" + Path.GetFullPath(StagedIPA) + "\" -device " + Params.Device.Substring(4));
+			RunAndLog(CmdEnv, IPPExe, "Deploy \"" + Path.GetFullPath(StagedIPA) + "\"" + (String.IsNullOrEmpty(Params.Device) ? "" : " -device " + Params.Device.Substring(4)));
 		}
         PrintRunTime();
     }
@@ -685,6 +685,30 @@ public class IOSPlatform : Platform
 		}
 		else
 		{
+			// get the CFBundleIdentifier and modify the command line
+			int Pos = ClientCmdLine.IndexOf("-Exe=") + 6;
+			int EndPos = ClientCmdLine.IndexOf("-Targetplatform=") - 2;
+			string Exe = ClientCmdLine.Substring(Pos, EndPos - Pos);
+			
+			// check for Info.plist
+			if (string.IsNullOrEmpty(Params.StageDirectoryParam))
+			{
+				// need to crack open the ipa and read it from there - todo
+			}
+			else
+			{
+				if (File.Exists(Params.BaseStageDirectory+"/IOS/Info.plist"))
+				{
+					string Contents = File.ReadAllText(Params.BaseStageDirectory + "/IOS/Info.plist");
+					Pos = Contents.IndexOf("CFBundleIdentifier");
+					Pos = Contents.IndexOf("<string>", Pos) + 8;
+					EndPos = Contents.IndexOf("</string>", Pos);
+					string id = Contents.Substring(Pos, EndPos - Pos);
+					id = id.Substring(id.LastIndexOf(".") + 1);
+					ClientCmdLine = ClientCmdLine.Replace(Exe, id + ".stub");
+				}
+			}
+
 			return base.RunClient(ClientRunFlags, ClientApp, ClientCmdLine, Params);
 		}
 	}
