fix for iOS staging movies to the correct spot
#ue4
#uat
#ios

diff --git a/Engine/Source/Programs/AutomationTool/IOS/IOSPlatform.Automation.cs b/Engine/Source/Programs/AutomationTool/IOS/IOSPlatform.Automation.cs
--- a/Engine/Source/Programs/AutomationTool/IOS/IOSPlatform.Automation.cs
+++ b/Engine/Source/Programs/AutomationTool/IOS/IOSPlatform.Automation.cs
@@ -543,6 +543,7 @@ public class IOSPlatform : Platform
 		// copy the movies from the project
 		{
 			SC.StageFiles(StagedFileType.NonUFS, CombinePaths(SC.ProjectRoot, "Build/IOS/Resources/Movies"), "*", false, null, "", true, false);
+			SC.StageFiles(StagedFileType.NonUFS, CombinePaths(SC.ProjectRoot, "Content/Movies"), "*", true, null, "", true, false);
 		}
 	}
 
@@ -713,6 +714,11 @@ public class IOSPlatform : Platform
 		}
 	}
 
+	public override bool StageMovies
+	{
+		get { return false; }
+	}
+
 	#region Hooks
 
 	public override void PreBuildAgenda(UE4Build Build, UE4Build.BuildAgenda Agenda)
diff --git a/Engine/Source/Programs/AutomationTool/Platform.cs b/Engine/Source/Programs/AutomationTool/Platform.cs
--- a/Engine/Source/Programs/AutomationTool/Platform.cs
+++ b/Engine/Source/Programs/AutomationTool/Platform.cs
@@ -300,6 +300,14 @@ namespace AutomationTool
             return new List<string>();
         }
 
+		/// <summary>
+		/// Remaps movie directory for platforms that need a remap
+		/// </summary>
+		public virtual bool StageMovies
+		{
+			get { return true; }
+		}
+
 		/// <summary>
 		/// UnrealTargetPlatform type for this platform.
 		/// </summary>
diff --git a/Engine/Source/Programs/AutomationTool/Scripts/CopyBuildToStagingDirectory.Automation.cs b/Engine/Source/Programs/AutomationTool/Scripts/CopyBuildToStagingDirectory.Automation.cs
--- a/Engine/Source/Programs/AutomationTool/Scripts/CopyBuildToStagingDirectory.Automation.cs
+++ b/Engine/Source/Programs/AutomationTool/Scripts/CopyBuildToStagingDirectory.Automation.cs
@@ -240,7 +240,10 @@ public partial class Project : CommandUtils
 				StagedFileTypeForMovies = StagedFileType.UFS;
 			}
 
-			SC.StageFiles(StagedFileTypeForMovies, CombinePaths(SC.ProjectRoot, "Content/Movies"), "*", true, null, CombinePaths(SC.RelativeProjectRootForStage, "Content/Movies"), true, !Params.Pak);
+			if (SC.StageTargetPlatform.StageMovies)
+			{
+				SC.StageFiles(StagedFileTypeForMovies, CombinePaths(SC.ProjectRoot, "Content/Movies"), "*", true, null, CombinePaths(SC.RelativeProjectRootForStage, "Content/Movies"), true, !Params.Pak);
+			}
 
 			// eliminate the sand box
 			SC.StageFiles(StagedFileType.UFS, CombinePaths(SC.ProjectRoot, "Saved", "Cooked", SC.CookPlatform), "*", true, null, "", true, !Params.Pak);
