build script, fix UAT warning

diff --git a/Engine/Source/Programs/AutomationTool/P4Utils.cs b/Engine/Source/Programs/AutomationTool/P4Utils.cs
--- a/Engine/Source/Programs/AutomationTool/P4Utils.cs
+++ b/Engine/Source/Programs/AutomationTool/P4Utils.cs
@@ -702,7 +702,6 @@ namespace AutomationTool
 								int SummaryChangeAt = Line.IndexOf(MatchChange);
 								int SummaryOnAt = Line.IndexOf(MatchOn);
 								int SummaryByAt = Line.IndexOf(MatchBy);
-								int SummaryAtAt = Line.IndexOf("@");
 								if (SummaryChangeAt == 0 && SummaryOnAt > SummaryChangeAt && SummaryByAt > SummaryOnAt)
 								{
 									// OK, we found a new change. This isn't part of our summary.  We're done with the summary.  Back we go.
