fix for UAT and the ability to build non-Widows programs
#ue4
#uat

diff --git a/Engine/Source/Programs/AutomationTool/ProjectParams.cs b/Engine/Source/Programs/AutomationTool/ProjectParams.cs
--- a/Engine/Source/Programs/AutomationTool/ProjectParams.cs
+++ b/Engine/Source/Programs/AutomationTool/ProjectParams.cs
@@ -1184,6 +1184,18 @@ namespace AutomationTool
 					ProgramTarget = TargetData.TargetName;
 				}
 			}
+			else if (!CommandUtils.IsNullOrEmpty(Properties.Programs))
+			{
+				SingleTargetProperties TargetData = Properties.Programs[0];
+
+				bDebugBuildsActuallyUseDebugCRT = TargetData.Rules.bDebugBuildsActuallyUseDebugCRT;
+				bUsesSlate = TargetData.Rules.bUsesSlate;
+				bUsesSlateEditorStyle = TargetData.Rules.bUsesSlateEditorStyle;
+				bUsesSteam = TargetData.Rules.bUsesSteam;
+				ProjectType = TargetRules.TargetType.Program;
+				ProgramTarget = TargetData.TargetName;
+				GameTarget = TargetData.TargetName;
+			}
 			else if (!this.Build)
 			{
 				var ShortName = ProjectUtils.GetShortProjectName(RawProjectPath);
@@ -1221,7 +1233,7 @@ namespace AutomationTool
 
 			if (ProgramTargetsList == null)
 			{
-				if (!GlobalCommandLine.NoCompile && (ProjectType == TargetRules.TargetType.Program))
+				if (ProjectType == TargetRules.TargetType.Program)
 				{
 					ProgramTargetsList = new ParamList<string>(ProgramTarget);
 				}
diff --git a/Engine/Source/Programs/AutomationTool/ProjectUtils.cs b/Engine/Source/Programs/AutomationTool/ProjectUtils.cs
--- a/Engine/Source/Programs/AutomationTool/ProjectUtils.cs
+++ b/Engine/Source/Programs/AutomationTool/ProjectUtils.cs
@@ -157,7 +157,7 @@ namespace AutomationTool
 			if (CommandUtils.CmdEnv.HasCapabilityToCompile)
 			{
 				DetectTargetsForProject(Properties);
-				Properties.bIsCodeBasedProject = !CommandUtils.IsNullOrEmpty(Properties.Targets);
+				Properties.bIsCodeBasedProject = !CommandUtils.IsNullOrEmpty(Properties.Targets) || !CommandUtils.IsNullOrEmpty(Properties.Programs);
 			}
 			else
 			{
