build script, use correct env var for getting p4 authentication token

diff --git a/Engine/Source/Programs/AutomationTool/P4Utils.cs b/Engine/Source/Programs/AutomationTool/P4Utils.cs
--- a/Engine/Source/Programs/AutomationTool/P4Utils.cs
+++ b/Engine/Source/Programs/AutomationTool/P4Utils.cs
@@ -551,7 +551,7 @@ namespace AutomationTool
 			string AuthenticationToken = null;
 
 			string Output;
-			string P4Passwd = InternalUtils.GetEnvironmentVariable("P4PASSWD", "", true) + '\n';
+			string P4Passwd = InternalUtils.GetEnvironmentVariable("uebp_PASS", "", true) + '\n';
 			P4Output(out Output, "login -a -p", P4Passwd);
 
 			// Validate output.
