build script, add usertimeindex param

diff --git a/Engine/Source/Programs/AutomationTool/Scripts/GUBP.Automation.cs b/Engine/Source/Programs/AutomationTool/Scripts/GUBP.Automation.cs
--- a/Engine/Source/Programs/AutomationTool/Scripts/GUBP.Automation.cs
+++ b/Engine/Source/Programs/AutomationTool/Scripts/GUBP.Automation.cs
@@ -4345,6 +4345,7 @@ public class GUBP : BuildCommand
     [Help("Store=", "Sets the name of the temp storage block, normally, this is built for you.")]
     [Help("StoreSuffix=", "Tacked onto a store name constructed from CL, branch, etc")]
     [Help("TimeIndex=", "An integer used to determine subsets to run based on DependentCISFrequencyQuantumShift")]
+    [Help("UserTimeIndex=", "An integer used to determine subsets to run based on DependentCISFrequencyQuantumShift, this one overrides TimeIndex")]
     [Help("Node=", "Nodes to process, -node=Node1+Node2+Node3, if no nodes or games are specified, defaults to all nodes.")]
     [Help("SetupNode=", "Like -Node, but only applies with CommanderJobSetupOnly")]
     [Help("RelatedToNode=", "Nodes to process, -RelatedToNode=Node1+Node2+Node3, use all nodes that either depend on these nodes or these nodes depend on them.")]
@@ -4452,6 +4453,10 @@ public class GUBP : BuildCommand
         bFake = ParseParam("fake");
         bool bFakeEC = ParseParam("FakeEC");
         TimeIndex = ParseParamInt("TimeIndex", 0);
+        if (TimeIndex == 0)
+        {
+            TimeIndex = ParseParamInt("UserTimeIndex", 0);
+        }
 
         bNoIOSOnPC = HostPlatforms.Contains(UnrealTargetPlatform.Mac);
 
