build script, disable preflight success emails...because they won't work, no triggers in preflights

diff --git a/Engine/Source/Programs/AutomationTool/Scripts/GUBP.Automation.cs b/Engine/Source/Programs/AutomationTool/Scripts/GUBP.Automation.cs
--- a/Engine/Source/Programs/AutomationTool/Scripts/GUBP.Automation.cs
+++ b/Engine/Source/Programs/AutomationTool/Scripts/GUBP.Automation.cs
@@ -5331,7 +5331,8 @@ if (HostPlatform == UnrealTargetPlatform.Mac) continue; //temp hack till mac aut
                 AddNode(new CleanSharedTempStorageNode(this));
             }
         }
-
+#if false
+        // this doesn't work for lots of reasons...we can't figure out what the dependencies are until far later
         if (bPreflightBuild)
         {
             GeneralSuccessNode PreflightSuccessNode = new GeneralSuccessNode("Preflight");
@@ -5344,7 +5345,7 @@ if (HostPlatform == UnrealTargetPlatform.Mac) continue; //temp hack till mac aut
             }
             AddNode(PreflightSuccessNode);
         }
-
+#endif
         foreach (var NodeToDo in GUBPNodes)
         {
             foreach (var Dep in GUBPNodes[NodeToDo.Key].FullNamesOfDependencies)
@@ -5676,7 +5677,6 @@ if (HostPlatform == UnrealTargetPlatform.Mac) continue; //temp hack till mac aut
             NodesToDo = NewNodesToDo;
         }
 
-        string FakeFail = ParseParamValue("FakeFail");
         if (CommanderSetup)
         {
             if (!String.IsNullOrEmpty(ExplicitTrigger))
@@ -5713,6 +5713,25 @@ if (HostPlatform == UnrealTargetPlatform.Mac) continue; //temp hack till mac aut
                 }
             }
         }
+        if (bPreflightBuild)
+        {
+            Log("Culling triggers and downstream for preflight builds ");
+            var NewNodesToDo = new HashSet<string>();
+            foreach (var NodeToDo in NodesToDo)
+            {
+                var TriggerDot = GetControllingTriggerDotName(NodeToDo);
+                if (TriggerDot == "" && !GUBPNodes[NodeToDo].TriggerNode())
+                {
+                    Log("  Keeping {0}", NodeToDo);
+                    NewNodesToDo.Add(NodeToDo);
+                }
+                else
+                {
+                    Log("  Rejecting {0}", NodeToDo);
+                }
+            }
+            NodesToDo = NewNodesToDo;
+        }
 
         GUBPNodesCompleted = new Dictionary<string, bool>();
         GUBPNodesHistory = new Dictionary<string, NodeHistory>();
@@ -5793,8 +5812,10 @@ if (HostPlatform == UnrealTargetPlatform.Mac) continue; //temp hack till mac aut
             }
         }
 
+        string FakeFail = ParseParamValue("FakeFail");
         if (CommanderSetup)
         {
+
             if (OrdereredToDo.Count == 0)
             {
                 throw new AutomationException("No nodes to do!");
