build script, BuildCookRun can now find projects

diff --git a/Engine/Source/Programs/AutomationTool/Scripts/BuildCookRun.Automation.cs b/Engine/Source/Programs/AutomationTool/Scripts/BuildCookRun.Automation.cs
--- a/Engine/Source/Programs/AutomationTool/Scripts/BuildCookRun.Automation.cs
+++ b/Engine/Source/Programs/AutomationTool/Scripts/BuildCookRun.Automation.cs
@@ -323,7 +323,8 @@ public class BuildCookRun : BuildCommand
 				}
 				else
 				{
-                    var ProjectName = ParseParamValue("project", CombinePaths("Samples", "Sandbox", "BlankProject", "BlankProject.uproject"));
+                    var OriginalProjectName = ParseParamValue("project", CombinePaths("Samples", "Sandbox", "BlankProject", "BlankProject.uproject"));
+                    var ProjectName = OriginalProjectName;
 					ProjectName = ProjectName.Trim (new char[]{ '\"' });
                     if (ProjectName.IndexOfAny(new char[] { '\\', '/' }) < 0)
                     {
@@ -334,6 +335,15 @@ public class BuildCookRun : BuildCommand
                         ProjectName = CombinePaths(CmdEnv.LocalRoot, ProjectName);
                     }
                     if (!FileExists_NoExceptions(ProjectName))
+                    {
+                        var Branch = new BranchInfo(new List<UnrealTargetPlatform> { UnrealBuildTool.Utils.IsRunningOnMono ? UnrealTargetPlatform.Mac : UnrealTargetPlatform.Win64 });
+                        var GameProj = Branch.FindGame(OriginalProjectName);
+                        if (GameProj != null)
+                        {
+                            ProjectName = GameProj.FilePath;
+                        }
+                    }
+                    if (!FileExists_NoExceptions(ProjectName))
                     {
                         throw new AutomationException("Could not find a project file {0}.", ProjectName);
                     }
