build script, exclude a few things from red-after-red comparison

diff --git a/Engine/Source/Programs/AutomationTool/Scripts/GUBP.Automation.cs b/Engine/Source/Programs/AutomationTool/Scripts/GUBP.Automation.cs
--- a/Engine/Source/Programs/AutomationTool/Scripts/GUBP.Automation.cs
+++ b/Engine/Source/Programs/AutomationTool/Scripts/GUBP.Automation.cs
@@ -58,7 +58,10 @@ public class ECJobPropsUtils
                                                             }
                                                             if (IsMessage && PropertySheetChildDiagSheetElemInner2.Name == "value")
                                                             {
-                                                                if (!PropertySheetChildDiagSheetElemInner2.InnerText.Contains("LogTailsAndChanges"))
+                                                                if (!PropertySheetChildDiagSheetElemInner2.InnerText.Contains("LogTailsAndChanges")
+                                                                    && !PropertySheetChildDiagSheetElemInner2.InnerText.Contains("-MyJobStepId=")
+                                                                    && !PropertySheetChildDiagSheetElemInner2.InnerText.Contains("CommandUtils.Run: Run: Took ")
+                                                                    )
                                                                 {
                                                                     Result.Add(PropertySheetChildDiagSheetElemInner2.InnerText);
                                                                 }
