- Fixed another mono compile error

diff --git a/Engine/Source/Programs/AutomationTool/IOS/IPhonePackager/GenerateSigningRequestDialog.cs b/Engine/Source/Programs/AutomationTool/IOS/IPhonePackager/GenerateSigningRequestDialog.cs
--- a/Engine/Source/Programs/AutomationTool/IOS/IPhonePackager/GenerateSigningRequestDialog.cs
+++ b/Engine/Source/Programs/AutomationTool/IOS/IPhonePackager/GenerateSigningRequestDialog.cs
@@ -53,8 +53,6 @@ namespace iPhonePackager
 				return;
 			}
 
-			string EffectiveBuildPath = (Program.GameName.Length > 0) ? Config.BuildDirectory : Path.GetFullPath(".");
-
 			// Create a temporary file to write the key pair out to (in a format that OpenSSL understands)
 			string KeyFileName = Path.GetTempFileName();
 			TextWriter KeyWriter = new StreamWriter(KeyFileName);
@@ -161,7 +159,6 @@ namespace iPhonePackager
 			SaveDialog.Title = "Generate Certificate Signing Request";
 			if (SaveDialog.ShowDialog() == DialogResult.OK)
 			{
-				string EffectiveBuildPath = (Program.GameName.Length > 0) ? Config.BuildDirectory : Path.GetFullPath(".");
 				string TargetCertRequestFileName = SaveDialog.FileName;
 
 				GenerateSigningRequestViaOpenSSL(TargetCertRequestFileName, KeyPair);
