- Backing out the IPP integration into IOS.Automation - caused too many problems on Mac under Mono. Fixing tommorrow, hopefully.

diff --git a/Engine/Source/Programs/AutomationTool/IOS/IOS.Automation.csproj b/Engine/Source/Programs/AutomationTool/IOS/IOS.Automation.csproj
--- a/Engine/Source/Programs/AutomationTool/IOS/IOS.Automation.csproj
+++ b/Engine/Source/Programs/AutomationTool/IOS/IOS.Automation.csproj
@@ -34,27 +34,15 @@
     <TreatWarningsAsErrors>true</TreatWarningsAsErrors>
   </PropertyGroup>
   <ItemGroup>
-    <Reference Include="BouncyCastle.Crypto">
-      <HintPath>..\..\..\..\Binaries\ThirdParty\IOS\BouncyCastle.Crypto.dll</HintPath>
-      <Private>False</Private>
-    </Reference>
     <Reference Include="Ionic.Zip.Reduced">
       <HintPath>..\..\..\..\Binaries\DotNET\Ionic.Zip.Reduced.dll</HintPath>
       <Private>False</Private>
     </Reference>
-    <Reference Include="RPCUtility">
-      <HintPath>..\..\..\..\Binaries\DotNET\RPCUtility.exe</HintPath>
-      <Private>False</Private>
-    </Reference>
     <Reference Include="System" />
     <Reference Include="System.Core" />
-    <Reference Include="System.Drawing" />
-    <Reference Include="System.Runtime.Remoting" />
     <Reference Include="System.Runtime.Serialization" />
-    <Reference Include="System.Security" />
     <Reference Include="System.Web" />
     <Reference Include="System.Web.Extensions" />
-    <Reference Include="System.Windows.Forms" />
     <Reference Include="System.Xml.Linq" />
     <Reference Include="System.Data.DataSetExtensions" />
     <Reference Include="Microsoft.CSharp" />
@@ -66,60 +54,9 @@
       <Link>Properties\MetaData.cs</Link>
     </Compile>
     <Compile Include="IOSPlatform.Automation.cs" />
-    <Compile Include="IPhonePackager\CodeSigning.cs" />
-    <Compile Include="IPhonePackager\CompileTime.cs" />
-    <Compile Include="IPhonePackager\Config.cs" />
-    <Compile Include="IPhonePackager\ConfigureMobileGame.cs">
-      <SubType>Form</SubType>
-    </Compile>
-    <Compile Include="IPhonePackager\ConfigureMobileGame.Designer.cs">
-      <DependentUpon>ConfigureMobileGame.cs</DependentUpon>
-    </Compile>
-    <Compile Include="IPhonePackager\CookTime.cs" />
-    <Compile Include="IPhonePackager\DeployTime.cs" />
-    <Compile Include="IPhonePackager\FileOperations.cs" />
-    <Compile Include="IPhonePackager\GenerateSigningRequestDialog.cs">
-      <SubType>Form</SubType>
-    </Compile>
-    <Compile Include="IPhonePackager\GenerateSigningRequestDialog.Designer.cs">
-      <DependentUpon>GenerateSigningRequestDialog.cs</DependentUpon>
-    </Compile>
-    <Compile Include="IPhonePackager\GraphicalResignTool.cs">
-      <SubType>Form</SubType>
-    </Compile>
-    <Compile Include="IPhonePackager\GraphicalResignTool.Designer.cs">
-      <DependentUpon>GraphicalResignTool.cs</DependentUpon>
-    </Compile>
-    <Compile Include="IPhonePackager\iPhonePackager.cs" />
-    <Compile Include="IPhonePackager\MachObjects.cs" />
-    <Compile Include="IPhonePackager\MobileProvisionUtilities.cs" />
-    <Compile Include="IPhonePackager\PasswordDialog.cs">
-      <SubType>Form</SubType>
-    </Compile>
-    <Compile Include="IPhonePackager\PasswordDialog.Designer.cs">
-      <DependentUpon>PasswordDialog.cs</DependentUpon>
-    </Compile>
-    <Compile Include="IPhonePackager\SlowProgressDialog.cs">
-      <SubType>Form</SubType>
-    </Compile>
-    <Compile Include="IPhonePackager\SlowProgressDialog.Designer.cs">
-      <DependentUpon>SlowProgressDialog.cs</DependentUpon>
-    </Compile>
-    <Compile Include="IPhonePackager\ToolsHub.cs">
-      <SubType>Form</SubType>
-    </Compile>
-    <Compile Include="IPhonePackager\ToolsHub.Designer.cs">
-      <DependentUpon>ToolsHub.cs</DependentUpon>
-    </Compile>
-    <Compile Include="IPhonePackager\Utilities.cs" />
     <Compile Include="Properties\AssemblyInfo.cs" />
-    <Compile Include="Properties\Resources.Designer.cs" />
   </ItemGroup>
   <ItemGroup>
-    <ProjectReference Include="..\..\IOS\DeploymentInterface\DeploymentInterface.csproj">
-      <Project>{2da5397a-f93f-41e0-a1ad-b220d40fd739}</Project>
-      <Name>DeploymentInterface</Name>
-    </ProjectReference>
     <ProjectReference Include="..\AutomationTool.csproj">
       <Project>{a35b0261-d868-475e-a098-eb14ce6dc89f}</Project>
       <Name>AutomationTool</Name>
@@ -131,31 +68,5 @@
       <Private>False</Private>
     </ProjectReference>
   </ItemGroup>
-  <ItemGroup>
-    <EmbeddedResource Include="IPhonePackager\ConfigureMobileGame.resx">
-      <DependentUpon>ConfigureMobileGame.cs</DependentUpon>
-    </EmbeddedResource>
-    <EmbeddedResource Include="IPhonePackager\GenerateSigningRequestDialog.resx">
-      <DependentUpon>GenerateSigningRequestDialog.cs</DependentUpon>
-    </EmbeddedResource>
-    <EmbeddedResource Include="IPhonePackager\GraphicalResignTool.resx">
-      <DependentUpon>GraphicalResignTool.cs</DependentUpon>
-    </EmbeddedResource>
-    <EmbeddedResource Include="IPhonePackager\PasswordDialog.resx">
-      <DependentUpon>PasswordDialog.cs</DependentUpon>
-    </EmbeddedResource>
-    <EmbeddedResource Include="IPhonePackager\SlowProgressDialog.resx">
-      <DependentUpon>SlowProgressDialog.cs</DependentUpon>
-    </EmbeddedResource>
-    <EmbeddedResource Include="IPhonePackager\ToolsHub.resx">
-      <DependentUpon>ToolsHub.cs</DependentUpon>
-    </EmbeddedResource>
-    <EmbeddedResource Include="Properties\Resources.resx">
-      <SubType>Designer</SubType>
-    </EmbeddedResource>
-  </ItemGroup>
-  <ItemGroup>
-    <None Include="Resources\Image1.bmp" />
-  </ItemGroup>
   <Import Project="$(MSBuildToolsPath)\Microsoft.CSharp.targets" />
 </Project>
\ No newline at end of file
diff --git a/Engine/Source/Programs/AutomationTool/IOS/IOSPlatform.Automation.cs b/Engine/Source/Programs/AutomationTool/IOS/IOSPlatform.Automation.cs
--- a/Engine/Source/Programs/AutomationTool/IOS/IOSPlatform.Automation.cs
+++ b/Engine/Source/Programs/AutomationTool/IOS/IOSPlatform.Automation.cs
@@ -6,7 +6,6 @@ using System.Text;
 using System.IO;
 using AutomationTool;
 using UnrealBuildTool;
-using System.Text.RegularExpressions;
 using Ionic.Zip;
 using Ionic.Zlib;
 
@@ -23,65 +22,6 @@ public class IOSPlatform : Platform
 	{
 	}
 
-	// Run the integrated IPP code
-	// @todo ipp: How to log the output in a nice UAT way?
-	protected static int RunIPP(string IPPArguments)
-	{
-		List<string> Args = new List<string>();
-
-		StringBuilder Token = null;
-		int Index = 0;
-		bool bInQuote = false;
-		while (Index < IPPArguments.Length)
-		{
-			if (IPPArguments[Index] == '\"')
-			{
-				bInQuote = !bInQuote;
-			}
-			else if (IPPArguments[Index] == ' ' && !bInQuote)
-			{
-				if (Token != null)
-				{
-					Args.Add(Token.ToString());
-					Token = null;
-				}
-			}
-			else
-			{
-				if (Token == null)
-				{
-					Token = new StringBuilder();
-				}
-				Token.Append(IPPArguments[Index]);
-			}
-			Index++;
-		}
-
-		if (Token != null)
-		{
-			Args.Add(Token.ToString());
-		}
-	
-		return RunIPP(Args.ToArray());
-	}
-
-	protected static int RunIPP(string[] Args)
-	{
-		return iPhonePackager.Program.RealMain(Args);
-	}
-
-	public override int RunCommand(string Command)
-	{
-		// run generic IPP commands through the interface
-		if (Command.StartsWith("IPP:"))
-		{
-			return RunIPP(Command.Substring(4));
-		}
-
-		// non-zero error code
-		return 4;
-	}
-
 	protected string MakeIPAFileName( UnrealTargetConfiguration TargetConfiguration, string ProjectGameExeFilename )
 	{
 		var ProjectIPA = Path.ChangeExtension(ProjectGameExeFilename, null);
@@ -184,93 +124,47 @@ public class IOSPlatform : Platform
 
 			bool cookonthefly = Params.CookOnTheFly || Params.SkipCookOnTheFly;
 
-			// rename the .ipa if not code based
-			if (!Params.IsCodeBasedProject)
+			string IPPArguments = "RepackageFromStage \"" + (Params.IsCodeBasedProject ? Params.RawProjectPath : "Engine") + "\"";
+			IPPArguments += " -config " + TargetConfiguration.ToString();
+
+			if (TargetConfiguration == UnrealTargetConfiguration.Shipping)
 			{
-				ProjectIPA = Path.Combine(Path.GetDirectoryName(Params.RawProjectPath), "Binaries", "IOS", Params.ShortProjectName + ".ipa");
-				if (TargetConfiguration != UnrealTargetConfiguration.Development)
-				{
-					ProjectIPA = Path.Combine(Path.GetDirectoryName(Params.RawProjectPath), "Binaries", "IOS", Params.ShortProjectName + "-" + PlatformType.ToString() + "-" + TargetConfiguration.ToString() + ".ipa");
-				}
+				IPPArguments += " -compress=best";
 			}
 
-			// delete the .ipa to make sure it was made
-			DeleteFile(ProjectIPA);
+			// Determine if we should sign
+			bool bNeedToSign = GetCodeSignDesirability(Params);
 
-			if (RemoteToolChain.bUseRPCUtil)
+			if (!String.IsNullOrEmpty(Params.BundleName))
 			{
-				string IPPArguments = "RepackageFromStage \"" + (Params.IsCodeBasedProject ? Params.RawProjectPath : "Engine") + "\"";
-				IPPArguments += " -config " + TargetConfiguration.ToString();
-
-				if (TargetConfiguration == UnrealTargetConfiguration.Shipping)
-				{
-					IPPArguments += " -compress=best";
-				}
-
-				// Determine if we should sign
-				bool bNeedToSign = GetCodeSignDesirability(Params);
-
-				if (!String.IsNullOrEmpty(Params.BundleName))
-				{
-					// Have to sign when a bundle name is specified
-					bNeedToSign = true;
-					IPPArguments += " -bundlename " + Params.BundleName;
-				}
-
-				if (bNeedToSign)
-				{
-					IPPArguments += " -sign";
-				}
-
-				IPPArguments += (cookonthefly ? " -cookonthefly" : "");
-				IPPArguments += " -stagedir \"" + CombinePaths(Params.BaseStageDirectory, "IOS") + "\"";
-				IPPArguments += " -projectdir \"" + Path.GetDirectoryName(Params.RawProjectPath) + "\"";
-
-				RunAndLog(CmdEnv, IPPExe, IPPArguments);
+				// Have to sign when a bundle name is specified
+				bNeedToSign = true;
+				IPPArguments += " -bundlename " + Params.BundleName;
 			}
-			else
-			{
-				List<string> IPPArguments = new List<string>();
-				IPPArguments.Add("RepackageFromStage");
-				IPPArguments.Add(Params.IsCodeBasedProject ? Params.RawProjectPath : "Engine");
-				IPPArguments.Add("-config");
-				IPPArguments.Add(TargetConfiguration.ToString());
 
-				if (TargetConfiguration == UnrealTargetConfiguration.Shipping)
-				{
-					IPPArguments.Add("-compress=best");
-				}
-
-				// Determine if we should sign
-				bool bNeedToSign = GetCodeSignDesirability(Params);
+			if (bNeedToSign)
+			{
+				IPPArguments += " -sign";
+			}
 
-				if (!String.IsNullOrEmpty(Params.BundleName))
-				{
-					// Have to sign when a bundle name is specified
-					bNeedToSign = true;
-					IPPArguments.Add("-bundlename");
-					IPPArguments.Add(Params.BundleName);
-				}
+			IPPArguments += (cookonthefly ? " -cookonthefly" : "");
+			IPPArguments += " -stagedir \"" + CombinePaths(Params.BaseStageDirectory, "IOS") + "\"";
+			IPPArguments += " -projectdir \"" + Path.GetDirectoryName(Params.RawProjectPath) + "\"";
 
-				if (bNeedToSign)
+			// rename the .ipa if not code based
+			if (!Params.IsCodeBasedProject)
+			{
+				ProjectIPA = Path.Combine(Path.GetDirectoryName(Params.RawProjectPath), "Binaries", "IOS", Params.ShortProjectName + ".ipa");
+				if (TargetConfiguration != UnrealTargetConfiguration.Development)
 				{
-					IPPArguments.Add("-sign");
+					ProjectIPA = Path.Combine(Path.GetDirectoryName(Params.RawProjectPath), "Binaries", "IOS", Params.ShortProjectName + "-" + PlatformType.ToString() + "-" + TargetConfiguration.ToString() + ".ipa");
 				}
+			}
 
-				if (cookonthefly)
-				{
-					IPPArguments.Add(" -cookonthefly");
-				}
-				IPPArguments.Add(" -stagedir");
-				IPPArguments.Add(CombinePaths(Params.BaseStageDirectory, "IOS"));
-				IPPArguments.Add(" -projectdir");
-				IPPArguments.Add(Path.GetDirectoryName(Params.RawProjectPath));
+            // delete the .ipa to make sure it was made
+            DeleteFile(ProjectIPA);
 
-				if (RunIPP(IPPArguments.ToArray()) != 0)
-				{
-					throw new AutomationException("IPP Failed");
-				}
-			}
+			RunAndLog(CmdEnv, IPPExe, IPPArguments);
 
 			// verify the .ipa exists
 			if (!FileExists(ProjectIPA))
