build script, mangle formal build archive directories in preflights

diff --git a/Engine/Source/Programs/AutomationTool/CommandUtils.cs b/Engine/Source/Programs/AutomationTool/CommandUtils.cs
--- a/Engine/Source/Programs/AutomationTool/CommandUtils.cs
+++ b/Engine/Source/Programs/AutomationTool/CommandUtils.cs
@@ -2060,7 +2060,8 @@ namespace AutomationTool
                         if (JustDir.StartsWith(StartString, StringComparison.InvariantCultureIgnoreCase) && (String.IsNullOrEmpty(EndString) || JustDir.EndsWith(EndString, StringComparison.InvariantCultureIgnoreCase)))
                         {
                             string CLPart = JustDir.Substring(StartString.Length, JustDir.Length - StartString.Length - EndString.Length);
-                            if (!CLPart.Contains("-") && !CLPart.Contains("+"))
+                            if (CLPart.Contains("-PF-") || // always delete preflights
+                                (!CLPart.Contains("-") && !CLPart.Contains("+"))) // never delete anything else that is weird, this is probably another branch
                             {
                                 DirectoryInfo ThisDirInfo = new DirectoryInfo(TopLevelDir.FullName);
                                 bool bOld = false;
diff --git a/Engine/Source/Programs/AutomationTool/Scripts/GUBP.Automation.cs b/Engine/Source/Programs/AutomationTool/Scripts/GUBP.Automation.cs
--- a/Engine/Source/Programs/AutomationTool/Scripts/GUBP.Automation.cs
+++ b/Engine/Source/Programs/AutomationTool/Scripts/GUBP.Automation.cs
@@ -2615,7 +2615,12 @@ public class GUBP : BuildCommand
         {
             string BaseDir = ResolveSharedBuildDirectory(InGameProj.GameName);
             string NodeName = StaticGetFullName(InGameProj, InHostPlatform, InClientTargetPlatforms, InClientConfigs, InServerTargetPlatforms, InServerConfigs, InClientNotGame);
-            string ArchiveDirectory = CombinePaths(BaseDir, NodeName, P4Env.BuildRootEscaped + "-CL-" + P4Env.ChangelistString);
+            string Inner = P4Env.BuildRootEscaped + "-CL-" + P4Env.ChangelistString;
+            if (GUBP.bPreflightBuild)
+            {
+                Inner = Inner + PreflightMangleSuffix;
+            }
+            string ArchiveDirectory = CombinePaths(BaseDir, NodeName, Inner);
             return ArchiveDirectory;
         }
         public override void DoBuild(GUBP bp)
