build script, increase frequency of non-unity test compiles

diff --git a/Engine/Source/Programs/AutomationTool/Scripts/GUBP.Automation.cs b/Engine/Source/Programs/AutomationTool/Scripts/GUBP.Automation.cs
--- a/Engine/Source/Programs/AutomationTool/Scripts/GUBP.Automation.cs
+++ b/Engine/Source/Programs/AutomationTool/Scripts/GUBP.Automation.cs
@@ -2860,7 +2860,7 @@ public class GUBP : BuildCommand
         }
         public override int CISFrequencyQuantumShift(GUBP bp)
         {
-            return base.CISFrequencyQuantumShift(bp) + 3;
+            return base.CISFrequencyQuantumShift(bp) + 2;
         }
         public override void DoTest(GUBP bp)
         {
