- Fixed a path for the integrated IPP

diff --git a/Engine/Source/Programs/AutomationTool/IOS/IPhonePackager/Config.cs b/Engine/Source/Programs/AutomationTool/IOS/IPhonePackager/Config.cs
--- a/Engine/Source/Programs/AutomationTool/IOS/IPhonePackager/Config.cs
+++ b/Engine/Source/Programs/AutomationTool/IOS/IPhonePackager/Config.cs
@@ -17,7 +17,7 @@ namespace iPhonePackager
 		/// </summary>
 		public static string AppDisplayName = "Unreal iOS Configuration";
 
-		public static string RootRelativePath = @"..\..\..\..\";
+		public static string RootRelativePath = @"";
 
 		public static string GameDirectory = "";		// "..\\..\\..\\..\\Engine\\Source\\UE4";
 		public static string BinariesDirectory = "";	// "..\\..\\..\\..\\Engine\\Binaries\\IOS\\"
