build script, added PreflightUID

diff --git a/Engine/Source/Programs/AutomationTool/Scripts/GUBP.Automation.cs b/Engine/Source/Programs/AutomationTool/Scripts/GUBP.Automation.cs
--- a/Engine/Source/Programs/AutomationTool/Scripts/GUBP.Automation.cs
+++ b/Engine/Source/Programs/AutomationTool/Scripts/GUBP.Automation.cs
@@ -4354,6 +4354,7 @@ public class GUBP : BuildCommand
     [Help("StoreSuffix=", "Tacked onto a store name constructed from CL, branch, etc")]
     [Help("TimeIndex=", "An integer used to determine subsets to run based on DependentCISFrequencyQuantumShift")]
     [Help("UserTimeIndex=", "An integer used to determine subsets to run based on DependentCISFrequencyQuantumShift, this one overrides TimeIndex")]
+    [Help("PreflightUID=", "A unique integer tag from EC used as part of the tempstorage, builds and label names to distinguish multiple attempts.")]
     [Help("Node=", "Nodes to process, -node=Node1+Node2+Node3, if no nodes or games are specified, defaults to all nodes.")]
     [Help("SetupNode=", "Like -Node, but only applies with CommanderJobSetupOnly")]
     [Help("RelatedToNode=", "Nodes to process, -RelatedToNode=Node1+Node2+Node3, use all nodes that either depend on these nodes or these nodes depend on them.")]
@@ -4447,9 +4448,8 @@ public class GUBP : BuildCommand
         }
         if (bPreflightBuild)
         {
-            int TotalSeconds = (int)((DateTime.UtcNow - new DateTime(2014, 6, 1)).TotalSeconds);
-
-            PreflightMangleSuffix = String.Format("-PF-{0}-{1}", PreflightShelveCL, TotalSeconds);
+            int PreflightUID = ParseParamInt("PreflightUID", 0);
+            PreflightMangleSuffix = String.Format("-PF-{0}-{1}", PreflightShelveCL, PreflightUID);
             StoreSuffix = StoreSuffix + PreflightMangleSuffix;
         }
         CL = ParseParamInt("CL", 0);
