Changed UAT's IsProbablyAMacOrIOSExe() to handle binaries in Binaries/ThirdParty

diff --git a/Engine/Source/Programs/AutomationTool/BuildUtils.cs b/Engine/Source/Programs/AutomationTool/BuildUtils.cs
--- a/Engine/Source/Programs/AutomationTool/BuildUtils.cs
+++ b/Engine/Source/Programs/AutomationTool/BuildUtils.cs
@@ -91,7 +91,8 @@ namespace AutomationTool
                 Path.GetExtension(Filename) == ".command" ||
                 ((
                     CommandUtils.CombinePaths(Filename).ToLower().Contains(CommandUtils.CombinePaths("Binaries", "Mac").ToLower()) ||
-                    CommandUtils.CombinePaths(Filename).ToLower().Contains(CommandUtils.CombinePaths("Binaries", "IOS").ToLower())
+                    CommandUtils.CombinePaths(Filename).ToLower().Contains(CommandUtils.CombinePaths("Binaries", "IOS").ToLower()) ||
+					CommandUtils.CombinePaths(Filename).ToLower().Contains(CommandUtils.CombinePaths("Binaries", "ThirdParty").ToLower())
                 ) && (Path.GetExtension(Filename) == "" || Path.GetExtension(Filename) == "."));
         }
 
