- Improved the Android packaged build install batch file to handle errors, and give some suggestions.

diff --git a/Engine/Source/Programs/AutomationTool/Android/AndroidPlatform.Automation.cs b/Engine/Source/Programs/AutomationTool/Android/AndroidPlatform.Automation.cs
--- a/Engine/Source/Programs/AutomationTool/Android/AndroidPlatform.Automation.cs
+++ b/Engine/Source/Programs/AutomationTool/Android/AndroidPlatform.Automation.cs
@@ -156,10 +156,22 @@ public class AndroidPlatform : Platform
             "if not \"%1\"==\"\" set DEVICE=-s %1",
 			"%ADB% %DEVICE% uninstall " + PackageName,
 			"%ADB% %DEVICE% install " + Path.GetFileName(ApkName),
+			"if not ERRORLEVEL 0 goto Error",
 			"%ADB% %DEVICE% shell rm -r /mnt/sdcard/" + Params.ShortProjectName,
 			"%ADB% %DEVICE% shell rm -r /mnt/sdcard/UE4Game/UE4CommandLine.txt", // we need to delete the commandline in UE4Game or it will mess up loading
 			"%ADB% %DEVICE% shell rm -r /mnt/sdcard/obb/" + PackageName,
 			Params.OBBinAPK ? "" : "%ADB% %DEVICE% push " + Path.GetFileName(LocalObbName) + " " + DeviceObbName,
+			Params.OBBinAPK ? "" : "if not ERRORLEVEL 0 goto Error",
+			"goto:eof",
+			":Error",
+			"@echo.",
+			"@echo There was an error installing the game or the obb file. Look above for more info.",
+			"@echo.",
+			"@echo Things to try:",
+			"@echo Check that the device (and only the device) is listed with \"%ADB$ devices\" from a command prompt.",
+			"@echo Make sure all Developer options look normal on the device",
+			"@echo Check that the device has an SD card.",
+			"@pause"
 		};
 		File.WriteAllLines(BatchName, BatchLines);
 
