- Improved error handling in the Android install file

diff --git a/Engine/Source/Programs/AutomationTool/Android/AndroidPlatform.Automation.cs b/Engine/Source/Programs/AutomationTool/Android/AndroidPlatform.Automation.cs
--- a/Engine/Source/Programs/AutomationTool/Android/AndroidPlatform.Automation.cs
+++ b/Engine/Source/Programs/AutomationTool/Android/AndroidPlatform.Automation.cs
@@ -156,12 +156,12 @@ public class AndroidPlatform : Platform
             "if not \"%1\"==\"\" set DEVICE=-s %1",
 			"%ADB% %DEVICE% uninstall " + PackageName,
 			"%ADB% %DEVICE% install " + Path.GetFileName(ApkName),
-			"if not ERRORLEVEL 0 goto Error",
+			"@if \"%ERRORLEVEL%\" NEQ \"0\" goto Error",
 			"%ADB% %DEVICE% shell rm -r /mnt/sdcard/" + Params.ShortProjectName,
 			"%ADB% %DEVICE% shell rm -r /mnt/sdcard/UE4Game/UE4CommandLine.txt", // we need to delete the commandline in UE4Game or it will mess up loading
 			"%ADB% %DEVICE% shell rm -r /mnt/sdcard/obb/" + PackageName,
 			Params.OBBinAPK ? "" : "%ADB% %DEVICE% push " + Path.GetFileName(LocalObbName) + " " + DeviceObbName,
-			Params.OBBinAPK ? "" : "if not ERRORLEVEL 0 goto Error",
+			Params.OBBinAPK ? "" : "if \"%ERRORLEVEL%\" NEQ \"0\" goto Error",
 			"goto:eof",
 			":Error",
 			"@echo.",
