#UE4 #HTML5

  Allow "Packaging(Archiving)" and "Launch" to work from within the Editor. #TTP 340154, #TTP 340157

diff --git a/Engine/Source/Programs/AutomationTool/HTML5/HTML5Platform.Automation.cs b/Engine/Source/Programs/AutomationTool/HTML5/HTML5Platform.Automation.cs
--- a/Engine/Source/Programs/AutomationTool/HTML5/HTML5Platform.Automation.cs
+++ b/Engine/Source/Programs/AutomationTool/HTML5/HTML5Platform.Automation.cs
@@ -16,85 +16,98 @@ public class HTML5Platform : Platform
 
 	public override void Package(ProjectParams Params, DeploymentContext SC, int WorkingCL)
 	{
-		Log("Package {0}", Params.RawProjectPath);
-
-		string PackagePath = Path.Combine(Path.GetDirectoryName(Params.RawProjectPath), "Binaries", "HTML5");
-		if (!Directory.Exists(PackagePath))
-		{
-			Directory.CreateDirectory(PackagePath);
-		}
-		string FinalDataLocation = Path.Combine(PackagePath, Params.ShortProjectName) +".data";
-
-
-		// we need to operate in the root
-		using (new PushedDirectory(Path.Combine(Params.BaseStageDirectory, "HTML5")))
-		{
-			string BaseSDKPath = Environment.GetEnvironmentVariable("EMSCRIPTEN"); 
-            // make the file_packager command line
-			if (Utils.IsRunningOnMono) 
-			{
-				string PackagerPath = BaseSDKPath + "/tools/file_packager.py";
-				string CmdLine = string.Format ("-c \" python {0} '{1}' --preload . --js-output='{1}.js' \" ", PackagerPath, FinalDataLocation);
-				RunAndLog (CmdEnv, "/bin/bash", CmdLine);	
-			} else 
-			{ 
-				string PackagerPath = "\"" + BaseSDKPath + "\\tools\\file_packager.py\"";
-				string CmdLine = string.Format ("/c python {0} \"{1}\" --preload . --js-output=\"{1}.js\"", PackagerPath, FinalDataLocation);
-				RunAndLog (CmdEnv, CommandUtils.CombinePaths (Environment.SystemDirectory, "cmd.exe"), CmdLine);	
-			}
-		}
-
-		// copy the "Executable" to the package directory
-		string GameExe = Path.GetFileNameWithoutExtension(Params.ProjectGameExeFilename);
-		if (Params.ClientConfigsToBuild[0].ToString() != "Development")
-		{
-			GameExe += "-HTML5-" + Params.ClientConfigsToBuild[0].ToString();
-		}
-		GameExe += ".js";
-		if (Path.Combine(Path.GetDirectoryName(Params.ProjectGameExeFilename), GameExe) != Path.Combine(PackagePath, GameExe))
-		{
-			File.Copy(Path.Combine(Path.GetDirectoryName(Params.ProjectGameExeFilename), GameExe), Path.Combine(PackagePath, GameExe), true);
-			File.Copy(Path.Combine(Path.GetDirectoryName(Params.ProjectGameExeFilename), GameExe) + ".mem", Path.Combine(PackagePath, GameExe) + ".mem", true);
-		}
-		File.SetAttributes(Path.Combine(PackagePath, GameExe), FileAttributes.Normal);
-		File.SetAttributes(Path.Combine(PackagePath, GameExe) + ".mem", FileAttributes.Normal);
-		
-		// put the HTML file to the package directory
-		string TemplateFile = Path.Combine(CombinePaths(CmdEnv.LocalRoot, "Engine"), "Build", "HTML5", "Game.html.template");
-		string OutputFile = Path.Combine(PackagePath, (Params.ClientConfigsToBuild[0].ToString() != "Development" ? (Params.ShortProjectName + "-HTML5-" + Params.ClientConfigsToBuild[0].ToString()) : Params.ShortProjectName)) + ".html";
-    	
-		// find Heap Size.
-		ulong HeapSize;
-		var ConfigCache = new UnrealBuildTool.ConfigCacheIni(UnrealTargetPlatform.HTML5, "Engine", Path.GetDirectoryName(Params.RawProjectPath),CombinePaths(CmdEnv.LocalRoot, "Engine")); 
-
-		int ConfigHeapSize; 
-		if (!ConfigCache.GetInt32 ("BuildSettings","HeapSize"+Params.ClientConfigsToBuild[0].ToString(), out ConfigHeapSize)) // in Megs.
-		{ 	 
-			// we couldn't find a per config heap size, look for a common one.
-			if(!ConfigCache.GetInt32 ("BuildSettings","HeapSize", out ConfigHeapSize))
-				{
-					ConfigHeapSize = Params.IsCodeBasedProject ? 1024 : 512;
-					Log ("Could not find Heap Size setting in .ini for Client config {0}" ,Params.ClientConfigsToBuild[0].ToString()); 
-				}
-		}
+        Log("Package {0}", Params.RawProjectPath);
+		PrintRunTime();
+	}
 
-		HeapSize = (ulong)ConfigHeapSize * 1024L * 1024L; // convert to bytes. 
-		Log ("Setting Heap size to {0} Mb ", ConfigHeapSize); 
+    public override void PostStagingFileCopy(ProjectParams Params, DeploymentContext SC)
+    {
+        Log("Post Staging File Copy {0}", Params.RawProjectPath);
 
+        string PackagePath = Path.Combine(Path.GetDirectoryName(Params.RawProjectPath), "Binaries", "HTML5");
+        if (!Directory.Exists(PackagePath))
+        {
+            Directory.CreateDirectory(PackagePath);
+        }
+        string FinalDataLocation = Path.Combine(PackagePath, Params.ShortProjectName) + ".data";
 
-		GenerateFileFromTemplate(TemplateFile, OutputFile, Params.ShortProjectName, Params.ClientConfigsToBuild[0].ToString(), Params.StageCommandline, !Params.IsCodeBasedProject, HeapSize);
 
-		// copy the jstorage files to the binaries directory
-		string JSDir = Path.Combine(CombinePaths(CmdEnv.LocalRoot, "Engine"), "Build", "HTML5");
-		string OutDir = PackagePath;
-		File.Copy(JSDir + "/json2.js", OutDir + "/json2.js", true);
-		File.SetAttributes(OutDir + "/json2.js", FileAttributes.Normal);
-		File.Copy(JSDir + "/jstorage.js", OutDir + "/jstorage.js", true);
-		File.SetAttributes(OutDir + "/jstorage.js", FileAttributes.Normal);
-		File.Copy(JSDir + "/moz_binarystring.js", OutDir + "/moz_binarystring.js", true);
-		File.SetAttributes(OutDir + "/moz_binarystring.js", FileAttributes.Normal);
-		PrintRunTime();
-	}
+        // we need to operate in the root
+        using (new PushedDirectory(Path.Combine(Params.BaseStageDirectory, "HTML5")))
+        {
+            string BaseSDKPath = Environment.GetEnvironmentVariable("EMSCRIPTEN");
+            // make the file_packager command line
+            if (Utils.IsRunningOnMono)
+            {
+                string PackagerPath = BaseSDKPath + "/tools/file_packager.py";
+                string CmdLine = string.Format("-c \" python {0} '{1}' --preload . --exclude \"*{2}\\Binaries\\HTML5\\*.*\" --js-output='{1}.js' \" ", PackagerPath, FinalDataLocation, Params.ShortProjectName);
+                RunAndLog(CmdEnv, "/bin/bash", CmdLine);
+            }else
+            {
+                string PackagerPath = "\"" + BaseSDKPath + "\\tools\\file_packager.py\"";
+                string CmdLine = string.Format("/c python {0} \"{1}\" --preload . --exclude \"*{2}\\Binaries\\HTML5\\*.*\" --js-output=\"{1}.js\"", PackagerPath, FinalDataLocation, Params.ShortProjectName);
+                RunAndLog(CmdEnv, CommandUtils.CombinePaths(Environment.SystemDirectory, "cmd.exe"), CmdLine);
+            }
+        }
+
+        // copy the "Executable" to the package directory
+        string GameExe = Path.GetFileNameWithoutExtension(Params.ProjectGameExeFilename);
+        if (Params.ClientConfigsToBuild[0].ToString() != "Development")
+        {
+            GameExe += "-HTML5-" + Params.ClientConfigsToBuild[0].ToString();
+        }
+        GameExe += ".js";
+        if (Path.Combine(Path.GetDirectoryName(Params.ProjectGameExeFilename), GameExe) != Path.Combine(PackagePath, GameExe))
+        {
+            File.Copy(Path.Combine(Path.GetDirectoryName(Params.ProjectGameExeFilename), GameExe), Path.Combine(PackagePath, GameExe), true);
+            File.Copy(Path.Combine(Path.GetDirectoryName(Params.ProjectGameExeFilename), GameExe) + ".mem", Path.Combine(PackagePath, GameExe) + ".mem", true);
+        }
+        File.SetAttributes(Path.Combine(PackagePath, GameExe), FileAttributes.Normal);
+        File.SetAttributes(Path.Combine(PackagePath, GameExe) + ".mem", FileAttributes.Normal);
+
+        // put the HTML file to the package directory
+        string TemplateFile = Path.Combine(CombinePaths(CmdEnv.LocalRoot, "Engine"), "Build", "HTML5", "Game.html.template");
+        string OutputFile = Path.Combine(PackagePath, (Params.ClientConfigsToBuild[0].ToString() != "Development" ? (Params.ShortProjectName + "-HTML5-" + Params.ClientConfigsToBuild[0].ToString()) : Params.ShortProjectName)) + ".html";
+
+        // find Heap Size.
+        ulong HeapSize;
+        var ConfigCache = new UnrealBuildTool.ConfigCacheIni(UnrealTargetPlatform.HTML5, "Engine", Path.GetDirectoryName(Params.RawProjectPath), CombinePaths(CmdEnv.LocalRoot, "Engine"));
+
+        int ConfigHeapSize;
+        if (!ConfigCache.GetInt32("BuildSettings", "HeapSize" + Params.ClientConfigsToBuild[0].ToString(), out ConfigHeapSize)) // in Megs.
+        {
+            // we couldn't find a per config heap size, look for a common one.
+            if (!ConfigCache.GetInt32("BuildSettings", "HeapSize", out ConfigHeapSize))
+            {
+                ConfigHeapSize = Params.IsCodeBasedProject ? 1024 : 512;
+                Log("Could not find Heap Size setting in .ini for Client config {0}", Params.ClientConfigsToBuild[0].ToString());
+            }
+        }
+
+        HeapSize = (ulong)ConfigHeapSize * 1024L * 1024L; // convert to bytes. 
+        Log("Setting Heap size to {0} Mb ", ConfigHeapSize);
+
+
+        GenerateFileFromTemplate(TemplateFile, OutputFile, Params.ShortProjectName, Params.ClientConfigsToBuild[0].ToString(), Params.StageCommandline, !Params.IsCodeBasedProject, HeapSize);
+
+        // copy the jstorage files to the binaries directory
+        string JSDir = Path.Combine(CombinePaths(CmdEnv.LocalRoot, "Engine"), "Build", "HTML5");
+        string OutDir = PackagePath;
+        File.Copy(JSDir + "/json2.js", OutDir + "/json2.js", true);
+        File.SetAttributes(OutDir + "/json2.js", FileAttributes.Normal);
+        File.Copy(JSDir + "/jstorage.js", OutDir + "/jstorage.js", true);
+        File.SetAttributes(OutDir + "/jstorage.js", FileAttributes.Normal);
+        File.Copy(JSDir + "/moz_binarystring.js", OutDir + "/moz_binarystring.js", true);
+        File.SetAttributes(OutDir + "/moz_binarystring.js", FileAttributes.Normal);
+	    PrintRunTime();
+    }
+    
+    public override void GetFilesToArchive(ProjectParams Params, DeploymentContext SC)
+    {
+        // Archive everything in the Binaries for HTML5. 
+        string PackagePath = Path.Combine(Path.GetDirectoryName(Params.RawProjectPath), "Binaries", "HTML5"); 
+        SC.ArchiveFiles(PackagePath);
+    }
 
 	protected void GenerateFileFromTemplate(string InTemplateFile, string InOutputFile, string InGameName, string InGameConfiguration, string InArguments, bool IsContentOnly, ulong HeapSize)
 	{
