Port C# compactify routines to run via BuildPatchTool.exe in C++

Note that the existing C# compactify still runs after this CL, and following a successful code review, a further CL will be
submitted to switch over to using the new C++ compactify codebase.

diff --git a/Engine/Source/Programs/AutomationTool/MCPPublic.cs b/Engine/Source/Programs/AutomationTool/MCPPublic.cs
--- a/Engine/Source/Programs/AutomationTool/MCPPublic.cs
+++ b/Engine/Source/Programs/AutomationTool/MCPPublic.cs
@@ -268,6 +268,14 @@ namespace EpicGames.MCP.Automation
             /// Matches the corresponding BuildPatchTool command line argument.
             /// </summary>
             public MCPPlatform Platform;
+            /// <summary>
+            /// If specified, BuildPatchTool will run a compactify on this directory.
+            /// </summary>
+            public string CompactifyDirectory;
+            /// <summary>
+            /// Corresponds to the -preview parameter
+            /// </summary>
+            public bool bPreviewCompactify;
         }
 
         static BuildPatchToolBase Handler = null;
