Seperated the support for tcp and http transports when using file serving so that they can be used at the same time.

diff --git a/Engine/Source/Developer/LauncherServices/Private/Launcher/LauncherCookCommands.h b/Engine/Source/Developer/LauncherServices/Private/Launcher/LauncherCookCommands.h
--- a/Engine/Source/Developer/LauncherServices/Private/Launcher/LauncherCookCommands.h
+++ b/Engine/Source/Developer/LauncherServices/Private/Launcher/LauncherCookCommands.h
@@ -236,13 +236,39 @@ public:
 		CommandLine += TEXT(" -cookonthefly -skipserver");
 
 		// extract the port from the address list string
+		FString Ports;
 		FString Address, Rest;
-		if(ChainState.FileServerAddressListString.Split(TEXT(":"), &Address, &Rest))
+		Rest = ChainState.FileServerAddressListString;
+		TArray<FString> UniquePortArray;
+		while (Rest.Split(TEXT(":"), &Address, &Rest))
 		{
-			if (Rest.Split(TEXT("+"), &Address, &Rest))
+			FString Protocol = "tcp";
+			if ( Address == "http" || Address == "tcp")
 			{
-				CommandLine += FString::Printf(TEXT(" -port=%s"), *Address);
+				Protocol = Address;
+				Rest.Split( TEXT(":"), &Address, &Rest );
+				Address.RemoveFromStart("//");
 			}
+			FString Port;
+			if (Rest.Split(TEXT("+"), &Port, &Rest))
+			{
+				// one port can't support multiple protocols 
+				// only add each port once
+				if ( UniquePortArray.Contains(Port) == false )
+				{
+					UniquePortArray.Add( Port );
+					if ( Ports.Len() > 0)
+					{
+						Ports += "+";
+					}
+					Ports += FString::Printf(TEXT("%s:%s"), *Protocol, *Port);
+					
+				}
+			}
+		}
+		if ( Ports.Len() )
+		{
+			CommandLine += FString::Printf( TEXT(" -port=%s"), *Ports );
 		}
 		return CommandLine;
 	}
diff --git a/Engine/Source/Editor/UnrealEd/Classes/CookOnTheSide/CookOnTheFlyServer.h b/Engine/Source/Editor/UnrealEd/Classes/CookOnTheSide/CookOnTheFlyServer.h
--- a/Engine/Source/Editor/UnrealEd/Classes/CookOnTheSide/CookOnTheFlyServer.h
+++ b/Engine/Source/Editor/UnrealEd/Classes/CookOnTheSide/CookOnTheFlyServer.h
@@ -9,7 +9,6 @@
 #include "CookOnTheFlyServer.generated.h"
 
 
-
 UCLASS()
 class UCookOnTheFlyServer : public UObject, public FTickableEditorObject
 {
@@ -31,7 +30,7 @@ private:
 	/** Holds the sandbox file wrapper to handle sandbox path conversion. */
 	TAutoPtr<class FSandboxPlatformFile> SandboxFile;
 	/** Cook on the fly server uses the NetworkFileServer */
-	class INetworkFileServer* NetworkFileServer;
+	TArray<class INetworkFileServer*> NetworkFileServers;
 
 	/** We hook this up to a delegate to avoid reloading textures and whatnot */
 	TSet<FString> PackagesToNotReload;
diff --git a/Engine/Source/Editor/UnrealEd/Private/CookOnTheFlyServer.cpp b/Engine/Source/Editor/UnrealEd/Private/CookOnTheFlyServer.cpp
--- a/Engine/Source/Editor/UnrealEd/Private/CookOnTheFlyServer.cpp
+++ b/Engine/Source/Editor/UnrealEd/Private/CookOnTheFlyServer.cpp
@@ -92,7 +92,7 @@ TStatId UCookOnTheFlyServer::GetStatId() const
 
 
 
-bool UCookOnTheFlyServer::StartNetworkFileServer( bool BindAnyPort )
+bool UCookOnTheFlyServer::StartNetworkFileServer( const bool BindAnyPort )
 {
 	//GetDerivedDataCacheRef().WaitForQuiescence(false);
 
@@ -100,10 +100,19 @@ bool UCookOnTheFlyServer::StartNetworkFileServer( bool BindAnyPort )
 	FFileRequestDelegate FileRequestDelegate(FFileRequestDelegate::CreateUObject(this, &UCookOnTheFlyServer::HandleNetworkFileServerFileRequest));
 	FRecompileShadersDelegate RecompileShadersDelegate(FRecompileShadersDelegate::CreateUObject(this, &UCookOnTheFlyServer::HandleNetworkFileServerRecompileShaders));
 
-	NetworkFileServer = FModuleManager::LoadModuleChecked<INetworkFileSystemModule>("NetworkFileSystem")
-		.CreateNetworkFileServer(BindAnyPort ? 0 : -1, &FileRequestDelegate, &RecompileShadersDelegate);
-
+	INetworkFileServer *TcpFileServer = FModuleManager::LoadModuleChecked<INetworkFileSystemModule>("NetworkFileSystem")
+		.CreateNetworkFileServer(BindAnyPort ? 0 : -1, &FileRequestDelegate, &RecompileShadersDelegate, ENetworkFileServerProtocol::NFSP_Tcp);
+	if ( TcpFileServer )
+	{
+		NetworkFileServers.Add(TcpFileServer);
+	}
 
+	INetworkFileServer *HttpFileServer = FModuleManager::LoadModuleChecked<INetworkFileSystemModule>("NetworkFileSystem")
+		.CreateNetworkFileServer(BindAnyPort ? 0 : -1, &FileRequestDelegate, &RecompileShadersDelegate, ENetworkFileServerProtocol::NFSP_Http);
+	if ( HttpFileServer )
+	{
+		NetworkFileServers.Add( HttpFileServer );
+	}
 	// loop while waiting for requests
 	GIsRequestingExit = false;
 	return true;
@@ -112,41 +121,45 @@ bool UCookOnTheFlyServer::StartNetworkFileServer( bool BindAnyPort )
 
 bool UCookOnTheFlyServer::BroadcastFileserverPresence( const FGuid &InstanceId )
 {
-	TArray<TSharedPtr<FInternetAddr> > AddressList;
+	
+	TArray<FString> AddressStringList;
 
-	if ((NetworkFileServer == NULL || !NetworkFileServer->IsItReadyToAcceptConnections() || !NetworkFileServer->GetAddressList(AddressList)))
+	for ( int i = 0; i < NetworkFileServers.Num(); ++i )
 	{
-		UE_LOG(LogCookOnTheFly, Error, TEXT("Failed to create network file server"));
-
-		return false;
-	}
+		TArray<TSharedPtr<FInternetAddr> > AddressList;
+		INetworkFileServer *NetworkFileServer = NetworkFileServers[i];
+		if ((NetworkFileServer == NULL || !NetworkFileServer->IsItReadyToAcceptConnections() || !NetworkFileServer->GetAddressList(AddressList)))
+		{
+			UE_LOG(LogCookOnTheFly, Error, TEXT("Failed to create network file server"));
 
-	// broadcast our presence
-	if (InstanceId.IsValid())
-	{
-		TArray<FString> AddressStringList;
+			continue;
+		}
 
-		for (int32 AddressIndex = 0; AddressIndex < AddressList.Num(); ++AddressIndex)
+		// broadcast our presence
+		if (InstanceId.IsValid())
 		{
-			AddressStringList.Add(AddressList[AddressIndex]->ToString(true));
+			for (int32 AddressIndex = 0; AddressIndex < AddressList.Num(); ++AddressIndex)
+			{
+				AddressStringList.Add(FString::Printf( TEXT("%s://%s"), *NetworkFileServer->GetSupportedProtocol(),  *AddressList[AddressIndex]->ToString(true)));
+			}
+
 		}
+	}
 
-		FMessageEndpointPtr MessageEndpoint = FMessageEndpoint::Builder("UCookOnTheFlyServer").Build();
 
-		if (MessageEndpoint.IsValid())
-		{
-			MessageEndpoint->Publish(new FFileServerReady(AddressStringList, InstanceId), EMessageScope::Network);
-		}		
-	}
+
+	FMessageEndpointPtr MessageEndpoint = FMessageEndpoint::Builder("UCookOnTheFlyServer").Build();
+
+	if (MessageEndpoint.IsValid())
+	{
+		MessageEndpoint->Publish(new FFileServerReady(AddressStringList, InstanceId), EMessageScope::Network);
+	}		
+	
 	return true;
 }
 
 uint32 UCookOnTheFlyServer::TickCookOnTheSide( const float TimeSlice, uint32 &CookedPackageCount )
 {
-
-	check( NetworkFileServer != NULL );
-
-
 	double StartTime = FPlatformTime::Seconds();
 
 	uint32 Result = 0;
@@ -373,8 +386,9 @@ void UCookOnTheFlyServer::MarkPackageDirtyForCooker( UPackage *Package )
 
 void UCookOnTheFlyServer::EndNetworkFileServer()
 {
-	if ( NetworkFileServer )
+	for ( int i = 0; i < NetworkFileServers.Num(); ++i )
 	{
+		INetworkFileServer *NetworkFileServer = NetworkFileServers[i];
 		// shutdown the server
 		NetworkFileServer->Shutdown();
 		delete NetworkFileServer;
@@ -712,11 +726,16 @@ void UCookOnTheFlyServer::Initialize( bool inCompressed, bool inIterativeCooking
 
 uint32 UCookOnTheFlyServer::NumConnections() const
 {
-	if ( NetworkFileServer )
+	int Result= 0;
+	for ( int i = 0; i < NetworkFileServers.Num(); ++i )
 	{
-		return NetworkFileServer->NumConnections();
+		INetworkFileServer *NetworkFileServer = NetworkFileServers[i];
+		if ( NetworkFileServer )
+		{
+			Result += NetworkFileServer->NumConnections();
+		}
 	}
-	return 0;
+	return Result;
 }
 
 FString UCookOnTheFlyServer::GetOutputDirectoryOverride( const FString &OutputDirectoryOverride ) const
diff --git a/Engine/Source/Programs/AutomationTool/Scripts/CopyBuildToStagingDirectory.Automation.cs b/Engine/Source/Programs/AutomationTool/Scripts/CopyBuildToStagingDirectory.Automation.cs
--- a/Engine/Source/Programs/AutomationTool/Scripts/CopyBuildToStagingDirectory.Automation.cs
+++ b/Engine/Source/Programs/AutomationTool/Scripts/CopyBuildToStagingDirectory.Automation.cs
@@ -545,7 +545,7 @@ public partial class Project : CommandUtils
 			if (Params.CookOnTheFly || Params.FileServer)
 			{
 				FileHostParams += "-filehostip=";
-
+                bool FirstParam = true;
 				if (UnrealBuildTool.ExternalExecution.GetRuntimePlatform() == UnrealTargetPlatform.Mac)
 				{
 					NetworkInterface[] Interfaces = NetworkInterface.GetAllNetworkInterfaces();
@@ -558,13 +558,30 @@ public partial class Project : CommandUtils
 							{
 								if (IP.UnicastAddresses[Index].IsDnsEligible && IP.UnicastAddresses[Index].Address.AddressFamily == System.Net.Sockets.AddressFamily.InterNetwork)
 								{
-									FileHostParams += IP.UnicastAddresses[Index].Address.ToString();
-									if (String.IsNullOrEmpty(Params.Port) == false)
+									if (Params.Port != null)
 									{
-										FileHostParams += ":";
-										FileHostParams += Params.Port;
+                                        foreach (var Port in Params.Port)
+                                        {
+                                            if (!FirstParam)
+                                            {
+                                                FileHostParams += "+";
+                                                FirstParam = false;
+                                            }
+                                            string[] PortProtocol = Port.Split(new char[] { ':' });
+                                            if (PortProtocol.Length > 1)
+                                            {
+                                                FileHostParams += String.Format("{0}://{1}:{2}", PortProtocol[0], IP.UnicastAddresses[Index].Address.ToString(), PortProtocol[1]);
+                                            }
+                                            else
+                                            {
+                                                FileHostParams += IP.UnicastAddresses[Index].Address.ToString();
+                                                FileHostParams += ":";
+                                                FileHostParams += Params.Port;
+                                            }
+                                            
+                                        }
 									}
-									FileHostParams += "+";
+									
 								}
 							}
 						}
@@ -582,24 +599,58 @@ public partial class Project : CommandUtils
 							{
 								if (IP.UnicastAddresses[Index].IsDnsEligible)
 								{
-									FileHostParams += IP.UnicastAddresses[Index].Address.ToString();
-									if (String.IsNullOrEmpty(Params.Port) == false)
-									{
-										FileHostParams += ":";
-										FileHostParams += Params.Port;
-									}
-									FileHostParams += "+";
+                                    if (Params.Port != null)
+                                    {
+                                        foreach (var Port in Params.Port)
+                                        {
+                                            if (!FirstParam)
+                                            {
+                                                FileHostParams += "+";
+                                                FirstParam = false;
+                                            }
+                                            string[] PortProtocol = Port.Split(new char[] { ':' });
+                                            if (PortProtocol.Length > 1)
+                                            {
+                                                FileHostParams += String.Format("{0}://{1}:{2}", PortProtocol[0], IP.UnicastAddresses[Index].Address.ToString(), PortProtocol[1]);
+                                            }
+                                            else
+                                            {
+                                                FileHostParams += IP.UnicastAddresses[Index].Address.ToString();
+                                                FileHostParams += ":";
+                                                FileHostParams += Params.Port;
+                                            }
+                                        }
+                                    }
+                                    
 								}
 							}
 						}
 					}
 				}
-				FileHostParams += "127.0.0.1";
-				if (String.IsNullOrEmpty(Params.Port) == false)
-				{
-					FileHostParams += ":";
-					FileHostParams += Params.Port;
-				}
+				const string LocalHost = "127.0.0.1";
+                if (Params.Port != null)
+                {
+                    foreach (var Port in Params.Port)
+                    {
+                        if (!FirstParam)
+                        {
+                            FileHostParams += "+";
+                            FirstParam = false;
+                        }
+                        string[] PortProtocol = Port.Split(new char[] { ':' });
+                        if (PortProtocol.Length > 1)
+                        {
+                            FileHostParams += String.Format("{0}://{1}:{2}", PortProtocol[0], LocalHost, PortProtocol[1]);
+                        }
+                        else
+                        {
+                            FileHostParams += LocalHost;
+                            FileHostParams += ":";
+                            FileHostParams += Params.Port;
+                        }
+
+                    }
+                }
 				FileHostParams += " ";
 			}
 
diff --git a/Engine/Source/Programs/AutomationTool/Scripts/RunProjectCommand.Automation.cs b/Engine/Source/Programs/AutomationTool/Scripts/RunProjectCommand.Automation.cs
--- a/Engine/Source/Programs/AutomationTool/Scripts/RunProjectCommand.Automation.cs
+++ b/Engine/Source/Programs/AutomationTool/Scripts/RunProjectCommand.Automation.cs
@@ -589,6 +589,7 @@ public partial class Project : CommandUtils
 			if (Params.CookOnTheFly || Params.FileServer)
 			{
 				TempCmdLine += "-filehostip=";
+                bool FirstParam = true;
 				if (UnrealBuildTool.ExternalExecution.GetRuntimePlatform () == UnrealTargetPlatform.Mac)
 				{
 					NetworkInterface[] Interfaces = NetworkInterface.GetAllNetworkInterfaces ();
@@ -601,13 +602,28 @@ public partial class Project : CommandUtils
 							{
 								if (IP.UnicastAddresses [Index].IsDnsEligible && IP.UnicastAddresses [Index].Address.AddressFamily == System.Net.Sockets.AddressFamily.InterNetwork)
 								{
-									TempCmdLine += IP.UnicastAddresses[Index].Address.ToString();
-									if (String.IsNullOrEmpty (Params.Port) == false)
-									{
-										TempCmdLine += ":";
-										TempCmdLine += Params.Port;
-									}
-									TempCmdLine += "+";
+                                    if (Params.Port != null)
+                                    {
+                                        foreach (var Port in Params.Port)
+                                        {
+                                            if (!FirstParam)
+                                            {
+                                                TempCmdLine += "+";
+                                                FirstParam = false;
+                                            }
+                                            string[] PortProtocol = Port.Split(new char[] { ':' });
+                                            if (PortProtocol.Length > 1)
+                                            {
+                                                TempCmdLine += String.Format("{0}://{1}:{2}", PortProtocol[0], IP.UnicastAddresses[Index].Address.ToString(), PortProtocol[1]);
+                                            }
+                                            else
+                                            {
+                                                TempCmdLine += IP.UnicastAddresses[Index].Address.ToString();
+                                                TempCmdLine += ":";
+                                                TempCmdLine += Params.Port;
+                                            }
+                                        }
+                                    }
 								}
 							}
 						}
@@ -625,24 +641,59 @@ public partial class Project : CommandUtils
 							{
 								if (IP.UnicastAddresses[Index].IsDnsEligible)
 								{
-									TempCmdLine += IP.UnicastAddresses[Index].Address.ToString();
-									if (String.IsNullOrEmpty(Params.Port) == false)
-									{
-										TempCmdLine += ":";
-										TempCmdLine += Params.Port;
-									}
-									TempCmdLine += "+";
+                                    if (Params.Port != null)
+                                    {
+                                        foreach (var Port in Params.Port)
+                                        {
+                                            if (!FirstParam)
+                                            {
+                                                TempCmdLine += "+";
+                                                FirstParam = false;
+                                            }
+                                            string[] PortProtocol = Port.Split(new char[] { ':' });
+                                            if (PortProtocol.Length > 1)
+                                            {
+                                                TempCmdLine += String.Format("{0}://{1}:{2}", PortProtocol[0], IP.UnicastAddresses[Index].Address.ToString(), PortProtocol[1]);
+                                            }
+                                            else
+                                            {
+                                                TempCmdLine += IP.UnicastAddresses[Index].Address.ToString();
+                                                TempCmdLine += ":";
+                                                TempCmdLine += Params.Port;
+                                            }
+                                        }
+                                    }
 								}
 							}
 						}
 					}
 				}
-				TempCmdLine += "127.0.0.1";
-				if (String.IsNullOrEmpty(Params.Port) == false)
-				{
-					TempCmdLine += ":";
-					TempCmdLine += Params.Port;
-				}
+				
+                const string LocalHost = "127.0.0.1";
+
+                if (Params.Port != null)
+                {
+                    foreach (var Port in Params.Port)
+                    {
+                        if (!FirstParam)
+                        {
+                            TempCmdLine += "+";
+                            FirstParam = false;
+                        }
+                        string[] PortProtocol = Port.Split(new char[] { ':' });
+                        if (PortProtocol.Length > 1)
+                        {
+                            TempCmdLine += String.Format("{0}://{1}:{2}", PortProtocol[0], LocalHost, PortProtocol[1]);
+                        }
+                        else
+                        {
+                            TempCmdLine += LocalHost;
+                            TempCmdLine += ":";
+                            TempCmdLine += Params.Port;
+                        }
+
+                    }
+                }
 				TempCmdLine += " ";
 
 				if (!Params.Stage)
diff --git a/Engine/Source/Programs/UnrealBuildTool/HTML5/UEBuildHTML5.cs b/Engine/Source/Programs/UnrealBuildTool/HTML5/UEBuildHTML5.cs
--- a/Engine/Source/Programs/UnrealBuildTool/HTML5/UEBuildHTML5.cs
+++ b/Engine/Source/Programs/UnrealBuildTool/HTML5/UEBuildHTML5.cs
@@ -294,21 +294,14 @@ namespace UnrealBuildTool
 
                 if (InModule.ToString() == "NetworkFile")
                 {
-                    if (EnableHTTPForNFS)
-                    {   
-                        InModule.AddPublicDefinition("ENABLE_HTTP_FOR_NFS=1");
-                        if (Target.Architecture == "-win32")
-                        {
-                            InModule.AddPrivateDependencyModule("HTML5Win32");
-                        }
-                        else
-                        {
-                            InModule.AddPrivateDependencyModule("HTML5JS");
-                        }
+                    InModule.AddPublicDefinition("ENABLE_HTTP_FOR_NFS=1");
+                    if (Target.Architecture == "-win32")
+                    {
+                        InModule.AddPrivateDependencyModule("HTML5Win32");
                     }
-                    if (UseHTTPForNFS)
+                    else
                     {
-                        InModule.AddPublicDefinition("USE_HTTP_FOR_NFS=1");
+                        InModule.AddPrivateDependencyModule("HTML5JS");
                     }
                 }
 			}
@@ -319,35 +312,18 @@ namespace UnrealBuildTool
 					&& UEBuildConfiguration.bBuildDeveloperTools)
 					|| UEBuildConfiguration.bForceBuildTargetPlatforms)
 				{
-					InModule.AddPlatformSpecificDynamicallyLoadedModule("HTML5TargetPlatform");
-				}
-
-                if (EnableHTTPForNFS)
+                    InModule.AddPlatformSpecificDynamicallyLoadedModule("HTML5TargetPlatform");
+                }
+                if (InModule.ToString() == "NetworkFile") // client
                 {
-                    if (InModule.ToString() == "NetworkFile") // client
-                    {
-                        InModule.AddPrivateDependencyModule("HTTP");
-                        InModule.AddPublicDefinition("ENABLE_HTTP_FOR_NFS=1");
-                    }
-                    else if (InModule.ToString() == "NetworkFileSystem") // server 
-                    {
-                        InModule.AddPublicDependencyModule("WebSockets");
-                        InModule.AddPublicDefinition("ENABLE_HTTP_FOR_NFS=1");
-                    }
+                    InModule.AddPrivateDependencyModule("HTTP");
+                    InModule.AddPublicDefinition("ENABLE_HTTP_FOR_NFS=1");
                 }
-                if (UseHTTPForNFS)
+                else if (InModule.ToString() == "NetworkFileSystem") // server 
                 {
-                    if (InModule.ToString() == "NetworkFile") // client
-                    {
-                        InModule.AddPublicDefinition("USE_HTTP_FOR_NFS=1");
-                    }
-                    else if (InModule.ToString() == "NetworkFileSystem") // server 
-                    {
-                        InModule.AddPublicDefinition("USE_HTTP_FOR_NFS=1");
-                    }
+                    InModule.AddPublicDependencyModule("WebSockets");
+                    InModule.AddPublicDefinition("ENABLE_HTTP_FOR_NFS=1");
                 }
-
-
 			}
 		}
 
diff --git a/Engine/Source/Runtime/NetworkFile/Private/HTTPTransport.cpp b/Engine/Source/Runtime/NetworkFile/Private/HTTPTransport.cpp
--- a/Engine/Source/Runtime/NetworkFile/Private/HTTPTransport.cpp
+++ b/Engine/Source/Runtime/NetworkFile/Private/HTTPTransport.cpp
@@ -19,9 +19,27 @@ FHTTPTransport::FHTTPTransport()
 {
 }
 
-bool FHTTPTransport::Initialize(const TCHAR* HostIp)
+bool FHTTPTransport::Initialize(const TCHAR* InHostIp)
 {
-	FCString::Sprintf( Url, TEXT("http://%s:%d"), HostIp, (int)DEFAULT_FILE_SERVING_PORT); 
+	// parse out the format
+	FString HostIp = InHostIp;
+
+	// make sure that we have the correct protcol
+	ensure( HostIp.RemoveFromStart("http://") );
+
+	// check if we have specified the port also
+	if ( HostIp.Contains(":") == false)
+	{
+		// no port put the default one on
+		HostIp = FString::Printf(TEXT("%s:%d"), *HostIp, (int)(DEFAULT_HTTP_FILE_SERVING_PORT) );
+	}
+	else
+	{
+		// make sure that our string is correctly formated
+		FString::Printf(TEXT("http://%s"),*HostIp);
+	}
+
+	FCString::Sprintf(Url, *HostIp);
 
 #if !PLATFORM_HTML5
 	HttpRequest = FHttpModule::Get().CreateRequest(); 
diff --git a/Engine/Source/Runtime/NetworkFile/Private/NetworkPlatformFile.cpp b/Engine/Source/Runtime/NetworkFile/Private/NetworkPlatformFile.cpp
--- a/Engine/Source/Runtime/NetworkFile/Private/NetworkPlatformFile.cpp
+++ b/Engine/Source/Runtime/NetworkFile/Private/NetworkPlatformFile.cpp
@@ -24,15 +24,8 @@ FNetworkPlatformFile::FNetworkPlatformFile()
 	: bHasLoadedDDCDirectories(false)
 	, InnerPlatformFile(NULL)
 	, bIsUsable(false)
-	, FileServerPort(DEFAULT_FILE_SERVING_PORT)
 	, Transport(NULL)
 {
-#if !USE_HTTP_FOR_NFS
-	Transport = new FTCPTransport(); 
-#else
-	Transport = new FHTTPTransport(); 	
-#endif 
-
 }
 
 bool FNetworkPlatformFile::ShouldBeUsed(IPlatformFile* Inner, const TCHAR* CmdLine) const
@@ -41,6 +34,24 @@ bool FNetworkPlatformFile::ShouldBeUsed(IPlatformFile* Inner, const TCHAR* CmdLi
 	return FParse::Value(CmdLine, TEXT("-FileHostIP="), HostIp);
 }
 
+ITransport *CreateTransportForHostAddress(const FString &HostIp )
+{
+	if ( HostIp.StartsWith(TEXT("tcp://")))
+	{
+		return new FTCPTransport();
+	}
+
+	if ( HostIp.StartsWith(TEXT("http://")))
+	{
+#if ENABLE_HTTP_FOR_NFS
+		return new FHTTPTransport();
+#endif
+	}
+
+	// no transport specified assuming tcp
+	return new FTCPTransport();
+}
+
 bool FNetworkPlatformFile::Initialize(IPlatformFile* Inner, const TCHAR* CmdLine)
 {
 	bool bResult = false;
@@ -50,13 +61,24 @@ bool FNetworkPlatformFile::Initialize(IPlatformFile* Inner, const TCHAR* CmdLine
 		TArray<FString> HostIpList;
 		if (HostIpString.ParseIntoArray(&HostIpList, TEXT("+"), true) > 0)
 		{
-			// Try to initialize with each of the IP addresses found in the command line until we 
-			// get a working one.
 			for (int32 HostIpIndex = 0; !bResult && HostIpIndex < HostIpList.Num(); ++HostIpIndex)
 			{
-				bResult = Transport->Initialize( *HostIpList[HostIpIndex] ) && InitializeInternal(Inner, *HostIpList[HostIpIndex]);		
-				if (bResult)
-					break;
+				// Try to initialize with each of the IP addresses found in the command line until we 
+				// get a working one.
+
+				// find the correct transport for this ip address 
+				Transport = CreateTransportForHostAddress( HostIpList[HostIpIndex] );
+
+				if ( Transport )
+				{
+					bResult = Transport->Initialize( *HostIpList[HostIpIndex] ) && InitializeInternal(Inner, *HostIpList[HostIpIndex]);		
+					if (bResult)
+						break;
+
+					// try a different host might be a different protocol
+					delete Transport;
+				}
+				Transport = NULL;
 			}
 		}
 	}
diff --git a/Engine/Source/Runtime/NetworkFile/Private/TCPTransport.cpp b/Engine/Source/Runtime/NetworkFile/Private/TCPTransport.cpp
--- a/Engine/Source/Runtime/NetworkFile/Private/TCPTransport.cpp
+++ b/Engine/Source/Runtime/NetworkFile/Private/TCPTransport.cpp
@@ -11,15 +11,22 @@ FTCPTransport::FTCPTransport()
 {
 }
 
-bool FTCPTransport::Initialize(const TCHAR* HostIp)
+bool FTCPTransport::Initialize(const TCHAR* InHostIp)
 {
 	ISocketSubsystem* SSS = ISocketSubsystem::Get();
 
+	FString HostIp = InHostIp;
+
+
+	// the protocol isn't required for tcp it's assumed by default
+	HostIp.RemoveFromStart(TEXT("tcp://"));
+
 	// convert the string to a ip addr structure
-	TSharedRef<FInternetAddr> Addr = SSS->CreateInternetAddr(0, DEFAULT_FILE_SERVING_PORT);
+	// DEFAULT_TCP_FILE_SERVING_PORT is overridden 
+	TSharedRef<FInternetAddr> Addr = SSS->CreateInternetAddr(0, DEFAULT_TCP_FILE_SERVING_PORT);
 	bool bIsValid;
 
-	Addr->SetIp(HostIp, bIsValid);
+	Addr->SetIp(*HostIp, bIsValid);
 
 	if (bIsValid)
 	{
@@ -32,7 +39,7 @@ bool FTCPTransport::Initialize(const TCHAR* HostIp)
 			// on failure, shut it all down
 			SSS->DestroySocket(FileSocket);
 			FileSocket = NULL;
-			UE_LOG(LogNetworkPlatformFile, Error, TEXT("Failed to connect to file server at %s:%d."), HostIp, (int32)DEFAULT_FILE_SERVING_PORT);
+			UE_LOG(LogNetworkPlatformFile, Error, TEXT("Failed to connect to file server at %s."), *Addr->ToString(true));
 		}
 	}
 
diff --git a/Engine/Source/Runtime/NetworkFileSystem/Private/NetworkFileServer.cpp b/Engine/Source/Runtime/NetworkFileSystem/Private/NetworkFileServer.cpp
--- a/Engine/Source/Runtime/NetworkFileSystem/Private/NetworkFileServer.cpp
+++ b/Engine/Source/Runtime/NetworkFileSystem/Private/NetworkFileServer.cpp
@@ -89,6 +89,10 @@ private:
 FNetworkFileServer::FNetworkFileServer( int32 InPort, const FFileRequestDelegate* InFileRequestDelegate,const FRecompileShadersDelegate* InRecompileShadersDelegate, const TArray<ITargetPlatform*>& InActiveTargetPlatforms )
 	:ActiveTargetPlatforms(InActiveTargetPlatforms)
 {
+	if(InPort <0)
+	{
+		InPort = DEFAULT_TCP_FILE_SERVING_PORT;
+	}
 
 	Running.Set(false);
 	StopRequested.Set(false);
@@ -232,6 +236,12 @@ void FNetworkFileServer::Exit( )
 /* INetworkFileServer overrides
  *****************************************************************************/
 
+FString FNetworkFileServer::GetSupportedProtocol() const
+{
+	return FString("tcp");
+}
+
+
 bool FNetworkFileServer::GetAddressList( TArray<TSharedPtr<FInternetAddr> >& OutAddresses ) const
 {
 	if (ListenAddr.IsValid())
diff --git a/Engine/Source/Runtime/NetworkFileSystem/Private/NetworkFileServer.h b/Engine/Source/Runtime/NetworkFileSystem/Private/NetworkFileServer.h
--- a/Engine/Source/Runtime/NetworkFileSystem/Private/NetworkFileServer.h
+++ b/Engine/Source/Runtime/NetworkFileSystem/Private/NetworkFileServer.h
@@ -48,15 +48,13 @@ public:
 	// INetworkFileServer interface
 
 	virtual bool IsItReadyToAcceptConnections(void) const; 
-	virtual bool GetAddressList(TArray<TSharedPtr<FInternetAddr> >& OutAddresses) const;
+	virtual bool GetAddressList(TArray<TSharedPtr<FInternetAddr> >& OutAddresses) const override;
+	virtual FString GetSupportedProtocol() const override;
 	virtual int32 NumConnections() const;
 	virtual void Shutdown();
 
 private:
 
-	// Holds the port to use.
-	int32 FileServerPort;
-
 	// Holds the server (listening) socket.
 	FSocket* Socket;
 
diff --git a/Engine/Source/Runtime/NetworkFileSystem/Private/NetworkFileServerHttp.cpp b/Engine/Source/Runtime/NetworkFileSystem/Private/NetworkFileServerHttp.cpp
--- a/Engine/Source/Runtime/NetworkFileSystem/Private/NetworkFileServerHttp.cpp
+++ b/Engine/Source/Runtime/NetworkFileSystem/Private/NetworkFileServerHttp.cpp
@@ -47,6 +47,11 @@ FNetworkFileServerHttp::FNetworkFileServerHttp(
 	:ActiveTargetPlatforms(InActiveTargetPlatforms)
 	,Port(InPort)
 {
+	if (Port < 0 )
+	{
+		Port = DEFAULT_HTTP_FILE_SERVING_PORT;
+	}
+
 	UE_LOG(LogFileServer, Warning, TEXT("Unreal Network Http File Server starting up..."));
 
 	if (InFileRequestDelegate && InFileRequestDelegate->IsBound())
@@ -107,6 +112,8 @@ bool FNetworkFileServerHttp::Init()
 	
 	context = libwebsocket_create_context(&info);
 
+	Port = info.port;
+
 	if (context == NULL) {
 		UE_LOG(LogFileServer, Fatal, TEXT(" Could not create a libwebsocket content for port : %d"), Port);
 		return false;
@@ -117,6 +124,13 @@ bool FNetworkFileServerHttp::Init()
 }
 
 
+
+FString FNetworkFileServerHttp::GetSupportedProtocol() const
+{
+	return FString("http");
+}
+
+
 bool FNetworkFileServerHttp::GetAddressList(TArray<TSharedPtr<FInternetAddr> >& OutAddresses) const
 {
 	// if Init failed, its already too late. 
@@ -151,7 +165,7 @@ void FNetworkFileServerHttp::Shutdown()
 
 uint32 FNetworkFileServerHttp::Run()
 {
-	UE_LOG(LogFileServer, Display, TEXT("Unreal Network File Http Server is ready for client connections "));
+	UE_LOG(LogFileServer, Display, TEXT("Unreal Network File Http Server is ready for client connections on port %d"), Port);
 
 	// start servicing. 
 
diff --git a/Engine/Source/Runtime/NetworkFileSystem/Private/NetworkFileServerHttp.h b/Engine/Source/Runtime/NetworkFileSystem/Private/NetworkFileServerHttp.h
--- a/Engine/Source/Runtime/NetworkFileSystem/Private/NetworkFileServerHttp.h
+++ b/Engine/Source/Runtime/NetworkFileSystem/Private/NetworkFileServerHttp.h
@@ -21,7 +21,8 @@ public:
 	// INetworkFileServer Interface.
 
 	virtual bool IsItReadyToAcceptConnections(void) const; 
-	virtual bool GetAddressList(TArray<TSharedPtr<FInternetAddr> >& OutAddresses) const;
+	virtual FString GetSupportedProtocol() const override;
+	virtual bool GetAddressList(TArray<TSharedPtr<FInternetAddr> >& OutAddresses) const override;
 	virtual int32 NumConnections() const;
 	virtual void Shutdown();
 
@@ -64,7 +65,7 @@ private:
 	FRunnableThread* WorkerThread;
 
 	// port on which this http server runs. 
-	uint32 Port; 
+	int32 Port; 
 
 	// used to send simple message.
 	FThreadSafeCounter StopRequested; 
diff --git a/Engine/Source/Runtime/NetworkFileSystem/Private/NetworkFileSystemModule.cpp b/Engine/Source/Runtime/NetworkFileSystem/Private/NetworkFileSystemModule.cpp
--- a/Engine/Source/Runtime/NetworkFileSystem/Private/NetworkFileSystemModule.cpp
+++ b/Engine/Source/Runtime/NetworkFileSystem/Private/NetworkFileSystemModule.cpp
@@ -20,13 +20,8 @@ public:
 
 	// INetworkFileSystemModule interface
 
-	virtual INetworkFileServer* CreateNetworkFileServer( int32 Port, const FFileRequestDelegate* InFileRequestDelegate, const FRecompileShadersDelegate* InRecompileShadersDelegate ) const override
+	virtual INetworkFileServer* CreateNetworkFileServer( int32 Port, const FFileRequestDelegate* InFileRequestDelegate, const FRecompileShadersDelegate* InRecompileShadersDelegate, const ENetworkFileServerProtocol Protocol ) const override
 	{
-		if (Port < 0)
-		{
-			Port = DEFAULT_FILE_SERVING_PORT;
-		}
-
 		TArray<ITargetPlatform*> ActiveTargetPlatforms;
 
 		// only bother getting the target platforms if there was "-targetplatform" on the commandline, otherwise UnrealFileServer will 
@@ -38,11 +33,17 @@ public:
 			ActiveTargetPlatforms =  TPM.GetActiveTargetPlatforms();
 		}
 
-#if !USE_HTTP_FOR_NFS
-		return new FNetworkFileServer(Port, InFileRequestDelegate, InRecompileShadersDelegate, ActiveTargetPlatforms);
-#else
-		return new FNetworkFileServerHttp(Port, InFileRequestDelegate, InRecompileShadersDelegate, ActiveTargetPlatforms);
-#endif 
+		switch ( Protocol )
+		{
+#if ENABLE_HTTP_FOR_NFS
+		case NFSP_Http: 
+			return new FNetworkFileServerHttp(Port, InFileRequestDelegate, InRecompileShadersDelegate, ActiveTargetPlatforms);
+#endif
+		case NFSP_Tcp:
+			return new FNetworkFileServer(Port, InFileRequestDelegate, InRecompileShadersDelegate, ActiveTargetPlatforms);
+		}
+ 
+		return NULL;
 	}
 };
 
diff --git a/Engine/Source/Runtime/NetworkFileSystem/Public/Interfaces/INetworkFileServer.h b/Engine/Source/Runtime/NetworkFileSystem/Public/Interfaces/INetworkFileServer.h
--- a/Engine/Source/Runtime/NetworkFileSystem/Public/Interfaces/INetworkFileServer.h
+++ b/Engine/Source/Runtime/NetworkFileSystem/Public/Interfaces/INetworkFileServer.h
@@ -26,6 +26,15 @@ public:
 	  */
 	virtual bool GetAddressList( TArray<TSharedPtr<FInternetAddr> >& OutAddresses ) const = 0;
 
+	/**
+	  * Gets the list of local network addresses that the file server listens on.
+	  *
+	  * @param OutAddresses - Will hold the address list.
+	  *
+	  * @return true on success, false otherwise.
+	  */
+	virtual FString GetSupportedProtocol() const = 0;
+
 	/**
 	 * Gets the number of active connections.
 	 *
diff --git a/Engine/Source/Runtime/NetworkFileSystem/Public/Interfaces/INetworkFileSystemModule.h b/Engine/Source/Runtime/NetworkFileSystem/Public/Interfaces/INetworkFileSystemModule.h
--- a/Engine/Source/Runtime/NetworkFileSystem/Public/Interfaces/INetworkFileSystemModule.h
+++ b/Engine/Source/Runtime/NetworkFileSystem/Public/Interfaces/INetworkFileSystemModule.h
@@ -47,6 +47,11 @@ struct FShaderRecompileData
  */
 DECLARE_DELEGATE_OneParam(FRecompileShadersDelegate, const FShaderRecompileData&);
 
+enum ENetworkFileServerProtocol
+{
+	NFSP_Tcp,
+	NFSP_Http,
+};
 
 /**
  * Interface for network file system modules.
@@ -66,7 +71,7 @@ public:
 	 *
 	 * @return The new file server, or nullptr if creation failed.
 	 */
-	virtual INetworkFileServer* CreateNetworkFileServer( int32 Port = -1, const FFileRequestDelegate* InFileRequestDelegate = nullptr, const FRecompileShadersDelegate* InRecompileShadersDelegate = nullptr ) const = 0;
+	virtual INetworkFileServer* CreateNetworkFileServer( int32 Port = -1, const FFileRequestDelegate* InFileRequestDelegate = nullptr, const FRecompileShadersDelegate* InRecompileShadersDelegate = nullptr, const ENetworkFileServerProtocol Protocol = NFSP_Tcp ) const = 0;
 
 public:
 
diff --git a/Engine/Source/Runtime/Sockets/Public/NetworkMessage.h b/Engine/Source/Runtime/Sockets/Public/NetworkMessage.h
--- a/Engine/Source/Runtime/Sockets/Public/NetworkMessage.h
+++ b/Engine/Source/Runtime/Sockets/Public/NetworkMessage.h
@@ -4,7 +4,8 @@
 
 enum
 { 
-	DEFAULT_FILE_SERVING_PORT=41899 // port that the network file server uses
+	DEFAULT_TCP_FILE_SERVING_PORT=41899, // default port to use when making file server tcp connections (used if no protocol is specified)
+	DEFAULT_HTTP_FILE_SERVING_PORT=41898 // port that the network file server uses
 };
 
 // transitional define. Set to zero to avoid using multichannel TCP for the file server and client.
