Submitted missing file from last checkin.  Fix issue with ProjectParams.Port not being an array.

diff --git a/Engine/Source/Programs/AutomationTool/ProjectParams.cs b/Engine/Source/Programs/AutomationTool/ProjectParams.cs
--- a/Engine/Source/Programs/AutomationTool/ProjectParams.cs
+++ b/Engine/Source/Programs/AutomationTool/ProjectParams.cs
@@ -267,7 +267,7 @@ namespace AutomationTool
 			string Device = null,			
 			string MapToRun = null,	
 			string AdditionalServerMapParams = null,
-			string Port = null,
+			ParamList<string> Port = null,
 			string RunCommandline = null,						
 			string StageCommandline = null,
             string BundleName = null,
@@ -441,7 +441,6 @@ namespace AutomationTool
 			this.Deploy = GetParamValueIfNotSpecified(Command, Deploy, this.Deploy, "deploy");
 			this.Device = ParseParamValueIfNotSpecified(Command, Device, "device", String.Empty).Trim(new char[]{'\"'});
 			this.ServerDevice = ParseParamValueIfNotSpecified(Command, ServerDevice, "serverdevice", this.Device);
-			this.Port = ParseParamValueIfNotSpecified(Command, Port, "port", String.Empty);
 			this.NullRHI = GetParamValueIfNotSpecified(Command, NullRHI, this.NullRHI, "nullrhi");
 			this.FakeClient = GetParamValueIfNotSpecified(Command, FakeClient, this.FakeClient, "fakeclient");
 			this.EditorTest = GetParamValueIfNotSpecified(Command, EditorTest, this.EditorTest, "editortest");
@@ -474,6 +473,30 @@ namespace AutomationTool
 			{
 				this.ClientConfigsToBuild = ClientConfigsToBuild;
 			}
+
+            if (Port == null)
+            {
+                if( Command != null )
+                {
+                    this.Port = new ParamList<string>();
+
+                    var PortString = Command.ParseParamValue("port");
+                    if (String.IsNullOrEmpty(PortString) == false)
+                    {
+                        var Ports = new ParamList<string>(PortString.Split('+'));
+                        foreach (var P in Ports)
+                        {
+                            this.Port.Add(P);
+                        }
+                    }
+                    
+                }
+            }
+            else
+            {
+                this.Port = Port;
+            }
+
 			if (ServerConfigsToBuild == null)
 			{
 				if (Command != null)
@@ -987,7 +1010,7 @@ namespace AutomationTool
         [Help("Crash=index", "when running -editortest or a client, adds commands like debug crash, debug rendercrash, etc based on index")]
         public int CrashIndex;
 
-        public string Port;
+        public ParamList<string> Port;
 
         /// <summary>
         /// Run: Linux username for unattended key genereation
