Separate BuildPatchTool.exe options into two classes in automation tool

diff --git a/Engine/Source/Programs/AutomationTool/MCPPublic.cs b/Engine/Source/Programs/AutomationTool/MCPPublic.cs
--- a/Engine/Source/Programs/AutomationTool/MCPPublic.cs
+++ b/Engine/Source/Programs/AutomationTool/MCPPublic.cs
@@ -238,7 +238,7 @@ namespace EpicGames.MCP.Automation
         }
 
 
-        public class BuildPatchToolOptions
+        public class PatchGenerationOptions
         {
             /// <summary>
             /// Staging information
@@ -268,6 +268,10 @@ namespace EpicGames.MCP.Automation
             /// Matches the corresponding BuildPatchTool command line argument.
             /// </summary>
             public MCPPlatform Platform;
+        }
+
+		public class CompactifyOptions
+		{
             /// <summary>
             /// If specified, BuildPatchTool will run a compactify on this directory.
             /// </summary>
@@ -276,7 +280,7 @@ namespace EpicGames.MCP.Automation
             /// Corresponds to the -preview parameter
             /// </summary>
             public bool bPreviewCompactify;
-        }
+		}
 
         static BuildPatchToolBase Handler = null;
 
@@ -306,10 +310,16 @@ namespace EpicGames.MCP.Automation
         }
 
         /// <summary>
-        /// Runs the patcher executable using the supplied parameters.
+        /// Runs the Build Patch Tool executable to generate patch data using the supplied parameters.
         /// </summary>
-        /// <returns>BuildInfo descringing the output of the chunking process.</returns>
-        public abstract void Execute(BuildPatchToolOptions Opts);
+		/// <param name="Opts">Parameters which will be passed to the patch tool generation process</param>
+		public abstract void Execute(PatchGenerationOptions Opts);
+
+		/// <summary>
+		/// Runs the Build Patch Tool executable to compactify a cloud directory using the supplied parameters.
+		/// </summary>
+		/// <param name="Opts">Parameters which will be passed to the compactify process</param>
+		public abstract void Execute(CompactifyOptions Opts);
     }
 
 
