AutomationTool: Throw an exception when threaded copy fails

diff --git a/Engine/Source/Programs/AutomationTool/CommandUtils.cs b/Engine/Source/Programs/AutomationTool/CommandUtils.cs
--- a/Engine/Source/Programs/AutomationTool/CommandUtils.cs
+++ b/Engine/Source/Programs/AutomationTool/CommandUtils.cs
@@ -1517,6 +1517,10 @@ namespace AutomationTool
 					{
 						Request.Worker.Join();
 					}
+					if (Request.Result != null)
+					{
+						throw new AutomationException(String.Format("Failed to thread-copy files: {0}", Request.Result.Message), Request.Result);
+					}
 				}
 			}
 			else
