- Removed WinForms from the integrated IPP (in UAT), but no new functionality added yet

diff --git a/Engine/Source/Programs/AutomationTool/IOS/IOS.Automation.csproj b/Engine/Source/Programs/AutomationTool/IOS/IOS.Automation.csproj
--- a/Engine/Source/Programs/AutomationTool/IOS/IOS.Automation.csproj
+++ b/Engine/Source/Programs/AutomationTool/IOS/IOS.Automation.csproj
@@ -34,13 +34,24 @@
     <TreatWarningsAsErrors>true</TreatWarningsAsErrors>
   </PropertyGroup>
   <ItemGroup>
+    <Reference Include="BouncyCastle.Crypto">
+      <HintPath>..\..\..\..\Binaries\ThirdParty\IOS\BouncyCastle.Crypto.dll</HintPath>
+      <Private>False</Private>
+    </Reference>
     <Reference Include="Ionic.Zip.Reduced">
       <HintPath>..\..\..\..\Binaries\DotNET\Ionic.Zip.Reduced.dll</HintPath>
       <Private>False</Private>
     </Reference>
+    <Reference Include="RPCUtility">
+      <HintPath>..\..\..\..\Binaries\DotNET\RPCUtility.exe</HintPath>
+      <Private>False</Private>
+    </Reference>
     <Reference Include="System" />
     <Reference Include="System.Core" />
+    <Reference Include="System.Drawing" />
+    <Reference Include="System.Runtime.Remoting" />
     <Reference Include="System.Runtime.Serialization" />
+    <Reference Include="System.Security" />
     <Reference Include="System.Web" />
     <Reference Include="System.Web.Extensions" />
     <Reference Include="System.Xml.Linq" />
@@ -54,9 +65,24 @@
       <Link>Properties\MetaData.cs</Link>
     </Compile>
     <Compile Include="IOSPlatform.Automation.cs" />
+    <Compile Include="IPhonePackager\CodeSigning.cs" />
+    <Compile Include="IPhonePackager\CompileTime.cs" />
+    <Compile Include="IPhonePackager\Config.cs" />
+    <Compile Include="IPhonePackager\CookTime.cs" />
+    <Compile Include="IPhonePackager\DeployTime.cs" />
+    <Compile Include="IPhonePackager\FileOperations.cs" />
+    <Compile Include="IPhonePackager\iPhonePackager.cs" />
+    <Compile Include="IPhonePackager\MachObjects.cs" />
+    <Compile Include="IPhonePackager\MobileProvisionUtilities.cs" />
+    <Compile Include="IPhonePackager\Utilities.cs" />
     <Compile Include="Properties\AssemblyInfo.cs" />
+    <Compile Include="Properties\Resources.Designer.cs" />
   </ItemGroup>
   <ItemGroup>
+    <ProjectReference Include="..\..\IOS\DeploymentInterface\DeploymentInterface.csproj">
+      <Project>{2da5397a-f93f-41e0-a1ad-b220d40fd739}</Project>
+      <Name>DeploymentInterface</Name>
+    </ProjectReference>
     <ProjectReference Include="..\AutomationTool.csproj">
       <Project>{a35b0261-d868-475e-a098-eb14ce6dc89f}</Project>
       <Name>AutomationTool</Name>
@@ -68,5 +94,13 @@
       <Private>False</Private>
     </ProjectReference>
   </ItemGroup>
+  <ItemGroup>
+    <EmbeddedResource Include="Properties\Resources.resx">
+      <SubType>Designer</SubType>
+    </EmbeddedResource>
+  </ItemGroup>
+  <ItemGroup>
+    <None Include="Resources\Image1.bmp" />
+  </ItemGroup>
   <Import Project="$(MSBuildToolsPath)\Microsoft.CSharp.targets" />
 </Project>
\ No newline at end of file
diff --git a/Engine/Source/Programs/AutomationTool/IOS/IOSPlatform.Automation.cs b/Engine/Source/Programs/AutomationTool/IOS/IOSPlatform.Automation.cs
--- a/Engine/Source/Programs/AutomationTool/IOS/IOSPlatform.Automation.cs
+++ b/Engine/Source/Programs/AutomationTool/IOS/IOSPlatform.Automation.cs
@@ -6,6 +6,7 @@ using System.Text;
 using System.IO;
 using AutomationTool;
 using UnrealBuildTool;
+using System.Text.RegularExpressions;
 using Ionic.Zip;
 using Ionic.Zlib;
 
@@ -22,6 +23,65 @@ public class IOSPlatform : Platform
 	{
 	}
 
+	// Run the integrated IPP code
+	// @todo ipp: How to log the output in a nice UAT way?
+	protected static int RunIPP(string IPPArguments)
+	{
+		List<string> Args = new List<string>();
+
+		StringBuilder Token = null;
+		int Index = 0;
+		bool bInQuote = false;
+		while (Index < IPPArguments.Length)
+		{
+			if (IPPArguments[Index] == '\"')
+			{
+				bInQuote = !bInQuote;
+			}
+			else if (IPPArguments[Index] == ' ' && !bInQuote)
+			{
+				if (Token != null)
+				{
+					Args.Add(Token.ToString());
+					Token = null;
+				}
+			}
+			else
+			{
+				if (Token == null)
+				{
+					Token = new StringBuilder();
+				}
+				Token.Append(IPPArguments[Index]);
+			}
+			Index++;
+		}
+
+		if (Token != null)
+		{
+			Args.Add(Token.ToString());
+		}
+	
+		return RunIPP(Args.ToArray());
+	}
+
+	protected static int RunIPP(string[] Args)
+	{
+		return iPhonePackager.Program.RealMain(Args);
+	}
+
+	public override int RunCommand(string Command)
+	{
+		// run generic IPP commands through the interface
+		if (Command.StartsWith("IPP:"))
+		{
+			return RunIPP(Command.Substring(4));
+		}
+
+		// non-zero error code
+		return 4;
+	}
+
 	protected string MakeIPAFileName( UnrealTargetConfiguration TargetConfiguration, string ProjectGameExeFilename )
 	{
 		var ProjectIPA = Path.ChangeExtension(ProjectGameExeFilename, null);
@@ -124,47 +184,93 @@ public class IOSPlatform : Platform
 
 			bool cookonthefly = Params.CookOnTheFly || Params.SkipCookOnTheFly;
 
-			string IPPArguments = "RepackageFromStage \"" + (Params.IsCodeBasedProject ? Params.RawProjectPath : "Engine") + "\"";
-			IPPArguments += " -config " + TargetConfiguration.ToString();
-
-			if (TargetConfiguration == UnrealTargetConfiguration.Shipping)
+			// rename the .ipa if not code based
+			if (!Params.IsCodeBasedProject)
 			{
-				IPPArguments += " -compress=best";
+				ProjectIPA = Path.Combine(Path.GetDirectoryName(Params.RawProjectPath), "Binaries", "IOS", Params.ShortProjectName + ".ipa");
+				if (TargetConfiguration != UnrealTargetConfiguration.Development)
+				{
+					ProjectIPA = Path.Combine(Path.GetDirectoryName(Params.RawProjectPath), "Binaries", "IOS", Params.ShortProjectName + "-" + PlatformType.ToString() + "-" + TargetConfiguration.ToString() + ".ipa");
+				}
 			}
 
-			// Determine if we should sign
-			bool bNeedToSign = GetCodeSignDesirability(Params);
+			// delete the .ipa to make sure it was made
+			DeleteFile(ProjectIPA);
 
-			if (!String.IsNullOrEmpty(Params.BundleName))
+			if (RemoteToolChain.bUseRPCUtil)
 			{
-				// Have to sign when a bundle name is specified
-				bNeedToSign = true;
-				IPPArguments += " -bundlename " + Params.BundleName;
-			}
+				string IPPArguments = "RepackageFromStage \"" + (Params.IsCodeBasedProject ? Params.RawProjectPath : "Engine") + "\"";
+				IPPArguments += " -config " + TargetConfiguration.ToString();
 
-			if (bNeedToSign)
-			{
-				IPPArguments += " -sign";
-			}
+				if (TargetConfiguration == UnrealTargetConfiguration.Shipping)
+				{
+					IPPArguments += " -compress=best";
+				}
 
-			IPPArguments += (cookonthefly ? " -cookonthefly" : "");
-			IPPArguments += " -stagedir \"" + CombinePaths(Params.BaseStageDirectory, "IOS") + "\"";
-			IPPArguments += " -projectdir \"" + Path.GetDirectoryName(Params.RawProjectPath) + "\"";
+				// Determine if we should sign
+				bool bNeedToSign = GetCodeSignDesirability(Params);
 
-			// rename the .ipa if not code based
-			if (!Params.IsCodeBasedProject)
-			{
-				ProjectIPA = Path.Combine(Path.GetDirectoryName(Params.RawProjectPath), "Binaries", "IOS", Params.ShortProjectName + ".ipa");
-				if (TargetConfiguration != UnrealTargetConfiguration.Development)
+				if (!String.IsNullOrEmpty(Params.BundleName))
 				{
-					ProjectIPA = Path.Combine(Path.GetDirectoryName(Params.RawProjectPath), "Binaries", "IOS", Params.ShortProjectName + "-" + PlatformType.ToString() + "-" + TargetConfiguration.ToString() + ".ipa");
+					// Have to sign when a bundle name is specified
+					bNeedToSign = true;
+					IPPArguments += " -bundlename " + Params.BundleName;
 				}
+
+				if (bNeedToSign)
+				{
+					IPPArguments += " -sign";
+				}
+
+				IPPArguments += (cookonthefly ? " -cookonthefly" : "");
+				IPPArguments += " -stagedir \"" + CombinePaths(Params.BaseStageDirectory, "IOS") + "\"";
+				IPPArguments += " -projectdir \"" + Path.GetDirectoryName(Params.RawProjectPath) + "\"";
+
+				RunAndLog(CmdEnv, IPPExe, IPPArguments);
 			}
+			else
+			{
+				List<string> IPPArguments = new List<string>();
+				IPPArguments.Add("RepackageFromStage");
+				IPPArguments.Add(Params.IsCodeBasedProject ? Params.RawProjectPath : "Engine");
+				IPPArguments.Add("-config");
+				IPPArguments.Add(TargetConfiguration.ToString());
+
+				if (TargetConfiguration == UnrealTargetConfiguration.Shipping)
+				{
+					IPPArguments.Add("-compress=best");
+				}
 
-            // delete the .ipa to make sure it was made
-            DeleteFile(ProjectIPA);
+				// Determine if we should sign
+				bool bNeedToSign = GetCodeSignDesirability(Params);
 
-			RunAndLog(CmdEnv, IPPExe, IPPArguments);
+				if (!String.IsNullOrEmpty(Params.BundleName))
+				{
+					// Have to sign when a bundle name is specified
+					bNeedToSign = true;
+					IPPArguments.Add("-bundlename");
+					IPPArguments.Add(Params.BundleName);
+				}
+
+				if (bNeedToSign)
+				{
+					IPPArguments.Add("-sign");
+				}
+
+				if (cookonthefly)
+				{
+					IPPArguments.Add(" -cookonthefly");
+				}
+				IPPArguments.Add(" -stagedir");
+				IPPArguments.Add(CombinePaths(Params.BaseStageDirectory, "IOS"));
+				IPPArguments.Add(" -projectdir");
+				IPPArguments.Add(Path.GetDirectoryName(Params.RawProjectPath));
+
+				if (RunIPP(IPPArguments.ToArray()) != 0)
+				{
+					throw new AutomationException("IPP Failed");
+				}
+			}
 
 			// verify the .ipa exists
 			if (!FileExists(ProjectIPA))
diff --git a/Engine/Source/Programs/AutomationTool/IOS/IPhonePackager/CompileTime.cs b/Engine/Source/Programs/AutomationTool/IOS/IPhonePackager/CompileTime.cs
--- a/Engine/Source/Programs/AutomationTool/IOS/IPhonePackager/CompileTime.cs
+++ b/Engine/Source/Programs/AutomationTool/IOS/IPhonePackager/CompileTime.cs
@@ -11,7 +11,6 @@ using System.Net;
 using System.Runtime.InteropServices;
 using System.Text;
 using System.Threading;
-using System.Windows.Forms;
 using Ionic.Zlib;
 using UnrealBuildTool;
 
diff --git a/Engine/Source/Programs/AutomationTool/IOS/IPhonePackager/ConfigureMobileGame.cs b/Engine/Source/Programs/AutomationTool/IOS/IPhonePackager/ConfigureMobileGame.cs
deleted file mode 100644
diff --git a/Engine/Source/Programs/AutomationTool/IOS/IPhonePackager/CookTime.cs b/Engine/Source/Programs/AutomationTool/IOS/IPhonePackager/CookTime.cs
--- a/Engine/Source/Programs/AutomationTool/IOS/IPhonePackager/CookTime.cs
+++ b/Engine/Source/Programs/AutomationTool/IOS/IPhonePackager/CookTime.cs
@@ -10,7 +10,6 @@ using System.Net;
 using System.Runtime.InteropServices;
 using System.Text;
 using System.Threading;
-using System.Windows.Forms;
 using Ionic.Zip;
 using Ionic.Zlib;
 
diff --git a/Engine/Source/Programs/AutomationTool/IOS/IPhonePackager/DeployTime.cs b/Engine/Source/Programs/AutomationTool/IOS/IPhonePackager/DeployTime.cs
--- a/Engine/Source/Programs/AutomationTool/IOS/IPhonePackager/DeployTime.cs
+++ b/Engine/Source/Programs/AutomationTool/IOS/IPhonePackager/DeployTime.cs
@@ -10,7 +10,6 @@ using System.Net;
 using System.Runtime.InteropServices;
 using System.Text;
 using System.Threading;
-using System.Windows.Forms;
 using Microsoft.Win32;
 using System.Linq;
 using System.Runtime.Remoting.Channels.Ipc;
diff --git a/Engine/Source/Programs/AutomationTool/IOS/IPhonePackager/GenerateSigningRequestDialog.cs b/Engine/Source/Programs/AutomationTool/IOS/IPhonePackager/GenerateSigningRequestDialog.cs
deleted file mode 100644
diff --git a/Engine/Source/Programs/AutomationTool/IOS/IPhonePackager/GraphicalResignTool.cs b/Engine/Source/Programs/AutomationTool/IOS/IPhonePackager/GraphicalResignTool.cs
deleted file mode 100644
diff --git a/Engine/Source/Programs/AutomationTool/IOS/IPhonePackager/PasswordDialog.cs b/Engine/Source/Programs/AutomationTool/IOS/IPhonePackager/PasswordDialog.cs
deleted file mode 100644
diff --git a/Engine/Source/Programs/AutomationTool/IOS/IPhonePackager/SlowProgressDialog.cs b/Engine/Source/Programs/AutomationTool/IOS/IPhonePackager/SlowProgressDialog.cs
deleted file mode 100644
diff --git a/Engine/Source/Programs/AutomationTool/IOS/IPhonePackager/ToolsHub.cs b/Engine/Source/Programs/AutomationTool/IOS/IPhonePackager/ToolsHub.cs
deleted file mode 100644
diff --git a/Engine/Source/Programs/AutomationTool/IOS/IPhonePackager/iPhonePackager.cs b/Engine/Source/Programs/AutomationTool/IOS/IPhonePackager/iPhonePackager.cs
--- a/Engine/Source/Programs/AutomationTool/IOS/IPhonePackager/iPhonePackager.cs
+++ b/Engine/Source/Programs/AutomationTool/IOS/IPhonePackager/iPhonePackager.cs
@@ -11,7 +11,6 @@ using System.Reflection;
 using System.Runtime.InteropServices;
 using System.Text;
 using System.Threading;
-using System.Windows.Forms;
 using Ionic.Zlib;
 using System.ComponentModel;
 
@@ -27,7 +26,6 @@ namespace iPhonePackager
 		static public string GameConfiguration = "";
 		static public string Architecture = "";
 
-		static public SlowProgressDialog ProgressDialog = null;
 		static public BackgroundWorker BGWorker = null;
 		static public int ProgressIndex = 0;
 
@@ -266,15 +264,8 @@ namespace iPhonePackager
 				string CWD = Directory.GetCurrentDirectory();
 				if (!CodeSignatureBuilder.DoRequiredFilesExist())
 				{
-					StartVisuals();
-
-					Form DisplayForm = new ToolsHub();
-					Application.Run(DisplayForm);
-					if (DisplayForm.DialogResult != DialogResult.OK)
-					{
-						Error("One or more files necessary for packaging are missing and configuration was cancelled");
-						return false;
-					}
+					Error("One or more files necessary for packaging are missing - please use Project Settings in the editor to set up");
+					return false;
 				}
 				Directory.SetCurrentDirectory(CWD);
 			}
@@ -311,28 +302,6 @@ namespace iPhonePackager
 					bHandledCommand = DeploymentHelper.ExecuteDeployCommand(Command, GamePath, RPCCommand);
 				}
 
-				if (!bHandledCommand)
-				{
-					bHandledCommand = true;
-					switch (Command)
-					{
-						case "configure":
-							if (Config.bForDistribution)
-							{
-								Error("configure cannot be used with -distribution");
-							}
-							else
-							{
-								RunInVisualMode(delegate { return new ConfigureMobileGame(); });
-							}
-							break;
-
-						default:
-							bHandledCommand = false;
-							break;
-					}
-				}
-
 				if (!bHandledCommand)
 				{
 					Error("Unknown command");
@@ -409,24 +378,6 @@ namespace iPhonePackager
 			return ( null );
 		}
 
-		delegate Form CreateFormDelegate();
-
-		static void StartVisuals()
-		{
-			Application.EnableVisualStyles();
-			Application.SetCompatibleTextRenderingDefault(false);
-
-			ProgressDialog = new SlowProgressDialog();
-		}
-
-		static void RunInVisualMode(CreateFormDelegate Work)
-		{
-			StartVisuals();
-
-			Form DisplayForm = Work.Invoke();
-			Application.Run(DisplayForm);
-		}
-
 		static void ListDevices()
 		{
 			var DeviceList = DeploymentHelper.Get().EnumerateConnectedDevices();
@@ -574,18 +525,6 @@ namespace iPhonePackager
 						}
 						break;
 
-					case "resigntool":
-						RunInVisualMode(delegate { return new GraphicalResignTool(); });
-						break;
-
-					case "certrequest":
-						RunInVisualMode(delegate { return new GenerateSigningRequestDialog(); });
-						break;
-
-					case "gui":
-						RunInVisualMode(delegate { return ToolsHub.CreateShowingTools(); });
-						break;
-
 					case "devices":
 						ListDevices();
 						break;
