Fixed pak setting not working in ufe when cook ont he fly is enabled.

diff --git a/Engine/Source/Programs/AutomationTool/Android/AndroidPlatform.Automation.cs b/Engine/Source/Programs/AutomationTool/Android/AndroidPlatform.Automation.cs
--- a/Engine/Source/Programs/AutomationTool/Android/AndroidPlatform.Automation.cs
+++ b/Engine/Source/Programs/AutomationTool/Android/AndroidPlatform.Automation.cs
@@ -125,8 +125,10 @@ public class AndroidPlatform : Platform
 		// first, look for a .pak file in the staged directory
 		string[] PakFiles = Directory.GetFiles(SC.StageDirectory, "*.pak", SearchOption.AllDirectories);
 
+        bool bHasPakFile = PakFiles.Length > 1;
+
 		// for now, we only support 1 pak/obb file
-		if (PakFiles.Length != 1)
+		if (PakFiles.Length > 1)
 		{
 			throw new AutomationException("Can't package for Android with 0 or more than 1 pak file (found {0} pak files in {1})", PakFiles.Length, SC.StageDirectory);
 		}
@@ -140,7 +142,7 @@ public class AndroidPlatform : Platform
             File.Delete(LocalObbName);
         }
 
-        if (!Params.OBBinAPK)
+        if (!Params.OBBinAPK && bHasPakFile)
         {
             Log("Creating {0} from {1}", LocalObbName, PakFiles[0]);
             File.Copy(PakFiles[0], LocalObbName);
@@ -160,8 +162,8 @@ public class AndroidPlatform : Platform
 			"%ADB% %DEVICE% shell rm -r /mnt/sdcard/" + Params.ShortProjectName,
 			"%ADB% %DEVICE% shell rm -r /mnt/sdcard/UE4Game/UE4CommandLine.txt", // we need to delete the commandline in UE4Game or it will mess up loading
 			"%ADB% %DEVICE% shell rm -r /mnt/sdcard/obb/" + PackageName,
-			Params.OBBinAPK ? "" : "%ADB% %DEVICE% push " + Path.GetFileName(LocalObbName) + " " + DeviceObbName,
-			Params.OBBinAPK ? "" : "if \"%ERRORLEVEL%\" NEQ \"0\" goto Error",
+			Params.OBBinAPK && bHasPakFile ? "" : "%ADB% %DEVICE% push " + Path.GetFileName(LocalObbName) + " " + DeviceObbName,
+			Params.OBBinAPK && bHasPakFile ? "" : "if \"%ERRORLEVEL%\" NEQ \"0\" goto Error",
 			"goto:eof",
 			":Error",
 			"@echo.",
diff --git a/Engine/Source/Programs/AutomationTool/Scripts/CopyBuildToStagingDirectory.Automation.cs b/Engine/Source/Programs/AutomationTool/Scripts/CopyBuildToStagingDirectory.Automation.cs
--- a/Engine/Source/Programs/AutomationTool/Scripts/CopyBuildToStagingDirectory.Automation.cs
+++ b/Engine/Source/Programs/AutomationTool/Scripts/CopyBuildToStagingDirectory.Automation.cs
@@ -57,7 +57,7 @@ public partial class Project : CommandUtils
 	{
 		if (UnrealPakResponseFile.Count < 1)
 		{
-			throw new AutomationException("items to pak");
+            return;
 		}
 		string PakName = Path.GetFileNameWithoutExtension(OutputLocation);
 		string UnrealPakResponseFileName = CombinePaths(CmdEnv.LogFolder, "PakList_" + PakName + ".txt");
