Fixed error in deployment script when obbinapk was enabled but there was no pak file.

diff --git a/Engine/Source/Programs/AutomationTool/Android/AndroidPlatform.Automation.cs b/Engine/Source/Programs/AutomationTool/Android/AndroidPlatform.Automation.cs
--- a/Engine/Source/Programs/AutomationTool/Android/AndroidPlatform.Automation.cs
+++ b/Engine/Source/Programs/AutomationTool/Android/AndroidPlatform.Automation.cs
@@ -162,8 +162,8 @@ public class AndroidPlatform : Platform
 			"%ADB% %DEVICE% shell rm -r /mnt/sdcard/" + Params.ShortProjectName,
 			"%ADB% %DEVICE% shell rm -r /mnt/sdcard/UE4Game/UE4CommandLine.txt", // we need to delete the commandline in UE4Game or it will mess up loading
 			"%ADB% %DEVICE% shell rm -r /mnt/sdcard/obb/" + PackageName,
-			Params.OBBinAPK && bHasPakFile ? "" : "%ADB% %DEVICE% push " + Path.GetFileName(LocalObbName) + " " + DeviceObbName,
-			Params.OBBinAPK && bHasPakFile ? "" : "if \"%ERRORLEVEL%\" NEQ \"0\" goto Error",
+			Params.OBBinAPK || !bHasPakFile ? "" : "%ADB% %DEVICE% push " + Path.GetFileName(LocalObbName) + " " + DeviceObbName,
+			Params.OBBinAPK || !bHasPakFile ? "" : "if \"%ERRORLEVEL%\" NEQ \"0\" goto Error",
 			"goto:eof",
 			":Error",
 			"@echo.",
