Fix for pak file being incorrectly detected from deploment script.

diff --git a/Engine/Source/Programs/AutomationTool/Android/AndroidPlatform.Automation.cs b/Engine/Source/Programs/AutomationTool/Android/AndroidPlatform.Automation.cs
--- a/Engine/Source/Programs/AutomationTool/Android/AndroidPlatform.Automation.cs
+++ b/Engine/Source/Programs/AutomationTool/Android/AndroidPlatform.Automation.cs
@@ -125,7 +125,7 @@ public class AndroidPlatform : Platform
 		// first, look for a .pak file in the staged directory
 		string[] PakFiles = Directory.GetFiles(SC.StageDirectory, "*.pak", SearchOption.AllDirectories);
 
-        bool bHasPakFile = PakFiles.Length > 1;
+        bool bHasPakFile = PakFiles.Length >= 1;
 
 		// for now, we only support 1 pak/obb file
 		if (PakFiles.Length > 1)
