Only sign EXE files on Windows platforms.

diff --git a/Engine/Source/Programs/AutomationTool/CommandUtils.cs b/Engine/Source/Programs/AutomationTool/CommandUtils.cs
--- a/Engine/Source/Programs/AutomationTool/CommandUtils.cs
+++ b/Engine/Source/Programs/AutomationTool/CommandUtils.cs
@@ -2353,7 +2353,7 @@ namespace AutomationTool
 						SignMacFileOrFolder(File);
 					}
 				}
-				else
+				else if(TargetPlatform == UnrealBuildTool.UnrealTargetPlatform.Win32 || TargetPlatform == UnrealTargetPlatform.Win64 || TargetPlatform == UnrealTargetPlatform.WinRT || TargetPlatform == UnrealTargetPlatform.WinRT_ARM)
 				{
 					foreach (var File in Files)
 					{
