AutomationTool: Make ChunkManifest hashsets case insensitive

diff --git a/Engine/Source/Programs/AutomationTool/Scripts/CopyBuildToStagingDirectory.Automation.cs b/Engine/Source/Programs/AutomationTool/Scripts/CopyBuildToStagingDirectory.Automation.cs
--- a/Engine/Source/Programs/AutomationTool/Scripts/CopyBuildToStagingDirectory.Automation.cs
+++ b/Engine/Source/Programs/AutomationTool/Scripts/CopyBuildToStagingDirectory.Automation.cs
@@ -49,7 +49,7 @@ public partial class Project : CommandUtils
 	private static HashSet<string> ReadPakChunkManifest(string Filename)
 	{		
 		var ResponseFile = ReadAllLines(Filename);
-		var Result = new HashSet<string>(ResponseFile);
+		var Result = new HashSet<string>(ResponseFile, StringComparer.InvariantCultureIgnoreCase);
 		return Result;
 	}
 
