Removing a stray unused var that's causing some mac build issues.

Fixes:
sh: IPhonePackager/Utilities.cs(451,13): error CS0219: Warning as Error: The variable `DictionaryNode' is assigned but its value is never used

diff --git a/Engine/Source/Programs/AutomationTool/IOS/IPhonePackager/Utilities.cs b/Engine/Source/Programs/AutomationTool/IOS/IPhonePackager/Utilities.cs
--- a/Engine/Source/Programs/AutomationTool/IOS/IPhonePackager/Utilities.cs
+++ b/Engine/Source/Programs/AutomationTool/IOS/IPhonePackager/Utilities.cs
@@ -448,8 +448,6 @@ namespace iPhonePackager
 			/// </summary>
 			public bool HasKey(string KeyName)
 			{
-				XmlNode DictionaryNode = Doc.DocumentElement.SelectSingleNode("/plist/dict");
-
 				string PathToKey = String.Format("/plist/dict/key[.='{0}']", KeyName);
 
 				XmlNode KeyNode = Doc.DocumentElement.SelectSingleNode(PathToKey);
