Partial integration of the changes from 2125367 in the file
//depot/UE4/Engine/Source/Programs/IOS/iPhonePackager/CompileTime.cs

Also removed the other unused property. 'MacBinariesDir'

diff --git a/Engine/Source/Programs/AutomationTool/IOS/IPhonePackager/CompileTime.cs b/Engine/Source/Programs/AutomationTool/IOS/IPhonePackager/CompileTime.cs
--- a/Engine/Source/Programs/AutomationTool/IOS/IPhonePackager/CompileTime.cs
+++ b/Engine/Source/Programs/AutomationTool/IOS/IPhonePackager/CompileTime.cs
@@ -32,7 +32,6 @@ namespace iPhonePackager
 		private static string XcodeDeveloperDir = "";
 
 		private static string iPhone_SigningDevRootMac = "";
-		private static string iPhone_SigningDevRootWin = "";
 
 		/// <summary>
 		/// The file name (no path) of the temporary mobile provision that will be placed on the remote mac for use in makeapp
@@ -41,7 +40,6 @@ namespace iPhonePackager
 
 		private static string MacName = "";
 		private static string MacStagingRootDir = "";
-		private static string MacBinariesDir = "";
 		private static string MacXcodeStagingDir = "";
 
 		/** /MacStagingRootDir/Payload */
@@ -99,7 +97,6 @@ namespace iPhonePackager
 
 			// MacName=%ue4.iPhone_SigningServerName%
 			MacName = Config.OverrideMacName != null ? Config.OverrideMacName : Utilities.GetEnvironmentVariable( "ue.IOSSigningServer", "a1487" );
-			iPhone_SigningDevRootWin = "\\\\" + MacName + "\\UE4\\Builds";
 			iPhone_SigningDevRootMac = "/UE4/Builds";
 
 			// get the path to mirror into on the Mac
@@ -127,7 +124,6 @@ namespace iPhonePackager
 
 			// generate the directories to recursively copy into later on
 			MacStagingRootDir = string.Format("{0}/{1}/IOS", iPhone_SigningDevRootMac, GameBranchPath);
-			MacBinariesDir = string.Format("{0}/{1}/IOS", iPhone_SigningDevRootMac, GameBranchPath);
 			MacXcodeStagingDir = string.Format("{0}/{1}/IOS/XcodeSupportFiles", iPhone_SigningDevRootMac, GameBranchPath);
 
 			MacMobileProvisionFilename = MachineName + "_UE4Temp.mobileprovision";
