build script, work around copies from the share failing to get the right date

diff --git a/Engine/Source/Programs/AutomationTool/TempStorage.cs b/Engine/Source/Programs/AutomationTool/TempStorage.cs
--- a/Engine/Source/Programs/AutomationTool/TempStorage.cs
+++ b/Engine/Source/Programs/AutomationTool/TempStorage.cs
@@ -67,7 +67,7 @@ namespace AutomationTool
                     bSilentOkToBeDifferent = bSilentOkToBeDifferent || Name.Contains("Engine/Binaries/Mac");
 
                     System.Diagnostics.TraceEventType LogType = bOkToBeDifferent ? System.Diagnostics.TraceEventType.Warning : System.Diagnostics.TraceEventType.Error;
-                    if (bSilentOkToBeDifferent && bOkToBeDifferent)
+                    if (bSilentOkToBeDifferent)
                     {
                         LogType = System.Diagnostics.TraceEventType.Information;
                     }
@@ -75,7 +75,7 @@ namespace AutomationTool
                     if (!((Timestamp - Other.Timestamp).TotalSeconds < 1 && (Timestamp - Other.Timestamp).TotalSeconds > -1))
                     {
                         CommandUtils.Log(LogType, "File date mismatch {0} {1} {2} {3}", Name, Timestamp.ToString(), Other.Name, Other.Timestamp.ToString());
-                        bOk = bOkToBeDifferent;
+                        bOk = bOkToBeDifferent || bSilentOkToBeDifferent;
                     }
                     if (!(Size == Other.Size))
                     {
