fix to not sign the xboxone executables with Verisign.
check for XboxOne in the path
#ue4
#xboxone

diff --git a/Engine/Source/Programs/AutomationTool/CommandUtils.cs b/Engine/Source/Programs/AutomationTool/CommandUtils.cs
--- a/Engine/Source/Programs/AutomationTool/CommandUtils.cs
+++ b/Engine/Source/Programs/AutomationTool/CommandUtils.cs
@@ -2353,11 +2353,14 @@ namespace AutomationTool
 						SignMacFileOrFolder(File);
 					}
 				}
-				else if(TargetPlatform == UnrealBuildTool.UnrealTargetPlatform.Win32 || TargetPlatform == UnrealTargetPlatform.Win64 || TargetPlatform == UnrealTargetPlatform.WinRT || TargetPlatform == UnrealTargetPlatform.WinRT_ARM)
+				else
 				{
 					foreach (var File in Files)
 					{
-						SignSingleExecutableIfEXEOrDLL(File);
+						if (!Path.GetDirectoryName().Contains("XboxOne"))
+						{
+							SignSingleExecutableIfEXEOrDLL(File);
+						}
 					}
 				}
 			}
