compile fix and change it so it won't ignore XboxOne target dlls for verisign
#ue4
#xboxone

diff --git a/Engine/Source/Programs/AutomationTool/CommandUtils.cs b/Engine/Source/Programs/AutomationTool/CommandUtils.cs
--- a/Engine/Source/Programs/AutomationTool/CommandUtils.cs
+++ b/Engine/Source/Programs/AutomationTool/CommandUtils.cs
@@ -2355,9 +2355,9 @@ namespace AutomationTool
 				}
 				else
 				{
-					foreach (var File in Files)
+					foreach (string File in Files)
 					{
-						if (!Path.GetDirectoryName().Contains("XboxOne"))
+						if (!(Path.GetDirectoryName(File).Replace("\\", "/")).Contains("Binaries/XboxOne"))
 						{
 							SignSingleExecutableIfEXEOrDLL(File);
 						}
