#ttp 328717 ANDROID: Launching app remotely on locked device breaks display

- Added call to adb which sends an input to unlock the screen before launching the game

diff --git a/Engine/Source/Programs/AutomationTool/Android/AndroidPlatform.Automation.cs b/Engine/Source/Programs/AutomationTool/Android/AndroidPlatform.Automation.cs
--- a/Engine/Source/Programs/AutomationTool/Android/AndroidPlatform.Automation.cs
+++ b/Engine/Source/Programs/AutomationTool/Android/AndroidPlatform.Automation.cs
@@ -579,6 +579,9 @@ public class AndroidPlatform : Platform
             Run(CmdEnv.CmdExe, AdbCommand + "logcat -c");
         }
         
+        // Send a command to unlock the device before we try to run it
+        string UnlockCommandLine = "shell input keyevent 82";
+        ProcessResult UnlockProcess = Run(CmdEnv.CmdExe, AdbCommand + UnlockCommandLine, null);
 
 		// start the app on device!
 		ProcessResult ClientProcess = Run(CmdEnv.CmdExe, AdbCommand + CommandLine, null, ClientRunFlags);
