#ttp 328717 Android
- compile fix for build warning/error

diff --git a/Engine/Source/Programs/AutomationTool/Android/AndroidPlatform.Automation.cs b/Engine/Source/Programs/AutomationTool/Android/AndroidPlatform.Automation.cs
--- a/Engine/Source/Programs/AutomationTool/Android/AndroidPlatform.Automation.cs
+++ b/Engine/Source/Programs/AutomationTool/Android/AndroidPlatform.Automation.cs
@@ -581,7 +581,7 @@ public class AndroidPlatform : Platform
         
         // Send a command to unlock the device before we try to run it
         string UnlockCommandLine = "shell input keyevent 82";
-        ProcessResult UnlockProcess = Run(CmdEnv.CmdExe, AdbCommand + UnlockCommandLine, null);
+        Run(CmdEnv.CmdExe, AdbCommand + UnlockCommandLine, null);
 
 		// start the app on device!
 		ProcessResult ClientProcess = Run(CmdEnv.CmdExe, AdbCommand + CommandLine, null, ClientRunFlags);
