AutomationTool: Trying to work around exceptions when loading assemblies (possibly due to .Net 4.0 bug). Cleaning up referenced assemblies, better logging.

diff --git a/Engine/Source/Programs/AutomationTool/HTML5/HTML5.Automation.csproj b/Engine/Source/Programs/AutomationTool/HTML5/HTML5.Automation.csproj
--- a/Engine/Source/Programs/AutomationTool/HTML5/HTML5.Automation.csproj
+++ b/Engine/Source/Programs/AutomationTool/HTML5/HTML5.Automation.csproj
@@ -36,14 +36,7 @@
   <ItemGroup>
     <Reference Include="System" />
     <Reference Include="System.Core" />
-    <Reference Include="System.Runtime.Serialization" />
-    <Reference Include="System.Web" />
-    <Reference Include="System.Web.Extensions" />
-    <Reference Include="System.Xml.Linq" />
-    <Reference Include="System.Data.DataSetExtensions" />
     <Reference Include="Microsoft.CSharp" />
-    <Reference Include="System.Data" />
-    <Reference Include="System.Xml" />
   </ItemGroup>
   <ItemGroup>
     <Compile Include="..\..\DotNETCommon\MetaData.cs">
diff --git a/Engine/Source/Programs/AutomationTool/IOS/IOS.Automation.csproj b/Engine/Source/Programs/AutomationTool/IOS/IOS.Automation.csproj
--- a/Engine/Source/Programs/AutomationTool/IOS/IOS.Automation.csproj
+++ b/Engine/Source/Programs/AutomationTool/IOS/IOS.Automation.csproj
@@ -50,12 +50,7 @@
     <Reference Include="System.Core" />
     <Reference Include="System.Drawing" />
     <Reference Include="System.Runtime.Remoting" />
-    <Reference Include="System.Runtime.Serialization" />
     <Reference Include="System.Security" />
-    <Reference Include="System.Web" />
-    <Reference Include="System.Web.Extensions" />
-    <Reference Include="System.Xml.Linq" />
-    <Reference Include="System.Data.DataSetExtensions" />
     <Reference Include="Microsoft.CSharp" />
     <Reference Include="System.Data" />
     <Reference Include="System.Xml" />
diff --git a/Engine/Source/Programs/AutomationTool/Platform.cs b/Engine/Source/Programs/AutomationTool/Platform.cs
--- a/Engine/Source/Programs/AutomationTool/Platform.cs
+++ b/Engine/Source/Programs/AutomationTool/Platform.cs
@@ -42,7 +42,34 @@ namespace AutomationTool
 
 		private static void CreatePlatformsFromAssembly(Assembly ScriptAssembly)
 		{
-			var AllTypes = ScriptAssembly.GetTypes();
+			Log("Looking for platforms in {0}", ScriptAssembly.Location);
+			Type[] AllTypes = null;
+			try
+			{
+				AllTypes = ScriptAssembly.GetTypes();
+			}
+			catch (Exception Ex)
+			{
+				LogError("Failed to get assembly types for {0}", ScriptAssembly.Location);
+				if (Ex is ReflectionTypeLoadException)
+				{					
+					var TypeLoadException = (ReflectionTypeLoadException)Ex;
+					if (!IsNullOrEmpty(TypeLoadException.LoaderExceptions))
+					{
+						LogError("Loader Exceptions:");
+						foreach (var LoaderException in TypeLoadException.LoaderExceptions)
+						{
+							Log(System.Diagnostics.TraceEventType.Error, LoaderException);
+						}
+					}
+					else
+					{
+						Log("No Loader Exceptions available.");
+					}
+				}
+				// Re-throw, this is still a critical error!
+				throw Ex;
+			}
 			foreach (var PotentialPlatformType in AllTypes)
 			{
 				if (PotentialPlatformType != typeof(Platform) && typeof(Platform).IsAssignableFrom(PotentialPlatformType) && !PotentialPlatformType.IsAbstract)
