Back out changelist 2181411 (fixing TTP# 341148)

diff --git a/Engine/Source/Programs/AutomationTool/IOS/IOS.Automation.csproj b/Engine/Source/Programs/AutomationTool/IOS/IOS.Automation.csproj
--- a/Engine/Source/Programs/AutomationTool/IOS/IOS.Automation.csproj
+++ b/Engine/Source/Programs/AutomationTool/IOS/IOS.Automation.csproj
@@ -34,23 +34,12 @@
     <TreatWarningsAsErrors>true</TreatWarningsAsErrors>
   </PropertyGroup>
   <ItemGroup>
-    <Reference Include="BouncyCastle.Crypto">
-      <HintPath>..\..\..\..\Binaries\ThirdParty\IOS\BouncyCastle.Crypto.dll</HintPath>
-      <Private>False</Private>
-    </Reference>
     <Reference Include="Ionic.Zip.Reduced">
       <HintPath>..\..\..\..\Binaries\DotNET\Ionic.Zip.Reduced.dll</HintPath>
       <Private>False</Private>
     </Reference>
-    <Reference Include="RPCUtility">
-      <HintPath>..\..\..\..\Binaries\DotNET\RPCUtility.exe</HintPath>
-      <Private>False</Private>
-    </Reference>
     <Reference Include="System" />
     <Reference Include="System.Core" />
-    <Reference Include="System.Drawing" />
-    <Reference Include="System.Runtime.Remoting" />
-    <Reference Include="System.Security" />
     <Reference Include="Microsoft.CSharp" />
     <Reference Include="System.Data" />
     <Reference Include="System.Xml" />
@@ -60,24 +49,9 @@
       <Link>Properties\MetaData.cs</Link>
     </Compile>
     <Compile Include="IOSPlatform.Automation.cs" />
-    <Compile Include="IPhonePackager\CodeSigning.cs" />
-    <Compile Include="IPhonePackager\CompileTime.cs" />
-    <Compile Include="IPhonePackager\Config.cs" />
-    <Compile Include="IPhonePackager\CookTime.cs" />
-    <Compile Include="IPhonePackager\DeployTime.cs" />
-    <Compile Include="IPhonePackager\FileOperations.cs" />
-    <Compile Include="IPhonePackager\iPhonePackager.cs" />
-    <Compile Include="IPhonePackager\MachObjects.cs" />
-    <Compile Include="IPhonePackager\MobileProvisionUtilities.cs" />
-    <Compile Include="IPhonePackager\Utilities.cs" />
     <Compile Include="Properties\AssemblyInfo.cs" />
-    <Compile Include="Properties\Resources.Designer.cs" />
   </ItemGroup>
   <ItemGroup>
-    <ProjectReference Include="..\..\IOS\DeploymentInterface\DeploymentInterface.csproj">
-      <Project>{2da5397a-f93f-41e0-a1ad-b220d40fd739}</Project>
-      <Name>DeploymentInterface</Name>
-    </ProjectReference>
     <ProjectReference Include="..\AutomationTool.csproj">
       <Project>{a35b0261-d868-475e-a098-eb14ce6dc89f}</Project>
       <Name>AutomationTool</Name>
@@ -89,13 +63,5 @@
       <Private>False</Private>
     </ProjectReference>
   </ItemGroup>
-  <ItemGroup>
-    <EmbeddedResource Include="Properties\Resources.resx">
-      <SubType>Designer</SubType>
-    </EmbeddedResource>
-  </ItemGroup>
-  <ItemGroup>
-    <None Include="Resources\Image1.bmp" />
-  </ItemGroup>
   <Import Project="$(MSBuildToolsPath)\Microsoft.CSharp.targets" />
 </Project>
\ No newline at end of file
diff --git a/Engine/Source/Programs/AutomationTool/IOS/IOSPlatform.Automation.cs b/Engine/Source/Programs/AutomationTool/IOS/IOSPlatform.Automation.cs
--- a/Engine/Source/Programs/AutomationTool/IOS/IOSPlatform.Automation.cs
+++ b/Engine/Source/Programs/AutomationTool/IOS/IOSPlatform.Automation.cs
@@ -6,7 +6,6 @@ using System.Text;
 using System.IO;
 using AutomationTool;
 using UnrealBuildTool;
-using System.Text.RegularExpressions;
 using Ionic.Zip;
 using Ionic.Zlib;
 
@@ -23,65 +22,6 @@ public class IOSPlatform : Platform
 	{
 	}
 
-	// Run the integrated IPP code
-	// @todo ipp: How to log the output in a nice UAT way?
-	protected static int RunIPP(string IPPArguments)
-	{
-		List<string> Args = new List<string>();
-
-		StringBuilder Token = null;
-		int Index = 0;
-		bool bInQuote = false;
-		while (Index < IPPArguments.Length)
-		{
-			if (IPPArguments[Index] == '\"')
-			{
-				bInQuote = !bInQuote;
-			}
-			else if (IPPArguments[Index] == ' ' && !bInQuote)
-			{
-				if (Token != null)
-				{
-					Args.Add(Token.ToString());
-					Token = null;
-				}
-			}
-			else
-			{
-				if (Token == null)
-				{
-					Token = new StringBuilder();
-				}
-				Token.Append(IPPArguments[Index]);
-			}
-			Index++;
-		}
-
-		if (Token != null)
-		{
-			Args.Add(Token.ToString());
-		}
-	
-		return RunIPP(Args.ToArray());
-	}
-
-	protected static int RunIPP(string[] Args)
-	{
-		return iPhonePackager.Program.RealMain(Args);
-	}
-
-	public override int RunCommand(string Command)
-	{
-		// run generic IPP commands through the interface
-		if (Command.StartsWith("IPP:"))
-		{
-			return RunIPP(Command.Substring(4));
-		}
-
-		// non-zero error code
-		return 4;
-	}
-
 	protected string MakeIPAFileName( UnrealTargetConfiguration TargetConfiguration, string ProjectGameExeFilename )
 	{
 		var ProjectIPA = Path.ChangeExtension(ProjectGameExeFilename, null);
@@ -184,93 +124,47 @@ public class IOSPlatform : Platform
 
 			bool cookonthefly = Params.CookOnTheFly || Params.SkipCookOnTheFly;
 
-			// rename the .ipa if not code based
-			if (!Params.IsCodeBasedProject)
+			string IPPArguments = "RepackageFromStage \"" + (Params.IsCodeBasedProject ? Params.RawProjectPath : "Engine") + "\"";
+			IPPArguments += " -config " + TargetConfiguration.ToString();
+
+			if (TargetConfiguration == UnrealTargetConfiguration.Shipping)
 			{
-				ProjectIPA = Path.Combine(Path.GetDirectoryName(Params.RawProjectPath), "Binaries", "IOS", Params.ShortProjectName + ".ipa");
-				if (TargetConfiguration != UnrealTargetConfiguration.Development)
-				{
-					ProjectIPA = Path.Combine(Path.GetDirectoryName(Params.RawProjectPath), "Binaries", "IOS", Params.ShortProjectName + "-" + PlatformType.ToString() + "-" + TargetConfiguration.ToString() + ".ipa");
-				}
+				IPPArguments += " -compress=best";
 			}
 
-			// delete the .ipa to make sure it was made
-			DeleteFile(ProjectIPA);
+			// Determine if we should sign
+			bool bNeedToSign = GetCodeSignDesirability(Params);
 
-			if (RemoteToolChain.bUseRPCUtil)
+			if (!String.IsNullOrEmpty(Params.BundleName))
 			{
-				string IPPArguments = "RepackageFromStage \"" + (Params.IsCodeBasedProject ? Params.RawProjectPath : "Engine") + "\"";
-				IPPArguments += " -config " + TargetConfiguration.ToString();
-
-				if (TargetConfiguration == UnrealTargetConfiguration.Shipping)
-				{
-					IPPArguments += " -compress=best";
-				}
-
-				// Determine if we should sign
-				bool bNeedToSign = GetCodeSignDesirability(Params);
-
-				if (!String.IsNullOrEmpty(Params.BundleName))
-				{
-					// Have to sign when a bundle name is specified
-					bNeedToSign = true;
-					IPPArguments += " -bundlename " + Params.BundleName;
-				}
-
-				if (bNeedToSign)
-				{
-					IPPArguments += " -sign";
-				}
-
-				IPPArguments += (cookonthefly ? " -cookonthefly" : "");
-				IPPArguments += " -stagedir \"" + CombinePaths(Params.BaseStageDirectory, "IOS") + "\"";
-				IPPArguments += " -projectdir \"" + Path.GetDirectoryName(Params.RawProjectPath) + "\"";
-
-				RunAndLog(CmdEnv, IPPExe, IPPArguments);
+				// Have to sign when a bundle name is specified
+				bNeedToSign = true;
+				IPPArguments += " -bundlename " + Params.BundleName;
 			}
-			else
-			{
-				List<string> IPPArguments = new List<string>();
-				IPPArguments.Add("RepackageFromStage");
-				IPPArguments.Add(Params.IsCodeBasedProject ? Params.RawProjectPath : "Engine");
-				IPPArguments.Add("-config");
-				IPPArguments.Add(TargetConfiguration.ToString());
 
-				if (TargetConfiguration == UnrealTargetConfiguration.Shipping)
-				{
-					IPPArguments.Add("-compress=best");
-				}
-
-				// Determine if we should sign
-				bool bNeedToSign = GetCodeSignDesirability(Params);
+			if (bNeedToSign)
+			{
+				IPPArguments += " -sign";
+			}
 
-				if (!String.IsNullOrEmpty(Params.BundleName))
-				{
-					// Have to sign when a bundle name is specified
-					bNeedToSign = true;
-					IPPArguments.Add("-bundlename");
-					IPPArguments.Add(Params.BundleName);
-				}
+			IPPArguments += (cookonthefly ? " -cookonthefly" : "");
+			IPPArguments += " -stagedir \"" + CombinePaths(Params.BaseStageDirectory, "IOS") + "\"";
+			IPPArguments += " -projectdir \"" + Path.GetDirectoryName(Params.RawProjectPath) + "\"";
 
-				if (bNeedToSign)
+			// rename the .ipa if not code based
+			if (!Params.IsCodeBasedProject)
+			{
+				ProjectIPA = Path.Combine(Path.GetDirectoryName(Params.RawProjectPath), "Binaries", "IOS", Params.ShortProjectName + ".ipa");
+				if (TargetConfiguration != UnrealTargetConfiguration.Development)
 				{
-					IPPArguments.Add("-sign");
+					ProjectIPA = Path.Combine(Path.GetDirectoryName(Params.RawProjectPath), "Binaries", "IOS", Params.ShortProjectName + "-" + PlatformType.ToString() + "-" + TargetConfiguration.ToString() + ".ipa");
 				}
+			}
 
-				if (cookonthefly)
-				{
-					IPPArguments.Add(" -cookonthefly");
-				}
-				IPPArguments.Add(" -stagedir");
-				IPPArguments.Add(CombinePaths(Params.BaseStageDirectory, "IOS"));
-				IPPArguments.Add(" -projectdir");
-				IPPArguments.Add(Path.GetDirectoryName(Params.RawProjectPath));
+            // delete the .ipa to make sure it was made
+            DeleteFile(ProjectIPA);
 
-				if (RunIPP(IPPArguments.ToArray()) != 0)
-				{
-					throw new AutomationException("IPP Failed");
-				}
-			}
+			RunAndLog(CmdEnv, IPPExe, IPPArguments);
 
 			// verify the .ipa exists
 			if (!FileExists(ProjectIPA))
diff --git a/Engine/Source/Programs/AutomationTool/IOS/IPhonePackager/CompileTime.cs b/Engine/Source/Programs/AutomationTool/IOS/IPhonePackager/CompileTime.cs
--- a/Engine/Source/Programs/AutomationTool/IOS/IPhonePackager/CompileTime.cs
+++ b/Engine/Source/Programs/AutomationTool/IOS/IPhonePackager/CompileTime.cs
@@ -11,6 +11,7 @@ using System.Net;
 using System.Runtime.InteropServices;
 using System.Text;
 using System.Threading;
+using System.Windows.Forms;
 using Ionic.Zlib;
 using UnrealBuildTool;
 
diff --git a/Engine/Source/Programs/AutomationTool/IOS/IPhonePackager/ConfigureMobileGame.cs b/Engine/Source/Programs/AutomationTool/IOS/IPhonePackager/ConfigureMobileGame.cs
new file mode 100644
--- /dev/null
+++ b/Engine/Source/Programs/AutomationTool/IOS/IPhonePackager/ConfigureMobileGame.cs
@@ -0,0 +1,204 @@
+﻿/*
+ * Copyright 1998-2014 Epic Games, Inc. All Rights Reserved.
+ */
+
+using System;
+using System.Collections.Generic;
+using System.ComponentModel;
+using System.Data;
+using System.Drawing;
+using System.Linq;
+using System.Text;
+using System.Windows.Forms;
+using System.Diagnostics;
+using System.IO;
+using System.Text.RegularExpressions;
+
+namespace iPhonePackager
+{
+	public partial class ConfigureMobileGame : Form
+	{
+		public bool bSucceededConfiguring = false;
+
+		public ConfigureMobileGame()
+		{
+			InitializeComponent();
+		}
+
+		// CFBundleDisplayName
+		
+		// CFBundleName, fewer than 16 characters long
+
+		// CFBundleIdentifier
+		// http://developer.apple.com/library/ios/#documentation/general/Reference/InfoPlistKeyReference/Articles/CoreFoundationKeys.html#//apple_ref/doc/uid/20001431-102070
+		// Alphanumeric (A-Z,a-z,0-9), hyphen (-), and period (.) characters.
+		// The string should also be in reverse-DNS format.
+
+		void ReloadPList()
+		{
+			string Filename = Config.GetPlistOverrideFilename();
+			string SourcePList = ReadOrCreate(Filename);
+			Utilities.PListHelper Helper = new Utilities.PListHelper(SourcePList);
+			
+			string BundleID;
+			string BundleName;
+			string BundleDisplayName;
+			if (!Helper.GetString("CFBundleIdentifier", out BundleID))
+			{
+				BundleID = "com.YourCompany.GameNameNoSpaces";
+				Helper.SetString("CFBundleIdentifier", BundleID);
+			}
+
+			if (!Helper.GetString("CFBundleName", out BundleName))
+			{
+				BundleName = "MyUDKGame";
+				Helper.SetString("CFBundleName", BundleName);
+			}
+
+			if (!Helper.GetString("CFBundleDisplayName", out BundleDisplayName))
+			{
+				BundleDisplayName = "UDK Game";
+				Helper.SetString("CFBundleDisplayName", BundleDisplayName);
+			}
+
+			BundleIdentifierEdit.Text = BundleID;
+			BundleNameEdit.Text = BundleName;
+			BundleDisplayNameEdit.Text = BundleDisplayName;
+		}
+
+		private void ReloadButton_Click(object sender, EventArgs e)
+		{
+			ReloadPList();
+		}
+
+		private void SaveChanges()
+		{
+			// Development settings
+			{
+				// Open the existing plist override file
+				string DevFilename = Config.GetPlistOverrideFilename(false);
+				string SourcePList = ReadOrCreate(DevFilename);
+				Utilities.PListHelper Helper = new Utilities.PListHelper(SourcePList);
+
+				// Jam the edited values in
+				if (!Helper.HasKey("UIFileSharingEnabled"))
+				{
+					Helper.SetValueForKey("UIFileSharingEnabled", true);
+				}
+
+				Helper.SetString("CFBundleIdentifier", BundleIdentifierEdit.Text);
+				Helper.SetString("CFBundleName", BundleNameEdit.Text);
+				Helper.SetString("CFBundleDisplayName", BundleDisplayNameEdit.Text);
+
+				// Save the modified plist
+				SavePList(Helper, DevFilename);
+			}
+
+			// Distribution settings
+			{
+				// Open the existing plist override file
+				string DistFilename = Config.GetPlistOverrideFilename(true);
+				string SourcePList = ReadOrCreate(DistFilename);
+				Utilities.PListHelper Helper = new Utilities.PListHelper(SourcePList);
+
+				// Jam the edited values in
+				if (!Helper.HasKey("UIFileSharingEnabled"))
+				{
+					Helper.SetValueForKey("UIFileSharingEnabled", false);
+				}
+
+				// Save the modified plist
+				SavePList(Helper, DistFilename);
+			}
+		}
+
+		private void EditManuallyButton_Click(object sender, EventArgs e)
+		{
+			SaveChanges();
+
+			string Filename = Config.GetPlistOverrideFilename();
+			Process.Start("explorer.exe", String.Format("/select,\"{0}\"", Path.GetFullPath(Filename)));
+
+			DialogResult = DialogResult.OK;
+			Close();
+		}
+
+		string ReadOrCreate(string Filename)
+		{
+			string Data = "";
+			try
+			{
+				Data = File.ReadAllText(Filename, Encoding.UTF8);
+			}
+			catch
+			{
+				Data = "<?xml version=\"1.0\" encoding=\"UTF-8\"?>\n" +
+					"<!DOCTYPE plist PUBLIC \"-//Apple//DTD PLIST 1.0//EN\" \"http://www.apple.com/DTDs/PropertyList-1.0.dtd\">\n" +
+					"<plist version=\"1.0\">\n<dict>\n</dict>\n</plist>\n";
+			}
+
+			return Data;
+		}
+
+		private void ShowError(string Message)
+		{
+			MessageBox.Show(Message, Config.AppDisplayName, MessageBoxButtons.OK, MessageBoxIcon.Error);
+		}
+
+		private void GenerateResignFileButton_Click(object sender, EventArgs e)
+		{
+			// Validate the edited values
+			Regex AcceptableID = new Regex("[^a-zA-Z0-9.\\-]");
+			if (AcceptableID.IsMatch(BundleIdentifierEdit.Text))
+			{
+				MessageBox.Show("Error: Bundle Identifier must contain only characters in the range [a-z] [A-Z] [0-9] . and -");
+				return;
+			}
+
+			if (BundleIdentifierEdit.Text.Length < 1)
+			{
+				ShowError("Error: Bundle Identifier must be at least 1 character long");
+				return;
+			}
+
+			if ((BundleNameEdit.Text.Length < 1) || (BundleNameEdit.Text.Length >= 16))
+			{
+				ShowError("Error: Bundle Name must be between 1 and 15 characters long");
+				return;
+			}
+
+			if (BundleDisplayNameEdit.Text.Length < 1)
+			{
+				BundleDisplayNameEdit.Text = BundleNameEdit.Text;
+			}
+
+			SaveChanges();
+
+			DialogResult = DialogResult.OK;
+			Close();
+		}
+
+		private void SavePList(Utilities.PListHelper Helper, string Filename)
+		{
+			FileInfo DestInfo = new FileInfo(Filename);
+			if (DestInfo.Exists && DestInfo.IsReadOnly)
+			{
+				DestInfo.IsReadOnly = false;
+			}
+			Helper.SaveToFile(Filename);
+		}
+
+		private void ConfigureMobileGame_Load(object sender, EventArgs e)
+		{
+			ReloadPList();
+			Text = Config.AppDisplayName + ": Customize Info.plist";
+		}
+
+		private void ShowHyperlink(object sender, LinkLabelLinkClickedEventArgs e)
+		{
+			string Target = (sender as LinkLabel).Tag as string;
+			ProcessStartInfo PSI = new ProcessStartInfo(Target);
+			Process.Start(PSI);
+		}
+	}
+}
diff --git a/Engine/Source/Programs/AutomationTool/IOS/IPhonePackager/CookTime.cs b/Engine/Source/Programs/AutomationTool/IOS/IPhonePackager/CookTime.cs
--- a/Engine/Source/Programs/AutomationTool/IOS/IPhonePackager/CookTime.cs
+++ b/Engine/Source/Programs/AutomationTool/IOS/IPhonePackager/CookTime.cs
@@ -10,6 +10,7 @@ using System.Net;
 using System.Runtime.InteropServices;
 using System.Text;
 using System.Threading;
+using System.Windows.Forms;
 using Ionic.Zip;
 using Ionic.Zlib;
 
diff --git a/Engine/Source/Programs/AutomationTool/IOS/IPhonePackager/DeployTime.cs b/Engine/Source/Programs/AutomationTool/IOS/IPhonePackager/DeployTime.cs
--- a/Engine/Source/Programs/AutomationTool/IOS/IPhonePackager/DeployTime.cs
+++ b/Engine/Source/Programs/AutomationTool/IOS/IPhonePackager/DeployTime.cs
@@ -10,6 +10,7 @@ using System.Net;
 using System.Runtime.InteropServices;
 using System.Text;
 using System.Threading;
+using System.Windows.Forms;
 using Microsoft.Win32;
 using System.Linq;
 using System.Runtime.Remoting.Channels.Ipc;
diff --git a/Engine/Source/Programs/AutomationTool/IOS/IPhonePackager/GenerateSigningRequestDialog.cs b/Engine/Source/Programs/AutomationTool/IOS/IPhonePackager/GenerateSigningRequestDialog.cs
new file mode 100644
--- /dev/null
+++ b/Engine/Source/Programs/AutomationTool/IOS/IPhonePackager/GenerateSigningRequestDialog.cs
@@ -0,0 +1,228 @@
+﻿// Copyright 1998-2014 Epic Games, Inc. All Rights Reserved.
+
+using System;
+using System.Collections.Generic;
+using System.ComponentModel;
+using System.Data;
+using System.Drawing;
+using System.Linq;
+using System.Text;
+using System.Windows.Forms;
+using System.IO;
+using iPhonePackager;
+
+using Org.BouncyCastle.Pkcs;
+using Org.BouncyCastle.Crypto.Parameters;
+using Org.BouncyCastle.Crypto;
+using System.Collections;
+using Org.BouncyCastle.Asn1.X509;
+using Org.BouncyCastle.Security;
+using Org.BouncyCastle.Math;
+using Org.BouncyCastle.Asn1;
+using Org.BouncyCastle.OpenSsl;
+
+namespace iPhonePackager
+{
+	public partial class GenerateSigningRequestDialog : Form
+	{
+		protected Dictionary<bool, Bitmap> CheckStateImages = new Dictionary<bool, Bitmap>();
+
+		public GenerateSigningRequestDialog()
+		{
+			InitializeComponent();
+
+			CheckStateImages.Add(false, iPhonePackager.Properties.Resources.GreyCheck);
+			CheckStateImages.Add(true, iPhonePackager.Properties.Resources.GreenCheck);
+
+			GenerateRequestButton.Enabled = false;
+		}
+
+		private void GenerateSigningRequestDialog_Load(object sender, EventArgs e)
+		{
+			SaveDialog.InitialDirectory = Environment.GetFolderPath(Environment.SpecialFolder.Desktop);
+		}
+
+
+		private void GenerateSigningRequestViaOpenSSL(string TargetCertRequestFileName, AsymmetricCipherKeyPair KeyPair)
+		{
+			// We expect openssl.exe to exist in the same directory as iPhonePackager
+			string OpenSSLPath = Path.GetDirectoryName( Application.ExecutablePath ) + @"\openssl.exe";
+			if (!File.Exists( OpenSSLPath ))
+			{
+				MessageBox.Show("A version of OpenSSL is required to generate certificate requests.  Please place OpenSSL.exe in Binaries\\IPhone", Config.AppDisplayName, MessageBoxButtons.OK, MessageBoxIcon.Error);
+				return;
+			}
+
+			// Create a temporary file to write the key pair out to (in a format that OpenSSL understands)
+			string KeyFileName = Path.GetTempFileName();
+			TextWriter KeyWriter = new StreamWriter(KeyFileName);
+
+			PemWriter KeyWriterPEM = new PemWriter(KeyWriter);
+			KeyWriterPEM.WriteObject(KeyPair);
+			KeyWriter.Close();
+
+			// Create a temporary file containing the configuration settings to drive OpenSSL
+			string ConfigFileName = Path.GetTempFileName();
+			TextWriter ConfigFile = new StreamWriter(ConfigFileName);
+				
+			ConfigFile.WriteLine("[ req ]");
+			ConfigFile.WriteLine("distinguished_name     = req_distinguished_name");
+			ConfigFile.WriteLine("prompt                 = no");
+
+			ConfigFile.WriteLine("[ req_distinguished_name ]");
+			ConfigFile.WriteLine("emailAddress           = {0}", EMailEditBox.Text);
+			ConfigFile.WriteLine("commonName             = {0}", CommonNameEditBox.Text);
+			ConfigFile.WriteLine("countryName            = {0}", System.Globalization.CultureInfo.CurrentCulture.TwoLetterISOLanguageName);
+
+			ConfigFile.Close();
+
+			// Invoke OpenSSL to generate the certificate request
+			Program.Log("Running OpenSSL to generate certificate request...");
+
+			string ResultsText;
+			string Executable = OpenSSLPath;
+			string Arguments = String.Format("req -new -nodes -out \"{0}\" -key \"{1}\" -config \"{2}\"",
+				TargetCertRequestFileName, KeyFileName, ConfigFileName);
+			Utilities.RunExecutableAndWait(Executable, Arguments, out ResultsText);
+
+			Program.Log(ResultsText);
+
+			if (!File.Exists(TargetCertRequestFileName))
+			{
+				Program.Error("... Failed to generate certificate request");
+			}
+			else
+			{
+				Program.Log("... Successfully generated certificate request '{0}'", TargetCertRequestFileName);
+			}
+
+			// Clean up the temporary files we created
+			File.Delete(KeyFileName);
+			File.Delete(ConfigFileName);
+		}
+
+		//@TODO: Generates a seemingly valid CSR but the Apple website doesn't accept this one (there is a context dependant 0 missing and a null missing compared to Keychain or OpenSSL generated ones, so the files aren't structurally identical)
+		private void GenerateSigningRequestViaBouncyCastle(string TargetCertRequestFileName, AsymmetricCipherKeyPair KeyPair)
+		{
+			// Construct the request subject
+			Hashtable Attributes = new Hashtable();
+			Attributes.Add(X509Name.EmailAddress, EMailEditBox.Text);
+			Attributes.Add(X509Name.CN, CommonNameEditBox.Text);
+			Attributes.Add(X509Name.C, System.Globalization.RegionInfo.CurrentRegion.TwoLetterISORegionName);
+
+			X509Name Subject = new X509Name(new ArrayList(Attributes.Keys), Attributes);
+
+			// Generate a certificate signing request
+			Pkcs10CertificationRequest CSR = new Pkcs10CertificationRequest(
+				"SHA1withRSA",
+				Subject,
+				KeyPair.Public,
+				null,
+				KeyPair.Private);
+
+			// Save the CSR to disk
+			File.WriteAllBytes(TargetCertRequestFileName, CSR.GetEncoded());
+		}
+
+		private void GenerateRequestButton_Click(object sender, EventArgs e)
+		{
+			if (EMailEditBox.Text == "")
+			{
+				MessageBox.Show("An Email address is required to generate a signing request", Config.AppDisplayName, MessageBoxButtons.OK, MessageBoxIcon.Error);
+				return;
+			}
+
+			if (CommonNameEditBox.Text == "")
+			{
+				MessageBox.Show("A common name is required to generate a signing request", Config.AppDisplayName, MessageBoxButtons.OK, MessageBoxIcon.Error);
+				return;
+			}
+
+			AsymmetricCipherKeyPair KeyPair = null;
+			try
+			{
+				KeyPair = CryptoAdapter.LoadKeyPairFromDiskBouncy(KeyPairFilenameBox.Text);
+				if (KeyPair == null)
+				{
+					throw new InvalidDataException();
+				}
+			}
+			catch (Exception)
+			{
+				MessageBox.Show("A public/private key pair is required to generate a signing request (failed to find or open specified key pair file)", Config.AppDisplayName, MessageBoxButtons.OK, MessageBoxIcon.Error);
+				return;
+			}
+
+			SaveDialog.FileName = "CertificateSigningRequest.csr";
+			SaveDialog.DefaultExt = "csr";
+			SaveDialog.Filter = ToolsHub.CertificateRequestFilter;
+			SaveDialog.Title = "Generate Certificate Signing Request";
+			if (SaveDialog.ShowDialog() == DialogResult.OK)
+			{
+				string TargetCertRequestFileName = SaveDialog.FileName;
+
+				GenerateSigningRequestViaOpenSSL(TargetCertRequestFileName, KeyPair);
+				//GenerateSigningRequestViaBouncyCastle(TargetCertRequestFileName, KeyPair);
+
+				// Close this dialog
+				Close();
+			}
+		}
+
+		private void GenerateSigningRequestDialog_FormClosing(object sender, FormClosingEventArgs e)
+		{
+
+		}
+
+		private void GenerateSigningRequestDialog_FormClosed(object sender, FormClosedEventArgs e)
+		{
+
+		}
+
+		private void GenerateKeyPairButton_Click(object sender, EventArgs e)
+		{
+			SaveDialog.FileName = "DeveloperKeyPair.key";
+			SaveDialog.DefaultExt = "key";
+			SaveDialog.Filter = ToolsHub.JustXmlKeysFilter;
+			if (SaveDialog.ShowDialog() == DialogResult.OK)
+			{
+				string TargetKeyFilename = SaveDialog.FileName;
+
+				// Generate a new 2048 bit RSA key and save it to disk
+				System.Security.Cryptography.RSACryptoServiceProvider KeyPair = new System.Security.Cryptography.RSACryptoServiceProvider(2048);
+				File.WriteAllText(TargetKeyFilename, KeyPair.ToXmlString(true));
+
+				KeyPairFilenameBox.Text = TargetKeyFilename;
+			}
+		}
+
+		private void ChooseKeyPairButton_Click(object sender, EventArgs e)
+		{
+			OpenDialog.Filter = ToolsHub.KeysFilter;
+			if (OpenDialog.ShowDialog() == DialogResult.OK)
+			{
+				KeyPairFilenameBox.Text = OpenDialog.FileName;
+
+				AsymmetricCipherKeyPair KeyPair = CryptoAdapter.LoadKeyPairFromDiskBouncy(KeyPairFilenameBox.Text);
+				if (KeyPair == null)
+				{
+					MessageBox.Show("Unable to load key pair from disk, make sure it is in a supported format (xml or pkcs 12 key store)", Config.AppDisplayName, MessageBoxButtons.OK, MessageBoxIcon.Error);
+					KeyPairFilenameBox.Text = "";
+				}
+			}
+		}
+
+		private void ValidateCertSteps(object sender, EventArgs e)
+		{
+			bool bHasEmail = EMailEditBox.Text.Length > 0;
+			bool bHasName = CommonNameEditBox.Text.Length > 0;
+			bool bHasKeyPair = KeyPairFilenameBox.Text.Length > 0;
+
+			EmailCheck.Image = CheckStateImages[bHasEmail];
+			CommonNameCheck.Image = CheckStateImages[bHasName];
+			KeyPairCheck.Image = CheckStateImages[bHasKeyPair];
+
+			GenerateRequestButton.Enabled = bHasEmail && bHasName && bHasKeyPair;
+		}
+	}
+}
diff --git a/Engine/Source/Programs/AutomationTool/IOS/IPhonePackager/GraphicalResignTool.cs b/Engine/Source/Programs/AutomationTool/IOS/IPhonePackager/GraphicalResignTool.cs
new file mode 100644
--- /dev/null
+++ b/Engine/Source/Programs/AutomationTool/IOS/IPhonePackager/GraphicalResignTool.cs
@@ -0,0 +1,817 @@
+﻿/**
+ * Copyright 1998-2014 Epic Games, Inc. All Rights Reserved.
+ */
+
+using System;
+using System.Collections.Generic;
+using System.ComponentModel;
+using System.Data;
+using System.Drawing;
+using System.Linq;
+using System.Text;
+using System.Windows.Forms;
+using Ionic.Zip;
+using System.IO;
+using System.Security.Cryptography.X509Certificates;
+
+namespace iPhonePackager
+{
+	public partial class GraphicalResignTool : Form
+	{
+		public static GraphicalResignTool ActiveInstance = null;
+
+		private string ExportImportPListFilename = "";
+		private byte[] ImportedPListData = null;
+
+		public static GraphicalResignTool GetActiveInstance()
+		{
+			if (ActiveInstance == null)
+			{
+				ActiveInstance = new GraphicalResignTool();
+			}
+			return ActiveInstance;
+		}
+
+		public GraphicalResignTool()
+		{
+			InitializeComponent();
+
+			SwitchingInfoEditMode(null, null);
+			SwitchingSigningMode(null, null);
+			SwitchingMobileProvisionMode(null, null);
+
+			InitializeGameSelectionBox();
+
+			ActiveInstance = this;
+		}
+
+		/// <summary>
+		/// Reads through directories that are peer to Binaries for any that match the *Game pattern
+		/// and adds the games to a drop-down list that can be picked.
+		/// </summary>
+		void InitializeGameSelectionBox()
+		{
+			string InstallPath = Path.GetDirectoryName(Application.ExecutablePath) + Path.DirectorySeparatorChar + ".." + Path.DirectorySeparatorChar + ".." + Path.DirectorySeparatorChar;
+			string[] GameDirectories = Directory.GetDirectories(InstallPath, "*Game");
+
+			// Fetch the currently active game name
+			string ActiveGameName = (Program.GameName != null) ? Program.GameName : "UDKGame";
+
+			// Add the game names, highlighting the active one
+			GameSelectionBox.Items.Clear();
+			foreach (string GameDir in GameDirectories)
+			{
+				int CutIndex = GameDir.LastIndexOf(Path.DirectorySeparatorChar);
+				string GameName = GameDir.Substring(CutIndex + 1);
+
+				int GameIndexInComboBox = GameSelectionBox.Items.Add(GameName);
+				if (GameName.ToUpperInvariant() == ActiveGameName.ToUpperInvariant())
+				{
+					GameSelectionBox.SelectedIndex = GameIndexInComboBox;
+				}
+			}
+		}
+
+		private void GraphicalResignTool_Load(object sender, EventArgs e)
+		{
+			Text = Config.AppDisplayName + ": Signing tool";
+		}
+
+
+		private void InputIPABrowseButton_Click(object sender, EventArgs e)
+		{
+			string PickedFilename;
+			if (ToolsHub.ShowOpenFileDialog(ToolsHub.IpaFilter, "Picking an application", "ipa", IPAFilenameEdit.Text, ref ToolsHub.ChoosingIpaDirectory, out PickedFilename))
+			{
+				IPAFilenameEdit.Text = PickedFilename;
+				ImportedPListData = null;
+			}
+		}
+
+		private void MobileProvisionBrowseButton_Click(object sender, EventArgs e)
+		{
+			string PickedFilename;
+			if (ToolsHub.ShowOpenFileDialog(ToolsHub.MobileProvisionFilter, "Picking a mobile provision to embed", "mobileprovision", MobileProvisionEdit.Text, ref ToolsHub.ChoosingIpaDirectory, out PickedFilename))
+			{
+				MobileProvisionEdit.Text = PickedFilename;
+			}
+		}
+
+		private void CertificateBrowseButton_Click(object sender, EventArgs e)
+		{
+			string PickedFilename;
+			if (ToolsHub.ShowOpenFileDialog(ToolsHub.CertificatesFilter, "Picking a signing certificate to codesign with", "p12", CertificateEdit.Text, ref ToolsHub.ChoosingIpaDirectory, out PickedFilename))
+			{
+				CertificateEdit.Text = PickedFilename;
+			}
+		}
+
+		void ShowError(string Context, Exception ex)
+		{
+			MessageBox.Show(String.Format("Error '{0}' encountered while {1}", ex.Message, Context), Config.AppDisplayName, MessageBoxButtons.OK, MessageBoxIcon.Error);
+		}
+
+		void ShowErrorMessage(string Msg, params object [] Args)
+		{
+			MessageBox.Show(String.Format(Msg, Args), Config.AppDisplayName, MessageBoxButtons.OK, MessageBoxIcon.Error);
+		}
+
+		private void ExportInfoButton_Click(object sender, EventArgs e)
+		{
+			try
+			{
+				FileOperations.FileSystemAdapter IPA = new FileOperations.ReadOnlyZipFileSystem(IPAFilenameEdit.Text);
+
+				if (IPA != null)
+				{
+					saveFileDialog1.DefaultExt = "plist";
+					saveFileDialog1.Title = "Exporting Info.plist";
+					saveFileDialog1.Filter = ToolsHub.PListFilter;
+					saveFileDialog1.FileName = ExportImportPListFilename;
+
+					string CWD = Directory.GetCurrentDirectory();
+					bool bDialogSucceeded = (saveFileDialog1.ShowDialog() == DialogResult.OK);
+					Directory.SetCurrentDirectory(CWD);
+
+					if (bDialogSucceeded)
+					{
+						ExportImportPListFilename = saveFileDialog1.FileName;
+						byte[] DataToSave = IPA.ReadAllBytes("Info.plist");
+						File.WriteAllBytes(ExportImportPListFilename, DataToSave);
+					}
+				}
+			}
+			catch (Exception ex)
+			{
+				ShowError("exporting Info.plist", ex);
+			}
+		}
+
+		private void ImportInfoButton_Click(object sender, EventArgs e)
+		{
+			try
+			{
+				string PickedFilename;
+				if (ToolsHub.ShowOpenFileDialog(ToolsHub.PListFilter, "Importing Info.plist", "plist", ExportImportPListFilename, ref ToolsHub.ChoosingIpaDirectory, out PickedFilename))
+				{
+					ExportImportPListFilename = PickedFilename;
+					ImportedPListData = File.ReadAllBytes(ExportImportPListFilename);
+				}
+			}
+			catch (Exception ex)
+			{
+				ShowError("importing Info.plist", ex);
+			}
+		}
+
+		private void SwitchingInfoEditMode(object sender, EventArgs e)
+		{
+			ImportInfoButton.Enabled = RBReplaceInfoPList.Checked;
+			ExportInfoButton.Enabled = RBReplaceInfoPList.Checked;
+			DisplayNameEdit.Enabled = RBModifyInfoPList.Checked;
+			BundleIDEdit.Enabled = RBModifyInfoPList.Checked;
+
+			if (!RBReplaceInfoPList.Checked)
+			{
+				ImportedPListData = null;
+			}
+		}
+
+		private void SwitchingSigningMode(object sender, EventArgs e)
+		{
+			CertificateEdit.Enabled = RBUseExplicitCert.Checked;
+			CertificateBrowseButton.Enabled = RBUseExplicitCert.Checked;
+		}
+
+		private void SwitchingMobileProvisionMode(object sender, EventArgs e)
+		{
+			MobileProvisionEdit.Enabled = RBSpecifyMobileProvision.Checked;
+			MobileProvisionBrowseButton.Enabled = RBSpecifyMobileProvision.Checked;
+		}
+
+		private void ResignButton_Click(object sender, EventArgs e)
+		{
+			saveFileDialog1.DefaultExt = "ipa";
+			saveFileDialog1.FileName = "";
+			saveFileDialog1.Filter = ToolsHub.IpaFilter;
+			saveFileDialog1.Title = "Choose a filename for the re-signed IPA";
+
+			string CWD = Directory.GetCurrentDirectory();
+			bool bDialogSucceeded = (saveFileDialog1.ShowDialog() == DialogResult.OK);
+			Directory.SetCurrentDirectory(CWD);
+
+			if (bDialogSucceeded)
+			{
+				bool bSavedVerbose = Config.bVerbose;
+				Config.bVerbose = true;
+				Program.ProgressDialog.OnBeginBackgroundWork = delegate
+				{
+					string SrcFilename = IPAFilenameEdit.Text;
+					string DestFilename = saveFileDialog1.FileName;
+
+					// Delete the target location and copy the source there
+					FileOperations.DeleteFile(DestFilename);
+					FileOperations.CopyRequiredFile(SrcFilename, DestFilename);
+
+					// Open the file
+					FileOperations.ZipFileSystem FileSystem = new FileOperations.ZipFileSystem(DestFilename);
+					FileSystem.SetCompression(CBCompressModifiedFiles.Checked ? Ionic.Zlib.CompressionLevel.BestCompression : Ionic.Zlib.CompressionLevel.None);
+
+					// Resign and save the file
+					ResignIPA(FileSystem);
+					FileSystem.Close();
+				};
+				Config.bVerbose = bSavedVerbose;
+
+				Program.ProgressDialog.ShowDialog();
+			}
+		}
+
+		/// <summary>
+		/// This class is used to override specific parts of the default resigning behavior
+		/// </summary>
+		class CustomCodeSigner : CodeSignatureBuilder
+		{
+			/// <summary>
+			/// If non-null, this will be used instead of the existing embedded mobile provision
+			/// </summary>
+			public byte[] CustomMobileProvision = null;
+
+			/// <summary>
+			/// If non-null, this certificate will be used instead of doing a cert scan
+			/// </summary>
+			public X509Certificate2 CustomSigningCert = null;
+
+			/// <summary>
+			/// If non-null, this data will completely replace the original Info.plist
+			/// </summary>
+			public byte[] CustomInfoPList = null;
+
+			protected override byte[] GetMobileProvision()
+			{
+				return (CustomMobileProvision != null) ? CustomMobileProvision : FileSystem.ReadAllBytes("embedded.mobileprovision");
+			}
+
+			protected override X509Certificate2 LoadSigningCertificate()
+			{
+				return (CustomSigningCert != null) ? CustomSigningCert : base.LoadSigningCertificate();
+			}
+
+			protected override Utilities.PListHelper LoadInfoPList()
+			{
+				if (CustomInfoPList != null)
+				{
+					return new Utilities.PListHelper(Encoding.UTF8.GetString(CustomInfoPList));
+				}
+				else
+				{
+					return base.LoadInfoPList();
+				}
+			}
+		}
+
+		/** 
+		 * Using the stub IPA previously compiled on the Mac, create a new IPA with assets
+		 */
+		void ResignIPA(FileOperations.FileSystemAdapter FileSystem)
+		{
+			try
+			{
+				DateTime StartTime = DateTime.Now;
+
+				// Configure the custom code signer
+				CustomCodeSigner SigningContext = new CustomCodeSigner();
+				SigningContext.FileSystem = FileSystem;
+
+				// Custom mobile provision?
+				if (RBSpecifyMobileProvision.Checked)
+				{
+					SigningContext.CustomMobileProvision = File.ReadAllBytes(MobileProvisionEdit.Text);
+				}
+
+				// Custom cert?
+				if (RBUseExplicitCert.Checked)
+				{
+					string CertificatePassword = "";
+					try
+					{
+						SigningContext.CustomSigningCert = new X509Certificate2(CertificateEdit.Text, CertificatePassword, X509KeyStorageFlags.PersistKeySet | X509KeyStorageFlags.Exportable | X509KeyStorageFlags.MachineKeySet);
+						//Cert = new X509Certificate2(CertificateFilename, CertificatePassword, X509KeyStorageFlags.PersistKeySet | X509KeyStorageFlags.Exportable | X509KeyStorageFlags.MachineKeySet);
+					}
+					catch (System.Security.Cryptography.CryptographicException ex)
+					{
+						// Try once with a password
+						if (PasswordDialog.RequestPassword(out CertificatePassword))
+						{
+							SigningContext.CustomSigningCert = new X509Certificate2(CertificateEdit.Text, CertificatePassword, X509KeyStorageFlags.PersistKeySet | X509KeyStorageFlags.Exportable | X509KeyStorageFlags.MachineKeySet);
+							//Cert = new X509Certificate2(CertificateFilename, CertificatePassword, X509KeyStorageFlags.PersistKeySet | X509KeyStorageFlags.Exportable | X509KeyStorageFlags.MachineKeySet);
+						}
+						else
+						{
+							// User cancelled dialog, rethrow
+							throw ex;
+						}
+					}
+				}
+
+				// Totally replace Info.plist? (just editing it is handled later)
+				if (RBReplaceInfoPList.Checked)
+				{
+					SigningContext.CustomInfoPList = ImportedPListData;
+				}
+
+				// Start the resign process
+				SigningContext.PrepareForSigning();
+
+				// Partially modify Info.plist?
+				if (RBModifyInfoPList.Checked)
+				{
+					SigningContext.Info.SetString("CFBundleDisplayName", DisplayNameEdit.Text);
+					SigningContext.Info.SetString("CFBundleIdentifier", BundleIDEdit.Text);
+				}
+
+				// Re-sign the executable
+				SigningContext.PerformSigning();
+
+				// Save the IPA
+				Program.Log("Saving IPA ...");
+				FileSystem.Close();
+
+				TimeSpan ElapsedTime = DateTime.Now - StartTime;
+				Program.Log(String.Format("Finished re-signing IPA in took {0:0.00} s", ElapsedTime.TotalSeconds));
+
+				MessageBox.Show("Re-signing succeeded!", Config.AppDisplayName, MessageBoxButtons.OK, MessageBoxIcon.Information);
+			}
+			catch (Exception ex)
+			{
+				ShowError("re-signing IPA", ex);
+			}
+		}
+
+		private void CloseThisWindowButton_Click(object sender, EventArgs e)
+		{
+			Close();
+		}
+
+		private void InstallIPAButton_Click(object sender, EventArgs e)
+		{
+			string PickedFilename;
+			if (ToolsHub.ShowOpenFileDialog(ToolsHub.IpaFilter, "Choose an IPA to install", "ipa", "", ref ToolsHub.ChoosingIpaDirectory, out PickedFilename))
+			{
+				Program.ProgressDialog.OnBeginBackgroundWork = delegate
+				{
+					DeploymentHelper.InstallIPAOnConnectedDevices(PickedFilename);
+				};
+				Program.ProgressDialog.ShowDialog();
+			}
+		}
+
+		private void UninstallIPAButton_Click(object sender, EventArgs e)
+		{
+			
+		}
+
+		void ScanConnectedDevices()
+		{
+			ConnectedDeviceInfo [] DeviceMap = null;
+
+			Program.ProgressDialog.OnBeginBackgroundWork = delegate
+			{
+				DeviceMap = DeploymentHelper.Get().EnumerateConnectedDevices();
+			};
+
+			if ((Program.ProgressDialog.ShowDialog() == DialogResult.OK) && (DeviceMap != null))
+			{
+				ConnectedDevicesList.BeginUpdate();
+				ConnectedDevicesList.Items.Clear();
+
+				foreach (var DeviceInfo in DeviceMap)
+				{
+					string DisplayName = DeviceInfo.DeviceName;
+					string UDID = DeviceInfo.UDID;
+					string DeviceType = DeviceInfo.DeviceType;
+
+					ListViewItem NewItem = ConnectedDevicesList.Items.Add(DisplayName);
+					NewItem.SubItems.Add(UDID);
+					NewItem.SubItems.Add(DeviceType);
+				}
+
+				ConnectedDevicesList.EndUpdate();
+			}
+			else
+			{
+				Program.Error("Enumerating devices failed");
+			}
+		}
+
+		private void FetchUDIDButton_Click(object sender, EventArgs e)
+		{
+			ScanConnectedDevices();
+		}
+
+		private void UDIDContextMenu_Opening(object sender, CancelEventArgs e)
+		{
+
+		}
+
+		private void textBox1_TextChanged(object sender, EventArgs e)
+		{
+
+		}
+
+		private void ImportCertificateButton3_Click(object sender, EventArgs e)
+		{
+			ToolsHub.TryInstallingCertificate_PromptForKey();
+			RefreshProvisionsList();
+		}
+
+		private void ImportProvisionButton3_Click(object sender, EventArgs e)
+		{
+			ToolsHub.TryInstallingMobileProvision();
+			RefreshProvisionsList();
+		}
+
+		void RefreshProvisionsList()
+		{
+			// Scan for all mobile provisions installed for the current game
+			Dictionary<string, MobileProvision> Provisions = new Dictionary<string, MobileProvision>();
+
+			Program.ProgressDialog.OnBeginBackgroundWork = delegate
+			{
+				Program.Log("Scanning for mobile provisions in '{0}'...", Config.BuildDirectory);
+
+				try
+				{
+					// Find all the provisions to check
+					string[] ProvisionFilenames = Directory.GetFiles(Config.BuildDirectory, "*.mobileprovision");
+
+					// Open each one in turn
+					int ProgressIndex = 0;
+					foreach (string Filename in ProvisionFilenames)
+					{
+						++ProgressIndex;
+						MobileProvision Provision = null;
+						try
+						{
+							Program.ProgressIndex = (ProgressIndex * 100) / ProvisionFilenames.Length;
+							Program.Log("... Loading provision '{0}'", Filename);
+
+							Provision = MobileProvisionParser.ParseFile(Filename);
+							Provision.Tag = CodeSignatureBuilder.FindCertificate(Provision);
+						}
+						catch
+						{
+						}
+						Provisions.Add(Filename, Provision);
+					}
+				}
+				catch
+				{
+				}
+			};
+
+			if (Program.ProgressDialog.ShowDialog() == DialogResult.OK)
+			{
+				KnownProvisionsTree.BeginUpdate();
+				KnownProvisionsTree.Nodes.Clear();
+
+				// Visualize the found provisions
+				foreach (var KVP in Provisions)
+				{
+					string Filename = KVP.Key;
+					MobileProvision Provision = KVP.Value;
+
+					TreeNode NewItem = KnownProvisionsTree.Nodes.Add(Path.GetFileNameWithoutExtension(Filename));
+					NewItem.Tag = KVP;
+
+					if (Provision != null)
+					{
+						X509Certificate2 Cert = Provision.Tag as X509Certificate2;
+
+						NewItem.Nodes.Add(String.Format("Name: '{0}'", Provision.ProvisionName));
+						NewItem.Nodes.Add(String.Format("Type: {0}", ToolsHub.IsProfileForDistribution(Provision) ? "Distribution" : "Development"));
+
+						string Prefix = Path.GetFileNameWithoutExtension(Filename).Replace(Program.GameName, "");
+						if (Prefix != "")
+						{
+							NewItem.Nodes.Add(String.Format("Signing prefix: '{0}'", Prefix));
+						}
+
+						NewItem.Nodes.Add(String.Format("Matching Cert: '{0}'", (Cert != null) ? Cert.FriendlyName : "None (missing)"));
+
+						// Add the device list
+						TreeNode DeviceList = NewItem.Nodes.Add(String.Format("Allowed devices: {0}", Provision.ProvisionedDeviceIDs.Count));
+						foreach (string DeviceID in Provision.ProvisionedDeviceIDs)
+						{
+							DeviceList.Nodes.Add(DeviceID);
+						}
+						DeviceList.Collapse();
+					}
+					else
+					{
+						NewItem.Nodes.Add("Error: Failed to parse");
+					}
+
+					NewItem.Expand();
+				}
+
+				KnownProvisionsTree.EndUpdate();
+			}
+		}
+
+		private void RefreshProvisionListButton_Click(object sender, EventArgs e)
+		{
+			RefreshProvisionsList();
+		}
+
+		private void GraphicalResignTool_FormClosed(object sender, FormClosedEventArgs e)
+		{
+			if (ActiveInstance == this)
+			{
+				ActiveInstance = null;
+			}
+		}
+
+		private void TabBook_SelectedIndexChanged(object sender, EventArgs e)
+		{
+			if (TabBook.SelectedTab == ProvisionCertPage)
+			{
+				RefreshProvisionsList();
+			}
+			else if (TabBook.SelectedTab == DeploymentPage)
+			{
+				ScanConnectedDevices();
+			}
+		}
+
+		private void GameSelectionBox_SelectedIndexChanged(object sender, EventArgs e)
+		{
+			try
+			{
+				// THIS IS NO LONGER VALID, WE NEED A PATH TO THE GAME, AND RE-RUN CONFIG
+				System.Windows.Forms.MessageBox.Show("This function is no longer supported - it will need work!");
+//				 Program.GameName = GameSelectionBox.Items[GameSelectionBox.SelectedIndex] as string;
+//				 Config.Initialize();
+//				 RefreshProvisionsList();
+			}
+			catch
+			{
+			}
+		}
+
+		private void copyToolStripMenuItem_Click(object sender, EventArgs e)
+		{
+			TreeNode Node = KnownProvisionsTree.SelectedNode;
+			if (Node != null)
+			{
+				Clipboard.SetText(Node.Text);
+			}
+		}
+
+		private void copyFullPathToProvisionToolStripMenuItem_Click(object sender, EventArgs e)
+		{
+			//
+			TreeNode Node = KnownProvisionsTree.SelectedNode;
+			if (Node != null)
+			{
+				// Walk up to the root
+				while (Node.Parent != null)
+				{
+					Node = Node.Parent;
+				}
+
+				// The tag of the root is a pair of filename, mobile provision
+				KeyValuePair<string, MobileProvision> KVP = (KeyValuePair<string, MobileProvision>)Node.Tag;
+				Clipboard.SetText(KVP.Key);
+			}
+		}
+
+		private void GraphicalResignTool_Shown(object sender, EventArgs e)
+		{
+			TabBook_SelectedIndexChanged(sender, e);
+		}
+
+		private void copyUDIDToClipboardToolStripMenuItem_Click(object sender, EventArgs e)
+		{
+			string ClipboardText = "";
+			foreach (ListViewItem i in ConnectedDevicesList.SelectedItems)
+			{
+				string Line = i.Text + " " + i.SubItems[1].Text;
+				if (ClipboardText == "")
+				{
+					ClipboardText = Line;
+				}
+				else
+				{
+					ClipboardText = ClipboardText + "\n" + Line;
+				}
+			}
+
+			Clipboard.SetText(ClipboardText);
+		}
+
+		static void CheckProvisionForUDID(MobileProvision Provision, string UDID, string Context)
+		{
+			if (Provision == null)
+			{
+				MessageBox.Show(String.Format("Failed to load or parse {0}", Context), Config.AppDisplayName, MessageBoxButtons.OK, MessageBoxIcon.Error);
+			}
+			else
+			{
+				if (Provision.ContainsUDID(UDID))
+				{
+					MessageBox.Show(String.Format("Found the UDID {0} in {1}", UDID, Context), Config.AppDisplayName, MessageBoxButtons.OK, MessageBoxIcon.Information);
+				}
+				else
+				{
+					MessageBox.Show(String.Format("Warning: Failed to find UDID {0} in {1}!", UDID, Context), Config.AppDisplayName, MessageBoxButtons.OK, MessageBoxIcon.Warning);
+				}
+			}
+		}
+
+		void PickProvisionAndCheckForUDID(string UDID)
+		{
+			if (UDID != null)
+			{
+				string PickedFilename;
+				if (ToolsHub.ShowOpenFileDialog(ToolsHub.MobileProvisionFilter, "Picking a mobile provision to check", "mobileprovision", "", ref ToolsHub.ChoosingIpaDirectory, out PickedFilename))
+				{
+					MobileProvision Provision = MobileProvisionParser.ParseFile(PickedFilename);
+
+					CheckProvisionForUDID(
+						Provision,
+						UDID,
+						String.Format("mobile provision {0}", Path.GetFileNameWithoutExtension(PickedFilename)));
+				}
+			}
+		}
+
+		private void CheckThisUDIDAgainstAMobileprovMenuItem_Click(object sender, EventArgs e)
+		{
+			string UDID = null;
+			foreach (ListViewItem i in ConnectedDevicesList.SelectedItems)
+			{
+				UDID = i.SubItems[1].Text;
+				break;
+			}
+
+			PickProvisionAndCheckForUDID(UDID);
+		}
+
+		void PickIpaAndCheckForUDID(string UDID)
+		{
+			if (UDID != null)
+			{
+				string PickedFilename;
+				if (ToolsHub.ShowOpenFileDialog(ToolsHub.IpaFilter, "Picking an application to check", "ipa", "", ref ToolsHub.ChoosingIpaDirectory, out PickedFilename))
+				{
+					MobileProvision Provision = MobileProvisionParser.ParseIPA(PickedFilename);
+
+					CheckProvisionForUDID(
+						Provision,
+						UDID,
+						String.Format("embedded mobile provision in {0}", Path.GetFileName(PickedFilename)));
+				}
+			}
+		}
+
+		private void CheckUDIDInIPAMenuItem_Click(object sender, EventArgs e)
+		{
+			string UDID = null;
+			foreach (ListViewItem i in ConnectedDevicesList.SelectedItems)
+			{
+				UDID = i.SubItems[1].Text;
+				break;
+			}
+
+			PickIpaAndCheckForUDID(UDID);
+		}
+
+		private void CheckIPAForUDIDButton_Click(object sender, EventArgs e)
+		{
+			string UDID = ScanningUDIDEdit.Text.Trim();
+
+			if (UDID.Length != 40)
+			{
+				ShowErrorMessage("The UDID must be exactly 40 characters long");
+			}
+			else
+			{
+				PickIpaAndCheckForUDID(UDID);
+			}
+		}
+
+		private void CheckProvisionForUDIDButton_Click(object sender, EventArgs e)
+		{
+			string UDID = ScanningUDIDEdit.Text.Trim();
+
+			if (UDID.Length != 40)
+			{
+				ShowErrorMessage("The UDID must be exactly 40 characters long");
+			}
+			else
+			{
+				PickProvisionAndCheckForUDID(UDID);
+			}
+		}
+
+		string PickIPAAndFetchBundleIdentifier()
+		{
+			string PickedFilename;
+			if (ToolsHub.ShowOpenFileDialog(ToolsHub.IpaFilter, "Choose an IPA to provide the bundle identifier", "ipa", "", ref ToolsHub.ChoosingIpaDirectory, out PickedFilename))
+			{
+				FileOperations.FileSystemAdapter IPA = new FileOperations.ReadOnlyZipFileSystem(PickedFilename);
+
+				if (IPA != null)
+				{
+					byte[] InfoData = IPA.ReadAllBytes("Info.plist");
+					Utilities.PListHelper Info = new Utilities.PListHelper(Encoding.UTF8.GetString(InfoData));
+
+					string BundleID;
+					if (Info.GetString("CFBundleIdentifier", out BundleID))
+					{
+						return BundleID;
+					}
+				}
+			}
+
+			return null;
+		}
+
+		private void BackupDocumentsButton_Click(object sender, EventArgs e)
+		{
+			string BundleID = PickIPAAndFetchBundleIdentifier();
+
+			if (BundleID != null)
+			{
+				Program.ProgressDialog.OnBeginBackgroundWork = delegate
+				{
+					DeploymentHelper.Get().BackupDocumentsDirectory(BundleID, Config.GetRootBackedUpDocumentsDirectory());
+				};
+				Program.ProgressDialog.ShowDialog();
+			}
+		}
+
+		private void UninstallIPAButton_Click_1(object sender, EventArgs e)
+		{
+			string BundleID = PickIPAAndFetchBundleIdentifier();
+
+			if (BundleID != null)
+			{
+				Program.ProgressDialog.OnBeginBackgroundWork = delegate
+				{
+					DeploymentHelper.Get().UninstallIPAOnDevice(BundleID);
+				};
+				Program.ProgressDialog.ShowDialog();
+			}
+		}
+
+		private void exportCertificateToolStripMenuItem_Click(object sender, EventArgs e)
+		{
+			TreeNode Node = KnownProvisionsTree.SelectedNode;
+			if (Node != null)
+			{
+				// Walk up to the root
+				while (Node.Parent != null)
+				{
+					Node = Node.Parent;
+				}
+
+				// The tag of the root is a pair of filename, mobile provision
+				KeyValuePair<string, MobileProvision> KVP = (KeyValuePair<string, MobileProvision>)Node.Tag;
+
+				MobileProvision Provision = KVP.Value;
+				X509Certificate2 Cert = Provision.Tag as X509Certificate2;
+
+				if (Cert != null)
+				{
+					// Ask for an export filename
+					saveFileDialog1.DefaultExt = "p12";
+					saveFileDialog1.Title = "Choose a destination for the exported PKCS#12 certificate + key chain";
+					saveFileDialog1.Filter = ToolsHub.JustP12Certificates;
+					saveFileDialog1.FileName = "";
+
+					string CWD = Directory.GetCurrentDirectory();
+					bool bDialogSucceeded = (saveFileDialog1.ShowDialog() == DialogResult.OK);
+					Directory.SetCurrentDirectory(CWD);
+
+					if (bDialogSucceeded)
+					{
+						string TargetFilename = saveFileDialog1.FileName;
+ 
+						// Ask for an export password
+						string ExportPassword;
+						if (PasswordDialog.RequestPassword(out ExportPassword))
+						{
+							File.WriteAllBytes(TargetFilename, Cert.Export(X509ContentType.Pkcs12, ExportPassword));
+						}
+					}
+				}
+				else
+				{
+					ShowErrorMessage("The selected mobile provision doesn't have an associated certificate to export!");
+				}
+			}
+		}
+	}
+}
diff --git a/Engine/Source/Programs/AutomationTool/IOS/IPhonePackager/PasswordDialog.cs b/Engine/Source/Programs/AutomationTool/IOS/IPhonePackager/PasswordDialog.cs
new file mode 100644
--- /dev/null
+++ b/Engine/Source/Programs/AutomationTool/IOS/IPhonePackager/PasswordDialog.cs
@@ -0,0 +1,39 @@
+﻿// Copyright 1998-2014 Epic Games, Inc. All Rights Reserved.
+
+using System;
+using System.Collections.Generic;
+using System.ComponentModel;
+using System.Data;
+using System.Drawing;
+using System.Linq;
+using System.Text;
+using System.Windows.Forms;
+
+namespace iPhonePackager
+{
+	public partial class PasswordDialog : Form
+	{
+		string PasswordResult = null;
+
+		public PasswordDialog()
+		{
+			InitializeComponent();
+		}
+
+		public static bool RequestPassword(out string Password)
+		{
+			PasswordDialog Dlg = new PasswordDialog();
+			Dlg.ShowDialog();
+
+			Password = Dlg.PasswordResult;
+			return Password != null;
+		}
+
+		private void OKButton_Click(object sender, EventArgs e)
+		{
+			DialogResult = DialogResult.OK;
+			PasswordResult = PasswordEntryBox.Text;
+			Close();
+		}
+	}
+}
diff --git a/Engine/Source/Programs/AutomationTool/IOS/IPhonePackager/SlowProgressDialog.cs b/Engine/Source/Programs/AutomationTool/IOS/IPhonePackager/SlowProgressDialog.cs
new file mode 100644
--- /dev/null
+++ b/Engine/Source/Programs/AutomationTool/IOS/IPhonePackager/SlowProgressDialog.cs
@@ -0,0 +1,63 @@
+﻿// Copyright 1998-2014 Epic Games, Inc. All Rights Reserved.
+
+using System;
+using System.Collections.Generic;
+using System.ComponentModel;
+using System.Data;
+using System.Drawing;
+using System.Text;
+using System.Windows.Forms;
+
+namespace iPhonePackager
+{
+	public partial class SlowProgressDialog : Form
+	{
+		public SlowProgressDialog()
+		{
+			InitializeComponent();
+		}
+
+		public delegate void WorkToDoSignature(BackgroundWorker BGWorker);
+		public WorkToDoSignature OnBeginBackgroundWork = null;
+
+		private void SlowWork_RunWorkerCompleted(object sender, RunWorkerCompletedEventArgs e)
+		{
+			DialogResult = DialogResult.OK;
+			ProgressIndicator.Value = 0;
+			StatusLabel.Text = "";
+		}
+
+		private void SlowWork_DoWork(object sender, DoWorkEventArgs e)
+		{
+			Program.BGWorker = SlowWork;
+			Program.ProgressIndex = 0;
+			try
+			{
+				OnBeginBackgroundWork(SlowWork);
+			}
+			catch (Exception ex)
+			{
+				Program.Error(ex.Message);
+			}
+			Program.BGWorker = null;
+		}
+
+		private void SlowWork_ProgressChanged(object sender, ProgressChangedEventArgs e)
+		{
+			ProgressIndicator.Value = e.ProgressPercentage;
+			StatusLabel.Text = e.UserState as string;
+		}
+
+		private void SlowProgressDialog_Shown(object sender, EventArgs e)
+		{
+			if (OnBeginBackgroundWork != null)
+			{
+				SlowWork.RunWorkerAsync();
+			}
+			else
+			{
+				DialogResult = DialogResult.Abort;
+			}
+		}
+	}
+}
diff --git a/Engine/Source/Programs/AutomationTool/IOS/IPhonePackager/ToolsHub.cs b/Engine/Source/Programs/AutomationTool/IOS/IPhonePackager/ToolsHub.cs
new file mode 100644
--- /dev/null
+++ b/Engine/Source/Programs/AutomationTool/IOS/IPhonePackager/ToolsHub.cs
@@ -0,0 +1,340 @@
+﻿// Copyright 1998-2014 Epic Games, Inc. All Rights Reserved.
+
+using System;
+using System.Collections.Generic;
+using System.ComponentModel;
+using System.Data;
+using System.Drawing;
+using System.Linq;
+using System.Text;
+using System.Windows.Forms;
+using System.IO;
+using System.Security.Cryptography.X509Certificates;
+using System.Diagnostics;
+
+namespace iPhonePackager
+{
+	public partial class ToolsHub : Form
+	{
+		protected Dictionary<bool, Bitmap> CheckStateImages = new Dictionary<bool, Bitmap>();
+
+		public static string IpaFilter = "iOS packaged applications (*.ipa)|*.ipa|All Files (*.*)|*.*";
+		public static string CertificateRequestFilter = "Certificate Request (*.csr)|*.csr|All Files (*.*)|*.*";
+		public static string MobileProvisionFilter = "Mobile provision files (*.mobileprovision)|*.mobileprovision|All files (*.*)|*.*";
+		public static string CertificatesFilter = "Code signing certificates (*.cer;*.p12)|*.cer;*.p12|All files (*.*)|*.*";
+		public static string KeysFilter = "Keys (*.key;*.p12)|*.key;*.p12|All files (*.*)|*.*";
+		public static string JustXmlKeysFilter = "XML Public/Private Key Pair (*.key)|*.key|All Files (*.*)|*.*";
+		public static string PListFilter = "Property Lists (*.plist)|*.plist|All Files (*.*)|*.*";
+		public static string JustP12Certificates = "Code signing certificates (*.p12)|*.p12|All files (*.*)|*.*";
+
+
+		//@TODO
+		public static string ChoosingFilesToInstallDirectory = null;
+		public static string ChoosingIpaDirectory = null;
+
+		/// <summary>
+		/// Shows a file dialog, preserving the current working directory of the application
+		/// </summary>
+		/// <param name="Filter">Filter string for the dialog</param>
+		/// <param name="Title">Display title for the dialog</param>
+		/// <param name="DefaultExtension">Default extension</param>
+		/// <param name="StartingFilename">Initial filename (can be null)</param>
+		/// <param name="SavedDirectory">Starting directory for the dialog (will be updated on a successful pick)</param>
+		/// <param name="OutFilename">(out) The picked filename, or null if the dialog was cancelled </param>
+		/// <returns>True on a successful pick (filename is in OutFilename)</returns>
+		public static bool ShowOpenFileDialog(string Filter, string Title, string DefaultExtension, string StartingFilename, ref string SavedDirectory, out string OutFilename)
+		{
+			// Save off the current working directory
+			string CWD = Directory.GetCurrentDirectory();
+
+			// Show the dialog
+			System.Windows.Forms.OpenFileDialog OpenDialog = new System.Windows.Forms.OpenFileDialog();
+			OpenDialog.DefaultExt = DefaultExtension;
+			OpenDialog.FileName = (StartingFilename != null) ? StartingFilename : "";
+			OpenDialog.Filter = Filter;
+			OpenDialog.Title = Title;
+			OpenDialog.InitialDirectory = (SavedDirectory != null) ? SavedDirectory : CWD;
+
+			bool bDialogSucceeded = OpenDialog.ShowDialog() == DialogResult.OK;
+
+			// Restore the current working directory
+			Directory.SetCurrentDirectory(CWD);
+
+			if (bDialogSucceeded)
+			{
+				SavedDirectory = Path.GetDirectoryName(OpenDialog.FileName);
+				OutFilename = OpenDialog.FileName;
+
+				return true;
+			}
+			else
+			{
+				OutFilename = null;
+				return false;
+			}
+		}
+
+		public ToolsHub()
+		{
+			InitializeComponent();
+
+			CheckStateImages.Add(false, iPhonePackager.Properties.Resources.GreyCheck);
+			CheckStateImages.Add(true, iPhonePackager.Properties.Resources.GreenCheck);
+
+			Text = Config.AppDisplayName + " Wizard";
+		}
+
+		public static ToolsHub CreateShowingTools()
+		{
+			ToolsHub Result = new ToolsHub();
+			Result.tabControl1.SelectTab(Result.tabPage3);
+			return Result;
+		}
+
+		private void ToolsHub_Shown(object sender, EventArgs e)
+		{
+			// Check the status of the various steps
+			UpdateStatus();
+		}
+
+		void UpdateStatus()
+		{
+			MobileProvision Provision;
+			X509Certificate2 Cert;
+			bool bOverridesExists;
+			CodeSignatureBuilder.FindRequiredFiles(out Provision, out Cert, out bOverridesExists);
+
+			MobileProvisionCheck2.Image = MobileProvisionCheck.Image = CheckStateImages[Provision != null];
+			CertificatePresentCheck2.Image = CertificatePresentCheck.Image = CheckStateImages[Cert != null];
+			OverridesPresentCheck2.Image = OverridesPresentCheck.Image = CheckStateImages[bOverridesExists];
+
+			ReadyToPackageButton.Enabled = bOverridesExists && (Provision != null) && (Cert != null);
+		}
+
+		private void CreateCSRButton_Click(object sender, EventArgs e)
+		{
+			GenerateSigningRequestDialog Dlg = new GenerateSigningRequestDialog();
+			Dlg.ShowDialog(this);
+		}
+
+		public static void ShowError(string Message)
+		{
+			MessageBox.Show(Message, Config.AppDisplayName, MessageBoxButtons.OK, MessageBoxIcon.Error);
+		}
+
+		public static bool IsProfileForDistribution(MobileProvision Provision)
+		{
+			return CryptoAdapter.GetCommonNameFromCert(Provision.DeveloperCertificates[0]).IndexOf("iPhone Distribution", StringComparison.InvariantCultureIgnoreCase) >= 0;
+		}
+
+		public static void TryInstallingMobileProvision()
+		{
+			string ProvisionFilename;
+			if (ShowOpenFileDialog(MobileProvisionFilter, "Choose a mobile provision to install", "mobileprovision", "", ref ChoosingFilesToInstallDirectory, out ProvisionFilename))
+			{
+				try
+				{
+					// Determine if this is a development or distribution certificate
+					bool bIsDistribution = false;
+					MobileProvision Provision = MobileProvisionParser.ParseFile(ProvisionFilename);
+					bIsDistribution = IsProfileForDistribution(Provision);
+
+					// Copy the file into the destination location
+					string EffectivePrefix = bIsDistribution ? "Distro_" : Config.SigningPrefix;
+					string DestinationFilename = Path.Combine(Config.BuildDirectory, EffectivePrefix + Program.GameName + ".mobileprovision");
+
+					if (File.Exists(DestinationFilename))
+					{
+						MobileProvision OldProvision = MobileProvisionParser.ParseFile(DestinationFilename);
+
+						string MessagePrompt = String.Format(
+							"{0} already contains a {1} mobile provision file.  Do you want to replace the provision '{2}' with '{3}'?",
+							Config.BuildDirectory,
+							bIsDistribution ? "distribution" : "development",
+							OldProvision.ProvisionName,
+							Provision.ProvisionName);
+
+						if (MessageBox.Show(MessagePrompt, Config.AppDisplayName, MessageBoxButtons.YesNo, MessageBoxIcon.Question) == DialogResult.No)
+						{
+							return;
+						}
+
+						FileOperations.DeleteFile(DestinationFilename);
+					}
+
+					FileOperations.CopyRequiredFile(ProvisionFilename, DestinationFilename);
+				}
+				catch (Exception ex)
+				{
+					ShowError(String.Format("Encountered an error '{0} while trying to install a mobile provision", ex.Message));
+				}
+			}
+
+		}
+
+		private void ImportMobileProvisionButton_Click(object sender, EventArgs e)
+		{
+			TryInstallingMobileProvision();
+			UpdateStatus();
+		}
+
+		public static void TryInstallingCertificate_PromptForKey()
+		{
+			try
+			{
+				string CertificateFilename;
+				if (ShowOpenFileDialog(CertificatesFilter, "Choose a code signing certificate to import", "", "", ref ChoosingFilesToInstallDirectory, out CertificateFilename))
+				{
+					// Load the certificate
+					string CertificatePassword = "";
+					X509Certificate2 Cert = null;
+					try
+					{
+						Cert = new X509Certificate2(CertificateFilename, CertificatePassword, X509KeyStorageFlags.PersistKeySet | X509KeyStorageFlags.Exportable | X509KeyStorageFlags.MachineKeySet);
+					}
+					catch (System.Security.Cryptography.CryptographicException ex)
+					{
+						// Try once with a password
+						if (PasswordDialog.RequestPassword(out CertificatePassword))
+						{
+							Cert = new X509Certificate2(CertificateFilename, CertificatePassword, X509KeyStorageFlags.PersistKeySet | X509KeyStorageFlags.Exportable | X509KeyStorageFlags.MachineKeySet);
+						}
+						else
+						{
+							// User cancelled dialog, rethrow
+							throw ex;
+						}
+					}
+
+					// If the certificate doesn't have a private key pair, ask the user to provide one
+					if (!Cert.HasPrivateKey)
+					{
+						string ErrorMsg = "Certificate does not include a private key and cannot be used to code sign";
+
+						// Prompt for a key pair
+						if (MessageBox.Show("Next, please choose the key pair that you made when generating the certificate request.",
+							Config.AppDisplayName,
+							MessageBoxButtons.OK,
+							MessageBoxIcon.Information) == DialogResult.OK)
+						{
+							string KeyFilename;
+							if (ShowOpenFileDialog(KeysFilter, "Choose the key pair that belongs with the signing certificate", "", "", ref ChoosingFilesToInstallDirectory, out KeyFilename))
+							{
+								Cert = CryptoAdapter.CombineKeyAndCert(CertificateFilename, KeyFilename);
+
+								if (Cert.HasPrivateKey)
+								{
+									ErrorMsg = null;
+								}
+							}
+						}
+
+						if (ErrorMsg != null)
+						{
+							throw new Exception(ErrorMsg);
+						}
+					}
+
+					// Add the certificate to the store
+					X509Store Store = new X509Store();
+					Store.Open(OpenFlags.ReadWrite);
+					Store.Add(Cert);
+					Store.Close();
+				}
+			}
+			catch (Exception ex)
+			{
+				string ErrorMsg = String.Format("Failed to load or install certificate due to an error: '{0}'", ex.Message);
+				Program.Error(ErrorMsg);
+				MessageBox.Show(ErrorMsg, Config.AppDisplayName, MessageBoxButtons.OK, MessageBoxIcon.Error);
+			}
+		}
+
+		private void ImportCertificateButton_Click(object sender, EventArgs e)
+		{
+			TryInstallingCertificate_PromptForKey();
+			UpdateStatus();
+		}
+
+		private void EditPlistButton_Click(object sender, EventArgs e)
+		{
+			ConfigureMobileGame Dlg = new ConfigureMobileGame();
+			Dlg.ShowDialog(this);
+			UpdateStatus();
+		}
+
+		private void CancelThisFormButton_Click(object sender, EventArgs e)
+		{
+			DialogResult = DialogResult.Cancel;
+			Close();
+		}
+
+		private void ReadyToPackageButton_Click(object sender, EventArgs e)
+		{
+			DialogResult = DialogResult.OK;
+			Close();
+		}
+
+		private void InstallIPAButton_Click(object sender, EventArgs e)
+		{
+			string PickedFilename;
+			if (ShowOpenFileDialog(IpaFilter, "Choose an IPA to install", "ipa", "", ref ChoosingIpaDirectory, out PickedFilename))
+			{
+				Program.ProgressDialog.OnBeginBackgroundWork = delegate
+				{
+					DeploymentHelper.InstallIPAOnConnectedDevices(PickedFilename);
+				};
+				Program.ProgressDialog.ShowDialog();
+			}
+		}
+
+		private void ResignIPAButton_Click(object sender, EventArgs e)
+		{
+			GraphicalResignTool ResignTool = GraphicalResignTool.GetActiveInstance();
+			ResignTool.TabBook.SelectTab(ResignTool.ResignPage);
+			ResignTool.Show();
+		}
+
+		private void ProvisionCertToolsButton_Click(object sender, EventArgs e)
+		{
+			GraphicalResignTool ResignTool = GraphicalResignTool.GetActiveInstance();
+			ResignTool.TabBook.SelectTab(ResignTool.ProvisionCertPage);
+			ResignTool.Show();
+		}
+
+		private void OtherDeployToolsButton_Click(object sender, EventArgs e)
+		{
+			GraphicalResignTool ResignTool = GraphicalResignTool.GetActiveInstance();
+			ResignTool.TabBook.SelectTab(ResignTool.DeploymentPage);
+			ResignTool.Show();
+		}
+
+		void OpenHelpWebsite()
+		{
+			string Target = "http://udn.epicgames.com/Three/AppleiOSProvisioning.html";
+			ProcessStartInfo PSI = new ProcessStartInfo(Target);
+			Process.Start(PSI);
+		}
+
+		private void ToolsHub_HelpButtonClicked(object sender, CancelEventArgs e)
+		{
+			OpenHelpWebsite();
+		}
+
+		private void ToolsHub_KeyUp(object sender, KeyEventArgs e)
+		{
+			if (e.KeyCode == Keys.F1)
+			{
+				OpenHelpWebsite();
+			}
+		}
+		
+		private void HyperlinkClicked(object sender, LinkLabelLinkClickedEventArgs e)
+		{
+			string Target = (sender as LinkLabel).Tag as string;
+			ProcessStartInfo PSI = new ProcessStartInfo(Target);
+			Process.Start(PSI);
+		}
+
+	}
+}
\ No newline at end of file
diff --git a/Engine/Source/Programs/AutomationTool/IOS/IPhonePackager/iPhonePackager.cs b/Engine/Source/Programs/AutomationTool/IOS/IPhonePackager/iPhonePackager.cs
--- a/Engine/Source/Programs/AutomationTool/IOS/IPhonePackager/iPhonePackager.cs
+++ b/Engine/Source/Programs/AutomationTool/IOS/IPhonePackager/iPhonePackager.cs
@@ -11,6 +11,7 @@ using System.Reflection;
 using System.Runtime.InteropServices;
 using System.Text;
 using System.Threading;
+using System.Windows.Forms;
 using Ionic.Zlib;
 using System.ComponentModel;
 
@@ -26,6 +27,7 @@ namespace iPhonePackager
 		static public string GameConfiguration = "";
 		static public string Architecture = "";
 
+		static public SlowProgressDialog ProgressDialog = null;
 		static public BackgroundWorker BGWorker = null;
 		static public int ProgressIndex = 0;
 
@@ -264,8 +266,15 @@ namespace iPhonePackager
 				string CWD = Directory.GetCurrentDirectory();
 				if (!CodeSignatureBuilder.DoRequiredFilesExist())
 				{
-					Error("One or more files necessary for packaging are missing - please use Project Settings in the editor to set up");
-					return false;
+					StartVisuals();
+
+					Form DisplayForm = new ToolsHub();
+					Application.Run(DisplayForm);
+					if (DisplayForm.DialogResult != DialogResult.OK)
+					{
+						Error("One or more files necessary for packaging are missing and configuration was cancelled");
+						return false;
+					}
 				}
 				Directory.SetCurrentDirectory(CWD);
 			}
@@ -302,6 +311,28 @@ namespace iPhonePackager
 					bHandledCommand = DeploymentHelper.ExecuteDeployCommand(Command, GamePath, RPCCommand);
 				}
 
+				if (!bHandledCommand)
+				{
+					bHandledCommand = true;
+					switch (Command)
+					{
+						case "configure":
+							if (Config.bForDistribution)
+							{
+								Error("configure cannot be used with -distribution");
+							}
+							else
+							{
+								RunInVisualMode(delegate { return new ConfigureMobileGame(); });
+							}
+							break;
+
+						default:
+							bHandledCommand = false;
+							break;
+					}
+				}
+
 				if (!bHandledCommand)
 				{
 					Error("Unknown command");
@@ -378,6 +409,24 @@ namespace iPhonePackager
 			return ( null );
 		}
 
+		delegate Form CreateFormDelegate();
+
+		static void StartVisuals()
+		{
+			Application.EnableVisualStyles();
+			Application.SetCompatibleTextRenderingDefault(false);
+
+			ProgressDialog = new SlowProgressDialog();
+		}
+
+		static void RunInVisualMode(CreateFormDelegate Work)
+		{
+			StartVisuals();
+
+			Form DisplayForm = Work.Invoke();
+			Application.Run(DisplayForm);
+		}
+
 		static void ListDevices()
 		{
 			var DeviceList = DeploymentHelper.Get().EnumerateConnectedDevices();
@@ -525,6 +574,18 @@ namespace iPhonePackager
 						}
 						break;
 
+					case "resigntool":
+						RunInVisualMode(delegate { return new GraphicalResignTool(); });
+						break;
+
+					case "certrequest":
+						RunInVisualMode(delegate { return new GenerateSigningRequestDialog(); });
+						break;
+
+					case "gui":
+						RunInVisualMode(delegate { return ToolsHub.CreateShowingTools(); });
+						break;
+
 					case "devices":
 						ListDevices();
 						break;
