[INTEGRATE] Latest from 4.3 branch.

diff --git a/Engine/Source/Programs/AutomationTool/Win/WinPlatform.Automation.cs b/Engine/Source/Programs/AutomationTool/Win/WinPlatform.Automation.cs
--- a/Engine/Source/Programs/AutomationTool/Win/WinPlatform.Automation.cs
+++ b/Engine/Source/Programs/AutomationTool/Win/WinPlatform.Automation.cs
@@ -111,7 +111,7 @@ public abstract class BaseWinPlatform : Platform
             if (Exe.StartsWith(CombinePaths(SC.RuntimeProjectRootDir, "Binaries", SC.PlatformDir)))
             {
                 // remap the project root. For Rocket executables, rename the executable to the game name.
-				if ((Params.Rocket || !Params.IsCodeBasedProject) && Exe == Exes[0])
+				if (!Params.IsCodeBasedProject && Exe == Exes[0])
 				{
 					StageExecutable("exe", SC, CombinePaths(SC.ProjectRoot, "Binaries", SC.PlatformDir), Path.GetFileNameWithoutExtension(Exe) + ".", true, null, CommandUtils.CombinePaths(SC.RelativeProjectRootForStage, "Binaries", SC.PlatformDir), false, WorkingFileType, SC.ShortProjectName + ".");
 				}
@@ -123,7 +123,7 @@ public abstract class BaseWinPlatform : Platform
             else if (Exe.StartsWith(CombinePaths(SC.RuntimeRootDir, "Engine/Binaries", SC.PlatformDir)))
             {
 				// Move the executable for non-code rocket projects into the game directory, using the game name, so it can figure out the UProject to look for and is consitent with code projects.
-				if ((Params.Rocket || !Params.IsCodeBasedProject) && Exe == Exes[0])
+				if (!Params.IsCodeBasedProject && Exe == Exes[0])
 				{
 					StageExecutable("exe", SC, CombinePaths(SC.LocalRoot, "Engine/Binaries", SC.PlatformDir), Path.GetFileNameWithoutExtension(Exe) + ".", true, null, CommandUtils.CombinePaths(SC.RelativeProjectRootForStage, "Binaries", SC.PlatformDir), false, WorkingFileType, SC.ShortProjectName + ".");
 				}
diff --git a/Engine/Source/Programs/UnrealBuildTool/Configuration/UEBuildTarget.cs b/Engine/Source/Programs/UnrealBuildTool/Configuration/UEBuildTarget.cs
--- a/Engine/Source/Programs/UnrealBuildTool/Configuration/UEBuildTarget.cs
+++ b/Engine/Source/Programs/UnrealBuildTool/Configuration/UEBuildTarget.cs
@@ -1442,7 +1442,7 @@ namespace UnrealBuildTool
 				// Set the define for the project name. This allows the executable to locate the correct project file to use, which may not be the same as the game name or target.
 				if(UnrealBuildTool.HasUProjectFile())
 				{
-					string ProjectName = Path.GetFileNameWithoutExtension(UnrealBuildTool.GetUProjectPath());
+					string ProjectName = Path.GetFileNameWithoutExtension(UnrealBuildTool.GetUProjectFile());
 					GlobalCompileEnvironment.Config.Definitions.Add(String.Format("UE_PROJECT_NAME={0}", ProjectName));
 				}
 
