BPS2 FakeUnshelve and preflight unshelving support.
FakeUnshelve is an process that is unshelving P4 shelve without checking out files.

diff --git a/Engine/Source/Programs/AutomationTool/Scripts/GUBP.Automation.cs b/Engine/Source/Programs/AutomationTool/Scripts/GUBP.Automation.cs
--- a/Engine/Source/Programs/AutomationTool/Scripts/GUBP.Automation.cs
+++ b/Engine/Source/Programs/AutomationTool/Scripts/GUBP.Automation.cs
@@ -4402,12 +4402,6 @@ public class GUBP : BuildCommand
                 }
             }
             bPreflightBuild = true;
-
-            var OpenFiles = P4.OpenedOutput();
-            if (OpenFiles.Contains("File(s) not opened on this client."))
-            {
-                throw new AutomationException("Apparently, nothing was unshelved for this preflight.");
-            }
         }
         
         ECProject = ParseParamValue("ECProject");
