build script, verify dates and sizes on all file copies and retry if there is a mismatch

diff --git a/Engine/Source/Programs/AutomationTool/Utils.cs b/Engine/Source/Programs/AutomationTool/Utils.cs
--- a/Engine/Source/Programs/AutomationTool/Utils.cs
+++ b/Engine/Source/Programs/AutomationTool/Utils.cs
@@ -276,7 +276,6 @@ namespace AutomationTool
 			return Result;
 		}
 
-
 		/// <summary>
 		/// Copies a file.
 		/// </summary>
@@ -296,29 +295,51 @@ namespace AutomationTool
 			do
 			{
 				Result = true;
+                bool Retry = true;
 				try
 				{
 					File.Copy(SourceName, TargetName, overwrite: true);
-				}
-				catch (Exception Ex)
-				{
-					if (File.Exists(TargetName) == false)
-					{
-                        if (Attempts + 1 < MaxAttempts)
+                    Retry = !File.Exists(TargetName);
+                    if (!Retry)
+                    {
+                        FileInfo SourceInfo = new FileInfo(SourceName);
+                        FileInfo TargetInfo = new FileInfo(TargetName);
+                        if (SourceInfo.Length != TargetInfo.Length)
                         {
-                            Log.WriteLine(TraceEventType.Warning, "Failed to copy {0} to {1}, waiting 10s.", SourceName, TargetName);
-                            Log.WriteLine(TraceEventType.Warning, LogUtils.FormatException(Ex));
-                            Thread.Sleep(10000);
+                            Log.WriteLine(TraceEventType.Warning, "Size mismatch {0} = {1} to {2} = {3}", SourceName, SourceInfo.Length, TargetName, TargetInfo.Length);
+                            Retry = true;
                         }
-                        else
+                        if (!((SourceInfo.LastWriteTimeUtc - TargetInfo.LastWriteTimeUtc).TotalSeconds < 1 && (SourceInfo.LastWriteTimeUtc - TargetInfo.LastWriteTimeUtc).TotalSeconds > -1))
                         {
-                            Log.WriteLine(TraceEventType.Warning, "Failed to copy {0} to {1}", SourceName, TargetName);
-                            Log.WriteLine(TraceEventType.Warning, LogUtils.FormatException(Ex));
+                            Log.WriteLine(TraceEventType.Warning, "Date mismatch {0} = {1} to {2} = {3}", SourceName, SourceInfo.LastWriteTimeUtc, TargetName, TargetInfo.LastWriteTimeUtc);
+                            Retry = true;
                         }
-                        Result = false;
-					}
+                    }
+                }
+				catch (Exception Ex)
+				{
+                    Log.WriteLine(System.Diagnostics.TraceEventType.Warning, "SafeCopyFile Exception was {0}", LogUtils.FormatException(Ex));
+                    Retry = true;
 				}
-			}
+
+                if (Retry)
+                {
+                    if (Attempts + 1 < MaxAttempts)
+                    {
+                        Log.WriteLine(TraceEventType.Warning, "Failed to copy {0} to {1}, deleting, waiting 10s and retrying.", SourceName, TargetName);
+                        if (File.Exists(TargetName))
+                        {
+                            SafeDeleteFile(TargetName);
+                        }
+                        Thread.Sleep(10000);
+                    }
+                    else
+                    {
+                        Log.WriteLine(TraceEventType.Warning, "Failed to copy {0} to {1}", SourceName, TargetName);
+                    }
+                    Result = false;
+                }
+            }
 			while (Result == false && ++Attempts < MaxAttempts);
 
 			return Result;
