Allow game control over disc layer that streaming chunks are placed on.

diff --git a/Engine/Source/Editor/UnrealEd/Private/Commandlets/ChunkManifestGenerator.cpp b/Engine/Source/Editor/UnrealEd/Private/Commandlets/ChunkManifestGenerator.cpp
--- a/Engine/Source/Editor/UnrealEd/Private/Commandlets/ChunkManifestGenerator.cpp
+++ b/Engine/Source/Editor/UnrealEd/Private/Commandlets/ChunkManifestGenerator.cpp
@@ -107,6 +107,9 @@ bool FChunkManifestGenerator::GenerateStreamingInstallManifest(const FString& Pl
 		return false;
 	}
 
+	FString PakChunkLayerInfoFilename = FString::Printf(TEXT("%s/pakchunklayers.txt"), *TmpPackagingDir);
+	TAutoPtr<FArchive> ChunkLayerFile(IFileManager::Get().CreateFileWriter(*PakChunkLayerInfoFilename));
+
 	// generate per-chunk pak list files
 	for (int32 Index = 0; Index < ChunkManifests.Num(); ++Index)
 	{			
@@ -135,9 +138,17 @@ bool FChunkManifestGenerator::GenerateStreamingInstallManifest(const FString& Pl
 			// add this pakfilelist to our master list of pakfilelists
 			FString PakChunkListLine = FString::Printf(TEXT("pakchunk%d.txt\r\n"), Index);
 			PakChunkListFile->Serialize(TCHAR_TO_ANSI(*PakChunkListLine), PakChunkListLine.Len());
+
+			int32 TargetLayer = 0;
+			FGameDelegates::Get().GetAssignLayerChunkDelegate().ExecuteIfBound(&Manifest, Platform, Index, TargetLayer);
+
+			FString LayerString = FString::Printf(TEXT("%d\r\n"), TargetLayer);
+			
+			ChunkLayerFile->Serialize(TCHAR_TO_ANSI(*LayerString), LayerString.Len());
 		}
 	}
 
+	ChunkLayerFile->Close();
 	PakChunkListFile->Close();
 
 	return true;
diff --git a/Engine/Source/Programs/AutomationTool/Scripts/CopyBuildToStagingDirectory.Automation.cs b/Engine/Source/Programs/AutomationTool/Scripts/CopyBuildToStagingDirectory.Automation.cs
--- a/Engine/Source/Programs/AutomationTool/Scripts/CopyBuildToStagingDirectory.Automation.cs
+++ b/Engine/Source/Programs/AutomationTool/Scripts/CopyBuildToStagingDirectory.Automation.cs
@@ -18,15 +18,15 @@ using UnrealBuildTool;
 public partial class Project : CommandUtils
 {
 
-	#region Utilities
-
-	/// <summary>
-	/// Writes a pak response file to disk
-	/// </summary>
-	/// <param name="Filename"></param>
-	/// <param name="ResponseFile"></param>
-	private static void WritePakResponseFile(string Filename, Dictionary<string, string> ResponseFile, bool Compressed)
-	{
+    #region Utilities
+
+    /// <summary>
+    /// Writes a pak response file to disk
+    /// </summary>
+    /// <param name="Filename"></param>
+    /// <param name="ResponseFile"></param>
+    private static void WritePakResponseFile(string Filename, Dictionary<string, string> ResponseFile, bool Compressed)
+    {
         using (var Writer = new StreamWriter(Filename, false, new System.Text.UTF8Encoding(true)))
         {
             foreach (var Entry in ResponseFile)
@@ -39,527 +39,537 @@ public partial class Project : CommandUtils
                 Writer.WriteLine(Line);
             }
         }
-	}
-
-	/// <summary>
-	/// Loads streaming install chunk manifest file from disk
-	/// </summary>
-	/// <param name="Filename"></param>
-	/// <returns></returns>
-	private static HashSet<string> ReadPakChunkManifest(string Filename)
-	{		
-		var ResponseFile = ReadAllLines(Filename);
-		var Result = new HashSet<string>(ResponseFile, StringComparer.InvariantCultureIgnoreCase);
-		return Result;
-	}
-
-	static public void RunUnrealPak(Dictionary<string, string> UnrealPakResponseFile, string OutputLocation, string EncryptionKeys, string PakOrderFileLocation, string PlatformOptions, bool Compressed)
-	{
-		if (UnrealPakResponseFile.Count < 1)
-		{
+    }
+
+    /// <summary>
+    /// Loads streaming install chunk manifest file from disk
+    /// </summary>
+    /// <param name="Filename"></param>
+    /// <returns></returns>
+    private static HashSet<string> ReadPakChunkManifest(string Filename)
+    {		
+        var ResponseFile = ReadAllLines(Filename);
+        var Result = new HashSet<string>(ResponseFile, StringComparer.InvariantCultureIgnoreCase);
+        return Result;
+    }
+
+    static public void RunUnrealPak(Dictionary<string, string> UnrealPakResponseFile, string OutputLocation, string EncryptionKeys, string PakOrderFileLocation, string PlatformOptions, bool Compressed)
+    {
+        if (UnrealPakResponseFile.Count < 1)
+        {
             return;
-		}
-		string PakName = Path.GetFileNameWithoutExtension(OutputLocation);
-		string UnrealPakResponseFileName = CombinePaths(CmdEnv.LogFolder, "PakList_" + PakName + ".txt");
+        }
+        string PakName = Path.GetFileNameWithoutExtension(OutputLocation);
+        string UnrealPakResponseFileName = CombinePaths(CmdEnv.LogFolder, "PakList_" + PakName + ".txt");
         WritePakResponseFile(UnrealPakResponseFileName, UnrealPakResponseFile, Compressed);
 
-		var UnrealPakExe = CombinePaths(CmdEnv.LocalRoot, "Engine/Binaries/Win64/unrealpak.exe");
-
-		Log("Running UnrealPak *******");
-		string CmdLine = CommandUtils.MakePathSafeToUseWithCommandLine(OutputLocation) + " -create=" + CommandUtils.MakePathSafeToUseWithCommandLine(UnrealPakResponseFileName);
-		if (!String.IsNullOrEmpty(EncryptionKeys))
-		{
-			CmdLine += " -sign=" + CommandUtils.MakePathSafeToUseWithCommandLine(EncryptionKeys);
-		}
-		if (GlobalCommandLine.Installed)
-		{
-			CmdLine += " -installed";
-		}
-		CmdLine += " -order=" + CommandUtils.MakePathSafeToUseWithCommandLine(PakOrderFileLocation);
-		if (GlobalCommandLine.UTF8Output)
-		{
-			CmdLine += " -UTF8Output";
-		}
+        var UnrealPakExe = CombinePaths(CmdEnv.LocalRoot, "Engine/Binaries/Win64/unrealpak.exe");
+
+        Log("Running UnrealPak *******");
+        string CmdLine = CommandUtils.MakePathSafeToUseWithCommandLine(OutputLocation) + " -create=" + CommandUtils.MakePathSafeToUseWithCommandLine(UnrealPakResponseFileName);
+        if (!String.IsNullOrEmpty(EncryptionKeys))
+        {
+            CmdLine += " -sign=" + CommandUtils.MakePathSafeToUseWithCommandLine(EncryptionKeys);
+        }
+        if (GlobalCommandLine.Installed)
+        {
+            CmdLine += " -installed";
+        }
+        CmdLine += " -order=" + CommandUtils.MakePathSafeToUseWithCommandLine(PakOrderFileLocation);
+        if (GlobalCommandLine.UTF8Output)
+        {
+            CmdLine += " -UTF8Output";
+        }
         CmdLine += PlatformOptions;
-		RunAndLog(CmdEnv, UnrealPakExe, CmdLine, Options: ERunOptions.Default | ERunOptions.UTF8Output);
-		Log("UnrealPak Done *******");
-	}
-
-	static public void LogDeploymentContext(DeploymentContext SC)
-	{
-		Log("Deployment Context **************");
-		Log("ArchiveDirectory = {0}", SC.ArchiveDirectory);
-		Log("RawProjectPath = {0}", SC.RawProjectPath);
-		Log("IsCodeBasedUprojectFile = {0}", SC.IsCodeBasedProject);
-		Log("DedicatedServer = {0}", SC.DedicatedServer);
-		Log("Stage = {0}", SC.Stage);
-		Log("StageTargetPlatform = {0}", SC.StageTargetPlatform.PlatformType.ToString());
-		Log("LocalRoot = {0}", SC.LocalRoot);
-		Log("ProjectRoot = {0}", SC.ProjectRoot);
-		Log("PlatformDir = {0}", SC.PlatformDir);
-		Log("StageProjectRoot = {0}", SC.StageProjectRoot);
-		Log("ShortProjectName = {0}", SC.ShortProjectName);
-		Log("StageDirectory = {0}", SC.StageDirectory);
-		Log("SourceRelativeProjectRoot = {0}", SC.SourceRelativeProjectRoot);
-		Log("RelativeProjectRootForStage = {0}", SC.RelativeProjectRootForStage);
-		Log("RelativeProjectRootForUnrealPak = {0}", SC.RelativeProjectRootForUnrealPak);
-		Log("ProjectArgForCommandLines = {0}", SC.ProjectArgForCommandLines);
-		Log("RuntimeRootDir = {0}", SC.RuntimeRootDir);
-		Log("RuntimeProjectRootDir = {0}", SC.RuntimeProjectRootDir);
-		Log("UProjectCommandLineArgInternalRoot = {0}", SC.UProjectCommandLineArgInternalRoot);
-		Log("PakFileInternalRoot = {0}", SC.PakFileInternalRoot);
-		Log("UnrealFileServerInternalRoot = {0}", SC.UnrealFileServerInternalRoot);
-		Log("End Deployment Context **************");
-	}
-
-	public static Dictionary<string, string> ConvertToLower(Dictionary<string, string> Mapping)
-	{
-		var Result = new Dictionary<string, string>();
-		foreach (var Pair in Mapping)
-		{
-			Result.Add(Pair.Key, Pair.Value.ToLowerInvariant());
-		}
-		return Result;
-	}
-
-	public static void MaybeConvertToLowerCase(ProjectParams Params, DeploymentContext SC)
-	{
-		var BuildPlatform = SC.StageTargetPlatform;
-
-		if (BuildPlatform.DeployLowerCaseFilenames(false))
-		{
-			SC.NonUFSStagingFiles = ConvertToLower(SC.NonUFSStagingFiles);
-			SC.NonUFSStagingFilesDebug = ConvertToLower(SC.NonUFSStagingFilesDebug);
-		}
-		if (Params.Pak && BuildPlatform.DeployPakInternalLowerCaseFilenames())
-		{
-			SC.UFSStagingFiles = ConvertToLower(SC.UFSStagingFiles);
-		}
-		else if (!Params.Pak && BuildPlatform.DeployLowerCaseFilenames(true))
-		{
-			SC.UFSStagingFiles = ConvertToLower(SC.UFSStagingFiles);
-		}
-	}
-
-	public static void CreateStagingManifest(ProjectParams Params, DeploymentContext SC)
-	{
-		if (!Params.Stage)
-		{
-			return;
-		}
-		var ThisPlatform = SC.StageTargetPlatform;		
-
-		ThisPlatform.GetFilesToDeployOrStage(Params, SC);
-
-		// Get the build.properties file
-		// this file needs to be treated as a UFS file for casing, but NonUFS for being put into the .pak file
-		// @todo: Maybe there should be a new category - UFSNotForPak
-		string BuildPropertiesPath = CombinePaths(SC.LocalRoot, "Engine/Build");
-		if (SC.StageTargetPlatform.DeployLowerCaseFilenames(true))
-		{
-			BuildPropertiesPath = BuildPropertiesPath.ToLower();
-		}
-		SC.StageFiles(StagedFileType.NonUFS, BuildPropertiesPath, "build.properties", false, null, null, true);
-		
-		// move the UE4Commandline.txt file to the root of the stage
-		// this file needs to be treated as a UFS file for casing, but NonUFS for being put into the .pak file
-		// @todo: Maybe there should be a new category - UFSNotForPak
-		string CommandLineFile = "UE4CommandLine.txt";
-		if (SC.StageTargetPlatform.DeployLowerCaseFilenames(true))
-		{
-			CommandLineFile = CommandLineFile.ToLower();
-		}
-		SC.StageFiles(StagedFileType.NonUFS, GetIntermediateCommandlineDir(SC), CommandLineFile, false, null, "", true, false);
-
-		if (!Params.CookOnTheFly && !Params.SkipCookOnTheFly) // only stage the UFS files if we are not using cook on the fly
-		{
-			// Engine ufs (content)
-			SC.StageFiles(StagedFileType.UFS, CombinePaths(SC.LocalRoot, "Engine/Config"), "*", true, null, null, false, !Params.Pak);
-
-			if (Params.bUsesSlate)
-			{
-				if (Params.bUsesSlateEditorStyle)
-				{
-					SC.StageFiles(StagedFileType.UFS, CombinePaths(SC.LocalRoot, "Engine/Content/Editor/Slate"), "*", true, null, null, false, !Params.Pak);
-				}
-				SC.StageFiles(StagedFileType.UFS, CombinePaths(SC.LocalRoot, "Engine/Content/Slate"), "*", true, null, null, false, !Params.Pak);
-				SC.StageFiles(StagedFileType.UFS, CombinePaths(SC.ProjectRoot, "Content/Slate"), "*", true, null, CombinePaths(SC.RelativeProjectRootForStage, "Content/Slate"), true, !Params.Pak);
-
-			}
+        RunAndLog(CmdEnv, UnrealPakExe, CmdLine, Options: ERunOptions.Default | ERunOptions.UTF8Output);
+        Log("UnrealPak Done *******");
+    }
+
+    static public void LogDeploymentContext(DeploymentContext SC)
+    {
+        Log("Deployment Context **************");
+        Log("ArchiveDirectory = {0}", SC.ArchiveDirectory);
+        Log("RawProjectPath = {0}", SC.RawProjectPath);
+        Log("IsCodeBasedUprojectFile = {0}", SC.IsCodeBasedProject);
+        Log("DedicatedServer = {0}", SC.DedicatedServer);
+        Log("Stage = {0}", SC.Stage);
+        Log("StageTargetPlatform = {0}", SC.StageTargetPlatform.PlatformType.ToString());
+        Log("LocalRoot = {0}", SC.LocalRoot);
+        Log("ProjectRoot = {0}", SC.ProjectRoot);
+        Log("PlatformDir = {0}", SC.PlatformDir);
+        Log("StageProjectRoot = {0}", SC.StageProjectRoot);
+        Log("ShortProjectName = {0}", SC.ShortProjectName);
+        Log("StageDirectory = {0}", SC.StageDirectory);
+        Log("SourceRelativeProjectRoot = {0}", SC.SourceRelativeProjectRoot);
+        Log("RelativeProjectRootForStage = {0}", SC.RelativeProjectRootForStage);
+        Log("RelativeProjectRootForUnrealPak = {0}", SC.RelativeProjectRootForUnrealPak);
+        Log("ProjectArgForCommandLines = {0}", SC.ProjectArgForCommandLines);
+        Log("RuntimeRootDir = {0}", SC.RuntimeRootDir);
+        Log("RuntimeProjectRootDir = {0}", SC.RuntimeProjectRootDir);
+        Log("UProjectCommandLineArgInternalRoot = {0}", SC.UProjectCommandLineArgInternalRoot);
+        Log("PakFileInternalRoot = {0}", SC.PakFileInternalRoot);
+        Log("UnrealFileServerInternalRoot = {0}", SC.UnrealFileServerInternalRoot);
+        Log("End Deployment Context **************");
+    }
+
+    public static Dictionary<string, string> ConvertToLower(Dictionary<string, string> Mapping)
+    {
+        var Result = new Dictionary<string, string>();
+        foreach (var Pair in Mapping)
+        {
+            Result.Add(Pair.Key, Pair.Value.ToLowerInvariant());
+        }
+        return Result;
+    }
+
+    public static void MaybeConvertToLowerCase(ProjectParams Params, DeploymentContext SC)
+    {
+        var BuildPlatform = SC.StageTargetPlatform;
+
+        if (BuildPlatform.DeployLowerCaseFilenames(false))
+        {
+            SC.NonUFSStagingFiles = ConvertToLower(SC.NonUFSStagingFiles);
+            SC.NonUFSStagingFilesDebug = ConvertToLower(SC.NonUFSStagingFilesDebug);
+        }
+        if (Params.Pak && BuildPlatform.DeployPakInternalLowerCaseFilenames())
+        {
+            SC.UFSStagingFiles = ConvertToLower(SC.UFSStagingFiles);
+        }
+        else if (!Params.Pak && BuildPlatform.DeployLowerCaseFilenames(true))
+        {
+            SC.UFSStagingFiles = ConvertToLower(SC.UFSStagingFiles);
+        }
+    }
+
+    public static void CreateStagingManifest(ProjectParams Params, DeploymentContext SC)
+    {
+        if (!Params.Stage)
+        {
+            return;
+        }
+        var ThisPlatform = SC.StageTargetPlatform;		
+
+        ThisPlatform.GetFilesToDeployOrStage(Params, SC);
+
+        // Get the build.properties file
+        // this file needs to be treated as a UFS file for casing, but NonUFS for being put into the .pak file
+        // @todo: Maybe there should be a new category - UFSNotForPak
+        string BuildPropertiesPath = CombinePaths(SC.LocalRoot, "Engine/Build");
+        if (SC.StageTargetPlatform.DeployLowerCaseFilenames(true))
+        {
+            BuildPropertiesPath = BuildPropertiesPath.ToLower();
+        }
+        SC.StageFiles(StagedFileType.NonUFS, BuildPropertiesPath, "build.properties", false, null, null, true);
+        
+        // move the UE4Commandline.txt file to the root of the stage
+        // this file needs to be treated as a UFS file for casing, but NonUFS for being put into the .pak file
+        // @todo: Maybe there should be a new category - UFSNotForPak
+        string CommandLineFile = "UE4CommandLine.txt";
+        if (SC.StageTargetPlatform.DeployLowerCaseFilenames(true))
+        {
+            CommandLineFile = CommandLineFile.ToLower();
+        }
+        SC.StageFiles(StagedFileType.NonUFS, GetIntermediateCommandlineDir(SC), CommandLineFile, false, null, "", true, false);
+
+        if (!Params.CookOnTheFly && !Params.SkipCookOnTheFly) // only stage the UFS files if we are not using cook on the fly
+        {
+            // Engine ufs (content)
+            SC.StageFiles(StagedFileType.UFS, CombinePaths(SC.LocalRoot, "Engine/Config"), "*", true, null, null, false, !Params.Pak);
+
+            if (Params.bUsesSlate)
+            {
+                if (Params.bUsesSlateEditorStyle)
+                {
+                    SC.StageFiles(StagedFileType.UFS, CombinePaths(SC.LocalRoot, "Engine/Content/Editor/Slate"), "*", true, null, null, false, !Params.Pak);
+                }
+                SC.StageFiles(StagedFileType.UFS, CombinePaths(SC.LocalRoot, "Engine/Content/Slate"), "*", true, null, null, false, !Params.Pak);
+                SC.StageFiles(StagedFileType.UFS, CombinePaths(SC.ProjectRoot, "Content/Slate"), "*", true, null, CombinePaths(SC.RelativeProjectRootForStage, "Content/Slate"), true, !Params.Pak);
+
+            }
             SC.StageFiles(StagedFileType.UFS, CombinePaths(SC.LocalRoot, "Engine/Content/Localization/Engine"), "*.locres", true, null, null, false, !Params.Pak);
-			SC.StageFiles(StagedFileType.UFS, CombinePaths(SC.LocalRoot, "Engine/Plugins"), "*.uplugin", true, null, null, true, !Params.Pak);
-
-			// Game ufs (content)
-			SC.StageFiles(StagedFileType.UFS, CombinePaths(SC.ProjectRoot), "*.uproject", false, null, CombinePaths(SC.RelativeProjectRootForStage), true, !Params.Pak);
-			SC.StageFiles(StagedFileType.UFS, CombinePaths(SC.ProjectRoot, "Config"), "*", true, null, CombinePaths(SC.RelativeProjectRootForStage, "Config"), true, !Params.Pak);
-			SC.StageFiles(StagedFileType.UFS, CombinePaths(SC.ProjectRoot, "Plugins"), "*.uplugin", true, null, null, true, !Params.Pak);
-
-			SC.StageFiles(StagedFileType.UFS, CombinePaths(SC.ProjectRoot, "Content/Localization/Engine"), "*.locres", true, null, CombinePaths(SC.RelativeProjectRootForStage, "Content/Localization/Engine"), true, !Params.Pak);
-
-			// Stage any additional UFS and NonUFS paths specified in the project ini files; these dirs are relative to the game content directory
-			ConfigCacheIni PlatformGameConfig = null;
-			if (Params.GameConfigs.TryGetValue(SC.StageTargetPlatform.PlatformType, out PlatformGameConfig))
-			{
-				var ProjectContentRoot = CombinePaths(SC.ProjectRoot, "Content");
-				var StageContentRoot = CombinePaths(SC.RelativeProjectRootForStage, "Content");
-
-				List<string> ExtraUFSDirs;
-				if (PlatformGameConfig.GetArray("/Script/UnrealEd.ProjectPackagingSettings", "DirectoriesToAlwaysStageAsUFS", out ExtraUFSDirs))
-				{
-					// Each string has the format '(Path="TheDirToStage")'
-					foreach (var PathStr in ExtraUFSDirs)
-					{
-						var PathParts = PathStr.Split('"');
-						if (PathParts.Length == 3)
-						{
-							var RelativePath = PathParts[1];
-							SC.StageFiles(StagedFileType.UFS, CombinePaths(ProjectContentRoot, RelativePath), "*", true, null, CombinePaths(StageContentRoot, RelativePath), true, !Params.Pak);
-						}
-					}
-				}
-
-				List<string> ExtraNonUFSDirs;
-				if (PlatformGameConfig.GetArray("/Script/UnrealEd.ProjectPackagingSettings", "DirectoriesToAlwaysStageAsNonUFS", out ExtraNonUFSDirs))
-				{
-					// Each string has the format '(Path="TheDirToStage")'
-					foreach (var PathStr in ExtraNonUFSDirs)
-					{
-						var PathParts = PathStr.Split('"');
-						if (PathParts.Length == 3)
-						{
-							var RelativePath = PathParts[1];
-							SC.StageFiles(StagedFileType.NonUFS, CombinePaths(ProjectContentRoot, RelativePath));
-						}
-					}
-				}
-			}
-
-			StagedFileType StagedFileTypeForMovies = StagedFileType.NonUFS;
-			if (Params.FileServer)
-			{
-				// UFS is required when using a file server
-				StagedFileTypeForMovies = StagedFileType.UFS;
-			}
-
-			if (SC.StageTargetPlatform.StageMovies)
-			{
-				SC.StageFiles(StagedFileTypeForMovies, CombinePaths(SC.ProjectRoot, "Content/Movies"), "*", true, null, CombinePaths(SC.RelativeProjectRootForStage, "Content/Movies"), true, !Params.Pak);
-			}
-
-			// eliminate the sand box
-			SC.StageFiles(StagedFileType.UFS, CombinePaths(SC.ProjectRoot, "Saved", "Cooked", SC.CookPlatform), "*", true, null, "", true, !Params.Pak);
-
-			// CrashReportClient is a standalone slate app that does not look in the generated pak file, so it needs the Content/Slate and Shaders/StandaloneRenderer folders Non-UFS
-			// @todo Make CrashReportClient more portable so we don't have to do this
-			if (SC.bStageCrashReporter && !SC.DedicatedServer)
-			{
+            SC.StageFiles(StagedFileType.UFS, CombinePaths(SC.LocalRoot, "Engine/Plugins"), "*.uplugin", true, null, null, true, !Params.Pak);
+
+            // Game ufs (content)
+            SC.StageFiles(StagedFileType.UFS, CombinePaths(SC.ProjectRoot), "*.uproject", false, null, CombinePaths(SC.RelativeProjectRootForStage), true, !Params.Pak);
+            SC.StageFiles(StagedFileType.UFS, CombinePaths(SC.ProjectRoot, "Config"), "*", true, null, CombinePaths(SC.RelativeProjectRootForStage, "Config"), true, !Params.Pak);
+            SC.StageFiles(StagedFileType.UFS, CombinePaths(SC.ProjectRoot, "Plugins"), "*.uplugin", true, null, null, true, !Params.Pak);
+
+            SC.StageFiles(StagedFileType.UFS, CombinePaths(SC.ProjectRoot, "Content/Localization/Engine"), "*.locres", true, null, CombinePaths(SC.RelativeProjectRootForStage, "Content/Localization/Engine"), true, !Params.Pak);
+
+            // Stage any additional UFS and NonUFS paths specified in the project ini files; these dirs are relative to the game content directory
+            ConfigCacheIni PlatformGameConfig = null;
+            if (Params.GameConfigs.TryGetValue(SC.StageTargetPlatform.PlatformType, out PlatformGameConfig))
+            {
+                var ProjectContentRoot = CombinePaths(SC.ProjectRoot, "Content");
+                var StageContentRoot = CombinePaths(SC.RelativeProjectRootForStage, "Content");
+
+                List<string> ExtraUFSDirs;
+                if (PlatformGameConfig.GetArray("/Script/UnrealEd.ProjectPackagingSettings", "DirectoriesToAlwaysStageAsUFS", out ExtraUFSDirs))
+                {
+                    // Each string has the format '(Path="TheDirToStage")'
+                    foreach (var PathStr in ExtraUFSDirs)
+                    {
+                        var PathParts = PathStr.Split('"');
+                        if (PathParts.Length == 3)
+                        {
+                            var RelativePath = PathParts[1];
+                            SC.StageFiles(StagedFileType.UFS, CombinePaths(ProjectContentRoot, RelativePath), "*", true, null, CombinePaths(StageContentRoot, RelativePath), true, !Params.Pak);
+                        }
+                    }
+                }
+
+                List<string> ExtraNonUFSDirs;
+                if (PlatformGameConfig.GetArray("/Script/UnrealEd.ProjectPackagingSettings", "DirectoriesToAlwaysStageAsNonUFS", out ExtraNonUFSDirs))
+                {
+                    // Each string has the format '(Path="TheDirToStage")'
+                    foreach (var PathStr in ExtraNonUFSDirs)
+                    {
+                        var PathParts = PathStr.Split('"');
+                        if (PathParts.Length == 3)
+                        {
+                            var RelativePath = PathParts[1];
+                            SC.StageFiles(StagedFileType.NonUFS, CombinePaths(ProjectContentRoot, RelativePath));
+                        }
+                    }
+                }
+            }
+
+            StagedFileType StagedFileTypeForMovies = StagedFileType.NonUFS;
+            if (Params.FileServer)
+            {
+                // UFS is required when using a file server
+                StagedFileTypeForMovies = StagedFileType.UFS;
+            }
+
+            if (SC.StageTargetPlatform.StageMovies)
+            {
+                SC.StageFiles(StagedFileTypeForMovies, CombinePaths(SC.ProjectRoot, "Content/Movies"), "*", true, null, CombinePaths(SC.RelativeProjectRootForStage, "Content/Movies"), true, !Params.Pak);
+            }
+
+            // eliminate the sand box
+            SC.StageFiles(StagedFileType.UFS, CombinePaths(SC.ProjectRoot, "Saved", "Cooked", SC.CookPlatform), "*", true, null, "", true, !Params.Pak);
+
+            // CrashReportClient is a standalone slate app that does not look in the generated pak file, so it needs the Content/Slate and Shaders/StandaloneRenderer folders Non-UFS
+            // @todo Make CrashReportClient more portable so we don't have to do this
+            if (SC.bStageCrashReporter && !SC.DedicatedServer)
+            {
                 //If the .dat file needs to be staged as NonUFS for non-Windows/Linux hosts we need to change the casing as we do with the build properties file above.
-				SC.StageFiles(StagedFileType.NonUFS, CombinePaths(SC.LocalRoot, "Engine/Content/Slate"));
-				SC.StageFiles(StagedFileType.NonUFS, CombinePaths(SC.LocalRoot, "Engine/Shaders/StandaloneRenderer"));
-			}
-		}
-	}
-
-	public static void DumpTargetManifest(Dictionary<string, string> Mapping, string Filename)
-	{
-		if (Mapping.Count > 0)
-		{
-			var Lines = new List<string>();
-			foreach (var Pair in Mapping)
-			{
-				string Dest = Pair.Value;
-
-				Lines.Add(Dest);
-			}
-			WriteAllLines(Filename, Lines.ToArray());
-		}
-	}
-
-	public static void CopyManifestFilesToStageDir(Dictionary<string, string> Mapping, string StageDir, string ManifestName)
-	{
-		string ManifestPath = "";
-		string ManifestFile = "";
-		if (!String.IsNullOrEmpty(ManifestName))
-		{
-			ManifestFile = "Manifest_" + ManifestName + ".txt";
-			ManifestPath = CombinePaths(StageDir, ManifestFile);
-			DeleteFile(ManifestPath);
-		}
-		foreach (var Pair in Mapping)
-		{
-			string Src = Pair.Key;
-			string Dest = CombinePaths(StageDir, Pair.Value);
-			if (Src != Dest)  // special case for things created in the staging directory, like the pak file
-			{
-				CopyFileIncremental(Src, Dest);
-			}
-		}
-		if (!String.IsNullOrEmpty(ManifestPath) && Mapping.Count > 0)
-		{
-			DumpTargetManifest(Mapping, ManifestPath);
-			if (!FileExists(ManifestPath))
-			{
-				throw new AutomationException("Failed to write manifest {0}", ManifestPath);
-			}
-			CopyFile(ManifestPath, CombinePaths(CmdEnv.LogFolder, ManifestFile));
-		}
-	}
-
-	public static void DumpManifest(Dictionary<string, string> Mapping, string Filename)
-	{
-		if (Mapping.Count > 0)
-		{
-			var Lines = new List<string>();
-			foreach (var Pair in Mapping)
-			{
-				string Src = Pair.Key;
-				string Dest = Pair.Value;
-
-				Lines.Add("\"" + Src + "\" \"" + Dest + "\"");
-			}
-			WriteAllLines(Filename, Lines.ToArray());
-		}
-	}
-
-	public static void DumpManifest(DeploymentContext SC, string BaseFilename, bool DumpUFSFiles = true)
-	{
-		DumpManifest(SC.NonUFSStagingFiles, BaseFilename + "_NonUFSFiles.txt");
-		if (DumpUFSFiles)
-		{
-			DumpManifest(SC.NonUFSStagingFilesDebug, BaseFilename + "_NonUFSFilesDebug.txt");
-		}
-		DumpManifest(SC.UFSStagingFiles, BaseFilename + "_UFSFiles.txt");
-	}
-
-	public static void CopyUsingStagingManifest(ProjectParams Params, DeploymentContext SC)
-	{
-		CopyManifestFilesToStageDir(SC.NonUFSStagingFiles, SC.StageDirectory, "NonUFSFiles");
-
-		if (!Params.NoDebugInfo)
-		{
-			CopyManifestFilesToStageDir(SC.NonUFSStagingFilesDebug, SC.StageDirectory, "DebugFiles");
-		}
-
-		bool bStageUnrealFileSystemFiles = !Params.CookOnTheFly && !Params.Pak && !Params.FileServer;
-		if (bStageUnrealFileSystemFiles)
-		{
-			CopyManifestFilesToStageDir(SC.UFSStagingFiles, SC.StageDirectory, "UFSFiles");
-		}
-	}
-
-	/// <summary>
-	/// Creates a pak file using staging context (single manifest)
-	/// </summary>
-	/// <param name="Params"></param>
-	/// <param name="SC"></param>
-	private static void CreatePakUsingStagingManifest(ProjectParams Params, DeploymentContext SC)
-	{
-		Log("Creating pak using staging manifest.");
-
-		DumpManifest(SC, CombinePaths(CmdEnv.LogFolder, "PrePak" + (SC.DedicatedServer ? "_Server" : "")));
-
-		var UnrealPakResponseFile = CreatePakResponseFileFromStagingManifest(SC);
-
-		CreatePak(Params, SC, UnrealPakResponseFile, SC.ShortProjectName);
-	}
-
-	/// <summary>
-	/// Creates a pak response file using stage context
-	/// </summary>
-	/// <param name="SC"></param>
-	/// <returns></returns>
-	private static Dictionary<string, string> CreatePakResponseFileFromStagingManifest(DeploymentContext SC)
-	{
-		var UnrealPakResponseFile = new Dictionary<string, string>(StringComparer.InvariantCultureIgnoreCase);
-		foreach (var Pair in SC.UFSStagingFiles)
-		{
-			string Src = Pair.Key;
-			string Dest = Pair.Value;
-
-			Dest = CombinePaths(PathSeparator.Slash, SC.PakFileInternalRoot, Dest);
-
-			UnrealPakResponseFile.Add(Src, Dest);
-		}
-		return UnrealPakResponseFile;
-	}
-
-	/// <summary>
-	/// Creates a pak file using response file.
-	/// </summary>
-	/// <param name="Params"></param>
-	/// <param name="SC"></param>
-	/// <param name="UnrealPakResponseFile"></param>
-	/// <param name="PakName"></param>
-	private static void CreatePak(ProjectParams Params, DeploymentContext SC, Dictionary<string, string> UnrealPakResponseFile, string PakName)
-	{
-		var OutputRealtiveLocation = CombinePaths(SC.RelativeProjectRootForStage, "Content/Paks/", PakName + "-" + SC.CookPlatform + ".pak");
-		if (SC.StageTargetPlatform.DeployLowerCaseFilenames(true))
-		{
-			OutputRealtiveLocation = OutputRealtiveLocation.ToLower();
-		}
-		OutputRealtiveLocation = SC.StageTargetPlatform.Remap(OutputRealtiveLocation);
-		var OutputLocation = CombinePaths(SC.RuntimeRootDir, OutputRealtiveLocation);
-		// Add input file to controll order of file within the pak
-		var PakOrderFileLocationBase = CombinePaths(SC.ProjectRoot, "Build", SC.CookPlatform, "FileOpenOrder");
-		var PakOrderFileLocation = CombinePaths(PakOrderFileLocationBase, "GameOpenOrder.log");
-		if (!FileExists_NoExceptions(PakOrderFileLocation))
-		{
-			// Use a fall back, it doesn't matter if this file exists or not. GameOpenOrder.log is preferred however
-			PakOrderFileLocation = CombinePaths(PakOrderFileLocationBase, "EditorOpenOrder.log");
-		}
+                SC.StageFiles(StagedFileType.NonUFS, CombinePaths(SC.LocalRoot, "Engine/Content/Slate"));
+                SC.StageFiles(StagedFileType.NonUFS, CombinePaths(SC.LocalRoot, "Engine/Shaders/StandaloneRenderer"));
+            }
+        }
+    }
+
+    public static void DumpTargetManifest(Dictionary<string, string> Mapping, string Filename)
+    {
+        if (Mapping.Count > 0)
+        {
+            var Lines = new List<string>();
+            foreach (var Pair in Mapping)
+            {
+                string Dest = Pair.Value;
+
+                Lines.Add(Dest);
+            }
+            WriteAllLines(Filename, Lines.ToArray());
+        }
+    }
+
+    public static void CopyManifestFilesToStageDir(Dictionary<string, string> Mapping, string StageDir, string ManifestName)
+    {
+        string ManifestPath = "";
+        string ManifestFile = "";
+        if (!String.IsNullOrEmpty(ManifestName))
+        {
+            ManifestFile = "Manifest_" + ManifestName + ".txt";
+            ManifestPath = CombinePaths(StageDir, ManifestFile);
+            DeleteFile(ManifestPath);
+        }
+        foreach (var Pair in Mapping)
+        {
+            string Src = Pair.Key;
+            string Dest = CombinePaths(StageDir, Pair.Value);
+            if (Src != Dest)  // special case for things created in the staging directory, like the pak file
+            {
+                CopyFileIncremental(Src, Dest);
+            }
+        }
+        if (!String.IsNullOrEmpty(ManifestPath) && Mapping.Count > 0)
+        {
+            DumpTargetManifest(Mapping, ManifestPath);
+            if (!FileExists(ManifestPath))
+            {
+                throw new AutomationException("Failed to write manifest {0}", ManifestPath);
+            }
+            CopyFile(ManifestPath, CombinePaths(CmdEnv.LogFolder, ManifestFile));
+        }
+    }
+
+    public static void DumpManifest(Dictionary<string, string> Mapping, string Filename)
+    {
+        if (Mapping.Count > 0)
+        {
+            var Lines = new List<string>();
+            foreach (var Pair in Mapping)
+            {
+                string Src = Pair.Key;
+                string Dest = Pair.Value;
+
+                Lines.Add("\"" + Src + "\" \"" + Dest + "\"");
+            }
+            WriteAllLines(Filename, Lines.ToArray());
+        }
+    }
+
+    public static void DumpManifest(DeploymentContext SC, string BaseFilename, bool DumpUFSFiles = true)
+    {
+        DumpManifest(SC.NonUFSStagingFiles, BaseFilename + "_NonUFSFiles.txt");
+        if (DumpUFSFiles)
+        {
+            DumpManifest(SC.NonUFSStagingFilesDebug, BaseFilename + "_NonUFSFilesDebug.txt");
+        }
+        DumpManifest(SC.UFSStagingFiles, BaseFilename + "_UFSFiles.txt");
+    }
+
+    public static void CopyUsingStagingManifest(ProjectParams Params, DeploymentContext SC)
+    {
+        CopyManifestFilesToStageDir(SC.NonUFSStagingFiles, SC.StageDirectory, "NonUFSFiles");
+
+        if (!Params.NoDebugInfo)
+        {
+            CopyManifestFilesToStageDir(SC.NonUFSStagingFilesDebug, SC.StageDirectory, "DebugFiles");
+        }
+
+        bool bStageUnrealFileSystemFiles = !Params.CookOnTheFly && !Params.Pak && !Params.FileServer;
+        if (bStageUnrealFileSystemFiles)
+        {
+            CopyManifestFilesToStageDir(SC.UFSStagingFiles, SC.StageDirectory, "UFSFiles");
+        }
+    }
+
+    /// <summary>
+    /// Creates a pak file using staging context (single manifest)
+    /// </summary>
+    /// <param name="Params"></param>
+    /// <param name="SC"></param>
+    private static void CreatePakUsingStagingManifest(ProjectParams Params, DeploymentContext SC)
+    {
+        Log("Creating pak using staging manifest.");
+
+        DumpManifest(SC, CombinePaths(CmdEnv.LogFolder, "PrePak" + (SC.DedicatedServer ? "_Server" : "")));
+
+        var UnrealPakResponseFile = CreatePakResponseFileFromStagingManifest(SC);
+
+        CreatePak(Params, SC, UnrealPakResponseFile, SC.ShortProjectName);
+    }
+
+    /// <summary>
+    /// Creates a pak response file using stage context
+    /// </summary>
+    /// <param name="SC"></param>
+    /// <returns></returns>
+    private static Dictionary<string, string> CreatePakResponseFileFromStagingManifest(DeploymentContext SC)
+    {
+        var UnrealPakResponseFile = new Dictionary<string, string>(StringComparer.InvariantCultureIgnoreCase);
+        foreach (var Pair in SC.UFSStagingFiles)
+        {
+            string Src = Pair.Key;
+            string Dest = Pair.Value;
+
+            Dest = CombinePaths(PathSeparator.Slash, SC.PakFileInternalRoot, Dest);
+
+            UnrealPakResponseFile.Add(Src, Dest);
+        }
+        return UnrealPakResponseFile;
+    }
+
+    /// <summary>
+    /// Creates a pak file using response file.
+    /// </summary>
+    /// <param name="Params"></param>
+    /// <param name="SC"></param>
+    /// <param name="UnrealPakResponseFile"></param>
+    /// <param name="PakName"></param>
+    private static void CreatePak(ProjectParams Params, DeploymentContext SC, Dictionary<string, string> UnrealPakResponseFile, string PakName)
+    {
+        var OutputRealtiveLocation = CombinePaths(SC.RelativeProjectRootForStage, "Content/Paks/", PakName + "-" + SC.CookPlatform + ".pak");
+        if (SC.StageTargetPlatform.DeployLowerCaseFilenames(true))
+        {
+            OutputRealtiveLocation = OutputRealtiveLocation.ToLower();
+        }
+        OutputRealtiveLocation = SC.StageTargetPlatform.Remap(OutputRealtiveLocation);
+        var OutputLocation = CombinePaths(SC.RuntimeRootDir, OutputRealtiveLocation);
+        // Add input file to controll order of file within the pak
+        var PakOrderFileLocationBase = CombinePaths(SC.ProjectRoot, "Build", SC.CookPlatform, "FileOpenOrder");
+        var PakOrderFileLocation = CombinePaths(PakOrderFileLocationBase, "GameOpenOrder.log");
+        if (!FileExists_NoExceptions(PakOrderFileLocation))
+        {
+            // Use a fall back, it doesn't matter if this file exists or not. GameOpenOrder.log is preferred however
+            PakOrderFileLocation = CombinePaths(PakOrderFileLocationBase, "EditorOpenOrder.log");
+        }
 
         RunUnrealPak(UnrealPakResponseFile, OutputLocation, Params.SignPak, PakOrderFileLocation, SC.StageTargetPlatform.GetPlatformPakCommandLine(), Params.Compressed);
 
-		// add the pak file as needing deployment and convert to lower case again if needed
-		SC.UFSStagingFiles.Add(OutputLocation, OutputRealtiveLocation);		
-	}
-
-	/// <summary>
-	/// Creates pak files using streaming install chunk manifests.
-	/// </summary>
-	/// <param name="Params"></param>
-	/// <param name="SC"></param>
-	private static void CreatePaksUsingChunkManifests(ProjectParams Params, DeploymentContext SC)
-	{
-		Log("Creating pak using streaming install manifests.");
-		DumpManifest(SC, CombinePaths(CmdEnv.LogFolder, "PrePak" + (SC.DedicatedServer ? "_Server" : "")));
-
-		var TmpPackagingPath = GetTmpPackagingPath(Params, SC);
-		var ChunkListFilename = GetChunkPakManifestListFilename(Params, SC);
-		var ChunkList = ReadAllLines(ChunkListFilename);
-		var ChunkResponseFiles = new HashSet<string>[ChunkList.Length];
-		for (int Index = 0; Index < ChunkList.Length; ++Index)
-		{
-			var ChunkManifestFilename = CombinePaths(TmpPackagingPath, ChunkList[Index]);
-			ChunkResponseFiles[Index] = ReadPakChunkManifest(ChunkManifestFilename);
-		}
-		// We still want to have a list of all files to stage. We will use the chunk manifests
-		// to put the files from staging manigest into the right chunk
-		var StagingManifestResponseFile = CreatePakResponseFileFromStagingManifest(SC);
-		// DefaultChunkIndex assumes 0 is the 'base' chunk
-		const int DefaultChunkIndex = 0;
-		var PakResponseFiles = new Dictionary<string, string>[ChunkList.Length];
-		for (int Index = 0; Index < PakResponseFiles.Length; ++Index)
-		{
-			PakResponseFiles[Index] = new Dictionary<string, string>(StringComparer.InvariantCultureIgnoreCase);
-		}
-		foreach (var StagingFile in StagingManifestResponseFile)
-		{
-			bool bAddedToChunk = false;
-			for (int ChunkIndex = 0; !bAddedToChunk && ChunkIndex < ChunkResponseFiles.Length; ++ChunkIndex)
-			{
-				if (ChunkResponseFiles[ChunkIndex].Contains(StagingFile.Key))
-				{
-					PakResponseFiles[ChunkIndex].Add(StagingFile.Key, StagingFile.Value);
-					bAddedToChunk = true;
-				}
-			}
-			if (!bAddedToChunk)
-			{
-				PakResponseFiles[DefaultChunkIndex].Add(StagingFile.Key, StagingFile.Value);
-			}
-		}
-
-		for (int ChunkIndex = 0; ChunkIndex < PakResponseFiles.Length; ++ChunkIndex)
-		{
-			var ChunkName = Path.GetFileNameWithoutExtension(ChunkList[ChunkIndex]);
-			CreatePak(Params, SC, PakResponseFiles[ChunkIndex], ChunkName);
-		}
-	}
-
-	private static bool DoesChunkPakManifestExist(ProjectParams Params, DeploymentContext SC)
-	{
-		return FileExists_NoExceptions(GetChunkPakManifestListFilename(Params, SC));
-	}
-
-	private static string GetChunkPakManifestListFilename(ProjectParams Params, DeploymentContext SC)
-	{
-		return CombinePaths(GetTmpPackagingPath(Params, SC), "pakchunklist.txt");
-	}
-
-	private static string GetTmpPackagingPath(ProjectParams Params, DeploymentContext SC)
-	{
-		return CombinePaths(Path.GetDirectoryName(Params.RawProjectPath), "Saved", "TmpPackaging", SC.StageTargetPlatform.GetCookPlatform(SC.DedicatedServer, false, Params.CookFlavor));
-	}
-
-	private static bool ShouldCreatePak(ProjectParams Params, DeploymentContext SC)
-	{
-		return (Params.Pak || Params.SignedPak || !String.IsNullOrEmpty(Params.SignPak) || SC.StageTargetPlatform.RequiresPak(Params)) && !Params.SkipPak;
-	}
-
-	private static bool ShouldCreatePak(ProjectParams Params)
-	{
-		return (Params.Pak || Params.SignedPak || !String.IsNullOrEmpty(Params.SignPak) || Params.ClientTargetPlatformInstances[0].RequiresPak(Params)) && !Params.SkipPak;
-	}
-
-	public static void ApplyStagingManifest(ProjectParams Params, DeploymentContext SC)
-	{
-		MaybeConvertToLowerCase(Params, SC);
-		if (SC.Stage && !Params.NoCleanStage && !Params.SkipStage)
-		{
-			DeleteDirectory(SC.StageDirectory);
-		}
-		if (ShouldCreatePak(Params, SC))
-		{
-			if (Params.Manifests && DoesChunkPakManifestExist(Params, SC))
-			{
-				CreatePaksUsingChunkManifests(Params, SC);
-			}
-			else
-			{
-				CreatePakUsingStagingManifest(Params, SC);
-			}
-		}
-		if (!SC.Stage || Params.SkipStage)
-		{
-			return;
-		}
-		DumpManifest(SC, CombinePaths(CmdEnv.LogFolder, "FinalCopy" + (SC.DedicatedServer ? "_Server" : "")), !Params.Pak);
-		CopyUsingStagingManifest(Params, SC);
-
-		var ThisPlatform = SC.StageTargetPlatform;
-		ThisPlatform.PostStagingFileCopy(Params, SC);
-	}
-
-	private static string GetIntermediateCommandlineDir(DeploymentContext SC)
-	{
-		return CombinePaths(SC.LocalRoot, "Engine/Intermediate/UAT", SC.CookPlatform);
-	}
-
-	public static void WriteStageCommandline(string IntermediateCmdLineFile, ProjectParams Params, DeploymentContext SC)
-	{
-		// this file needs to be treated as a UFS file for casing, but NonUFS for being put into the .pak file. 
-		// @todo: Maybe there should be a new category - UFSNotForPak
-		if (SC.StageTargetPlatform.DeployLowerCaseFilenames(true))
-		{
-			IntermediateCmdLineFile = IntermediateCmdLineFile.ToLower();
-		}
-		if (File.Exists(IntermediateCmdLineFile))
-		{
-			File.Delete(IntermediateCmdLineFile);
-		}
-
-		if (!string.IsNullOrEmpty(Params.StageCommandline) || !string.IsNullOrEmpty(Params.RunCommandline))
-		{
-			string FileHostParams = " ";
-			if (Params.CookOnTheFly || Params.FileServer)
-			{
-				FileHostParams += "-filehostip=";
+        // add the pak file as needing deployment and convert to lower case again if needed
+        SC.UFSStagingFiles.Add(OutputLocation, OutputRealtiveLocation);		
+    }
+
+    /// <summary>
+    /// Creates pak files using streaming install chunk manifests.
+    /// </summary>
+    /// <param name="Params"></param>
+    /// <param name="SC"></param>
+    private static void CreatePaksUsingChunkManifests(ProjectParams Params, DeploymentContext SC)
+    {
+        Log("Creating pak using streaming install manifests.");
+        DumpManifest(SC, CombinePaths(CmdEnv.LogFolder, "PrePak" + (SC.DedicatedServer ? "_Server" : "")));
+
+        var TmpPackagingPath = GetTmpPackagingPath(Params, SC);
+        var ChunkListFilename = GetChunkPakManifestListFilename(Params, SC);
+        var ChunkList = ReadAllLines(ChunkListFilename);
+        var ChunkResponseFiles = new HashSet<string>[ChunkList.Length];
+        for (int Index = 0; Index < ChunkList.Length; ++Index)
+        {
+            var ChunkManifestFilename = CombinePaths(TmpPackagingPath, ChunkList[Index]);
+            ChunkResponseFiles[Index] = ReadPakChunkManifest(ChunkManifestFilename);
+        }
+        // We still want to have a list of all files to stage. We will use the chunk manifests
+        // to put the files from staging manigest into the right chunk
+        var StagingManifestResponseFile = CreatePakResponseFileFromStagingManifest(SC);
+        // DefaultChunkIndex assumes 0 is the 'base' chunk
+        const int DefaultChunkIndex = 0;
+        var PakResponseFiles = new Dictionary<string, string>[ChunkList.Length];
+        for (int Index = 0; Index < PakResponseFiles.Length; ++Index)
+        {
+            PakResponseFiles[Index] = new Dictionary<string, string>(StringComparer.InvariantCultureIgnoreCase);
+        }
+        foreach (var StagingFile in StagingManifestResponseFile)
+        {
+            bool bAddedToChunk = false;
+            for (int ChunkIndex = 0; !bAddedToChunk && ChunkIndex < ChunkResponseFiles.Length; ++ChunkIndex)
+            {
+                if (ChunkResponseFiles[ChunkIndex].Contains(StagingFile.Key))
+                {
+                    PakResponseFiles[ChunkIndex].Add(StagingFile.Key, StagingFile.Value);
+                    bAddedToChunk = true;
+                }
+            }
+            if (!bAddedToChunk)
+            {
+                PakResponseFiles[DefaultChunkIndex].Add(StagingFile.Key, StagingFile.Value);
+            }
+        }
+
+        for (int ChunkIndex = 0; ChunkIndex < PakResponseFiles.Length; ++ChunkIndex)
+        {
+            var ChunkName = Path.GetFileNameWithoutExtension(ChunkList[ChunkIndex]);
+            CreatePak(Params, SC, PakResponseFiles[ChunkIndex], ChunkName);
+        }
+                
+        String ChunkLayerFilename = CombinePaths(GetTmpPackagingPath(Params, SC), GetChunkPakLayerListName());        
+        String OutputChunkLayerFilename = Path.Combine(SC.ProjectRoot, "Build", SC.CookPlatform, "ChunkLayerInfo", GetChunkPakLayerListName());
+        Directory.CreateDirectory(Path.GetDirectoryName(OutputChunkLayerFilename));
+        File.Copy(ChunkLayerFilename, OutputChunkLayerFilename, true);
+    }
+
+    private static bool DoesChunkPakManifestExist(ProjectParams Params, DeploymentContext SC)
+    {
+        return FileExists_NoExceptions(GetChunkPakManifestListFilename(Params, SC));
+    }
+
+    private static string GetChunkPakManifestListFilename(ProjectParams Params, DeploymentContext SC)
+    {
+        return CombinePaths(GetTmpPackagingPath(Params, SC), "pakchunklist.txt");
+    }
+
+    private static string GetChunkPakLayerListName()
+    {
+        return "pakchunklayers.txt";
+    }
+
+    private static string GetTmpPackagingPath(ProjectParams Params, DeploymentContext SC)
+    {
+        return CombinePaths(Path.GetDirectoryName(Params.RawProjectPath), "Saved", "TmpPackaging", SC.StageTargetPlatform.GetCookPlatform(SC.DedicatedServer, false, Params.CookFlavor));
+    }
+
+    private static bool ShouldCreatePak(ProjectParams Params, DeploymentContext SC)
+    {
+        return (Params.Pak || Params.SignedPak || !String.IsNullOrEmpty(Params.SignPak) || SC.StageTargetPlatform.RequiresPak(Params)) && !Params.SkipPak;
+    }
+
+    private static bool ShouldCreatePak(ProjectParams Params)
+    {
+        return (Params.Pak || Params.SignedPak || !String.IsNullOrEmpty(Params.SignPak) || Params.ClientTargetPlatformInstances[0].RequiresPak(Params)) && !Params.SkipPak;
+    }
+
+    public static void ApplyStagingManifest(ProjectParams Params, DeploymentContext SC)
+    {
+        MaybeConvertToLowerCase(Params, SC);
+        if (SC.Stage && !Params.NoCleanStage && !Params.SkipStage)
+        {
+            DeleteDirectory(SC.StageDirectory);
+        }
+        if (ShouldCreatePak(Params, SC))
+        {
+            if (Params.Manifests && DoesChunkPakManifestExist(Params, SC))
+            {
+                CreatePaksUsingChunkManifests(Params, SC);
+            }
+            else
+            {
+                CreatePakUsingStagingManifest(Params, SC);
+            }
+        }
+        if (!SC.Stage || Params.SkipStage)
+        {
+            return;
+        }
+        DumpManifest(SC, CombinePaths(CmdEnv.LogFolder, "FinalCopy" + (SC.DedicatedServer ? "_Server" : "")), !Params.Pak);
+        CopyUsingStagingManifest(Params, SC);
+
+        var ThisPlatform = SC.StageTargetPlatform;
+        ThisPlatform.PostStagingFileCopy(Params, SC);
+    }
+
+    private static string GetIntermediateCommandlineDir(DeploymentContext SC)
+    {
+        return CombinePaths(SC.LocalRoot, "Engine/Intermediate/UAT", SC.CookPlatform);
+    }
+
+    public static void WriteStageCommandline(string IntermediateCmdLineFile, ProjectParams Params, DeploymentContext SC)
+    {
+        // this file needs to be treated as a UFS file for casing, but NonUFS for being put into the .pak file. 
+        // @todo: Maybe there should be a new category - UFSNotForPak
+        if (SC.StageTargetPlatform.DeployLowerCaseFilenames(true))
+        {
+            IntermediateCmdLineFile = IntermediateCmdLineFile.ToLower();
+        }
+        if (File.Exists(IntermediateCmdLineFile))
+        {
+            File.Delete(IntermediateCmdLineFile);
+        }
+
+        if (!string.IsNullOrEmpty(Params.StageCommandline) || !string.IsNullOrEmpty(Params.RunCommandline))
+        {
+            string FileHostParams = " ";
+            if (Params.CookOnTheFly || Params.FileServer)
+            {
+                FileHostParams += "-filehostip=";
                 bool FirstParam = true;
-				if (UnrealBuildTool.ExternalExecution.GetRuntimePlatform() == UnrealTargetPlatform.Mac)
-				{
-					NetworkInterface[] Interfaces = NetworkInterface.GetAllNetworkInterfaces();
-					foreach (NetworkInterface adapter in Interfaces)
-					{
-						if (adapter.NetworkInterfaceType != NetworkInterfaceType.Loopback)
-						{
-							IPInterfaceProperties IP = adapter.GetIPProperties();
-							for (int Index = 0; Index < IP.UnicastAddresses.Count; ++Index)
-							{
-								if (IP.UnicastAddresses[Index].IsDnsEligible && IP.UnicastAddresses[Index].Address.AddressFamily == System.Net.Sockets.AddressFamily.InterNetwork)
-								{
-									if (Params.Port != null)
-									{
+                if (UnrealBuildTool.ExternalExecution.GetRuntimePlatform() == UnrealTargetPlatform.Mac)
+                {
+                    NetworkInterface[] Interfaces = NetworkInterface.GetAllNetworkInterfaces();
+                    foreach (NetworkInterface adapter in Interfaces)
+                    {
+                        if (adapter.NetworkInterfaceType != NetworkInterfaceType.Loopback)
+                        {
+                            IPInterfaceProperties IP = adapter.GetIPProperties();
+                            for (int Index = 0; Index < IP.UnicastAddresses.Count; ++Index)
+                            {
+                                if (IP.UnicastAddresses[Index].IsDnsEligible && IP.UnicastAddresses[Index].Address.AddressFamily == System.Net.Sockets.AddressFamily.InterNetwork)
+                                {
+                                    if (Params.Port != null)
+                                    {
                                         foreach (var Port in Params.Port)
                                         {
                                             if (!FirstParam)
@@ -580,25 +590,25 @@ public partial class Project : CommandUtils
                                             }
                                             
                                         }
-									}
-									
-								}
-							}
-						}
-					}
-				}
-				else
-				{
-					NetworkInterface[] Interfaces = NetworkInterface.GetAllNetworkInterfaces();
-					foreach (NetworkInterface adapter in Interfaces)
-					{
-						if (adapter.OperationalStatus == OperationalStatus.Up)
-						{
-							IPInterfaceProperties IP = adapter.GetIPProperties();
-							for (int Index = 0; Index < IP.UnicastAddresses.Count; ++Index)
-							{
-								if (IP.UnicastAddresses[Index].IsDnsEligible)
-								{
+                                    }
+                                    
+                                }
+                            }
+                        }
+                    }
+                }
+                else
+                {
+                    NetworkInterface[] Interfaces = NetworkInterface.GetAllNetworkInterfaces();
+                    foreach (NetworkInterface adapter in Interfaces)
+                    {
+                        if (adapter.OperationalStatus == OperationalStatus.Up)
+                        {
+                            IPInterfaceProperties IP = adapter.GetIPProperties();
+                            for (int Index = 0; Index < IP.UnicastAddresses.Count; ++Index)
+                            {
+                                if (IP.UnicastAddresses[Index].IsDnsEligible)
+                                {
                                     if (Params.Port != null)
                                     {
                                         foreach (var Port in Params.Port)
@@ -622,12 +632,12 @@ public partial class Project : CommandUtils
                                         }
                                     }
                                     
-								}
-							}
-						}
-					}
-				}
-				const string LocalHost = "127.0.0.1";
+                                }
+                            }
+                        }
+                    }
+                }
+                const string LocalHost = "127.0.0.1";
                 if (Params.Port != null)
                 {
                     foreach (var Port in Params.Port)
@@ -651,91 +661,91 @@ public partial class Project : CommandUtils
 
                     }
                 }
-				FileHostParams += " ";
-			}
-
-			String ProjectFile = String.Format("{0} ", SC.ProjectArgForCommandLines);
-			Directory.CreateDirectory(GetIntermediateCommandlineDir(SC));
-			string CommandLine = String.Format("{0} {1} {2} {3}\n", ProjectFile, Params.StageCommandline.Trim(new char[] { '\"' }), Params.RunCommandline.Trim(new char[] { '\"' }), FileHostParams);
-			File.WriteAllText(IntermediateCmdLineFile, CommandLine);
-		}
-		else if (!Params.IsCodeBasedProject)
-		{
-			String ProjectFile = String.Format("{0} ", SC.ProjectArgForCommandLines);
-			Directory.CreateDirectory(GetIntermediateCommandlineDir(SC));
-			File.WriteAllText(IntermediateCmdLineFile, ProjectFile);
-		}
-	}
-
-	private static void WriteStageCommandline(ProjectParams Params, DeploymentContext SC)
-	{
-		// always delete the existing commandline text file, so it doesn't reuse an old one
-		string IntermediateCmdLineFile = CombinePaths(GetIntermediateCommandlineDir(SC), "UE4CommandLine.txt");
-		WriteStageCommandline(IntermediateCmdLineFile, Params, SC);
-	}
-
-	#endregion
-
-	#region Stage Command
-
-	//@todo move this
-	public static List<DeploymentContext> CreateDeploymentContext(ProjectParams Params, bool InDedicatedServer, bool DoCleanStage = false)
-	{
+                FileHostParams += " ";
+            }
+
+            String ProjectFile = String.Format("{0} ", SC.ProjectArgForCommandLines);
+            Directory.CreateDirectory(GetIntermediateCommandlineDir(SC));
+            string CommandLine = String.Format("{0} {1} {2} {3}\n", ProjectFile, Params.StageCommandline.Trim(new char[] { '\"' }), Params.RunCommandline.Trim(new char[] { '\"' }), FileHostParams);
+            File.WriteAllText(IntermediateCmdLineFile, CommandLine);
+        }
+        else if (!Params.IsCodeBasedProject)
+        {
+            String ProjectFile = String.Format("{0} ", SC.ProjectArgForCommandLines);
+            Directory.CreateDirectory(GetIntermediateCommandlineDir(SC));
+            File.WriteAllText(IntermediateCmdLineFile, ProjectFile);
+        }
+    }
+
+    private static void WriteStageCommandline(ProjectParams Params, DeploymentContext SC)
+    {
+        // always delete the existing commandline text file, so it doesn't reuse an old one
+        string IntermediateCmdLineFile = CombinePaths(GetIntermediateCommandlineDir(SC), "UE4CommandLine.txt");
+        WriteStageCommandline(IntermediateCmdLineFile, Params, SC);
+    }
+
+    #endregion
+
+    #region Stage Command
+
+    //@todo move this
+    public static List<DeploymentContext> CreateDeploymentContext(ProjectParams Params, bool InDedicatedServer, bool DoCleanStage = false)
+    {
         ParamList<string> ListToProcess = InDedicatedServer && (Params.Cook || Params.CookOnTheFly) ? Params.ServerCookedTargets : Params.ClientCookedTargets;
         var ConfigsToProcess = InDedicatedServer && (Params.Cook || Params.CookOnTheFly) ? Params.ServerConfigsToBuild : Params.ClientConfigsToBuild;
-		
-		List<UnrealTargetPlatform> PlatformsToStage = Params.ClientTargetPlatforms;
-		if ( InDedicatedServer && (Params.Cook || Params.CookOnTheFly) )
-		{
-			PlatformsToStage = Params.ServerTargetPlatforms;
-		}
-
-		List<DeploymentContext> DeploymentContexts = new List<DeploymentContext>();
-		foreach ( var StagePlatform in PlatformsToStage )
-		{
-			List<string> ExecutablesToStage = new List<string>();
-
-			string PlatformName = StagePlatform.ToString();
-			foreach (var Target in ListToProcess)
-			{
-				foreach (var Config in ConfigsToProcess)
-				{
-					string Exe = Target;
-					if (Config != UnrealTargetConfiguration.Development)
-					{
-						Exe = Target + "-" + PlatformName + "-" + Config.ToString();
-					}
-					ExecutablesToStage.Add(Exe);
-				}
-			}
-
-			//@todo should pull StageExecutables from somewhere else if not cooked
-			var SC = new DeploymentContext(Params.RawProjectPath, CmdEnv.LocalRoot,
-				(Params.Stage || !String.IsNullOrEmpty(Params.StageDirectoryParam)) ? Params.BaseStageDirectory : "",
-				(Params.Archive || !String.IsNullOrEmpty(Params.ArchiveDirectoryParam)) ? Params.BaseArchiveDirectory : "",
-				Params.CookFlavor,
-				Params.GetTargetPlatformInstance(StagePlatform),
-				ConfigsToProcess,
-				ExecutablesToStage,
-				InDedicatedServer, 
-				Params.Cook || Params.CookOnTheFly, 
-				Params.CrashReporter, 
-				Params.Stage, 
-				Params.CookOnTheFly, 
-				Params.Archive, 
-				Params.IsProgramTarget,
-				Params.HasDedicatedServerAndClient
-				);
-			LogDeploymentContext(SC);
-
-			DeploymentContexts.Add(SC);
-		}
-
-		return DeploymentContexts;
-	}
-
-	public static void CopyBuildToStagingDirectory(ProjectParams Params)
-	{
+        
+        List<UnrealTargetPlatform> PlatformsToStage = Params.ClientTargetPlatforms;
+        if ( InDedicatedServer && (Params.Cook || Params.CookOnTheFly) )
+        {
+            PlatformsToStage = Params.ServerTargetPlatforms;
+        }
+
+        List<DeploymentContext> DeploymentContexts = new List<DeploymentContext>();
+        foreach ( var StagePlatform in PlatformsToStage )
+        {
+            List<string> ExecutablesToStage = new List<string>();
+
+            string PlatformName = StagePlatform.ToString();
+            foreach (var Target in ListToProcess)
+            {
+                foreach (var Config in ConfigsToProcess)
+                {
+                    string Exe = Target;
+                    if (Config != UnrealTargetConfiguration.Development)
+                    {
+                        Exe = Target + "-" + PlatformName + "-" + Config.ToString();
+                    }
+                    ExecutablesToStage.Add(Exe);
+                }
+            }
+
+            //@todo should pull StageExecutables from somewhere else if not cooked
+            var SC = new DeploymentContext(Params.RawProjectPath, CmdEnv.LocalRoot,
+                (Params.Stage || !String.IsNullOrEmpty(Params.StageDirectoryParam)) ? Params.BaseStageDirectory : "",
+                (Params.Archive || !String.IsNullOrEmpty(Params.ArchiveDirectoryParam)) ? Params.BaseArchiveDirectory : "",
+                Params.CookFlavor,
+                Params.GetTargetPlatformInstance(StagePlatform),
+                ConfigsToProcess,
+                ExecutablesToStage,
+                InDedicatedServer, 
+                Params.Cook || Params.CookOnTheFly, 
+                Params.CrashReporter, 
+                Params.Stage, 
+                Params.CookOnTheFly, 
+                Params.Archive, 
+                Params.IsProgramTarget,
+                Params.HasDedicatedServerAndClient
+                );
+            LogDeploymentContext(SC);
+
+            DeploymentContexts.Add(SC);
+        }
+
+        return DeploymentContexts;
+    }
+
+    public static void CopyBuildToStagingDirectory(ProjectParams Params)
+    {
         if (ShouldCreatePak(Params) || (Params.Stage && !Params.SkipStage))
         {
             Params.ValidateAndLog();
@@ -762,7 +772,7 @@ public partial class Project : CommandUtils
                 }
             }
         }
-	}
+    }
 
-	#endregion
+    #endregion
 }
diff --git a/Engine/Source/Runtime/Engine/Public/GameDelegates.h b/Engine/Source/Runtime/Engine/Public/GameDelegates.h
--- a/Engine/Source/Runtime/Engine/Public/GameDelegates.h
+++ b/Engine/Source/Runtime/Engine/Public/GameDelegates.h
@@ -24,6 +24,10 @@ DECLARE_DELEGATE_OneParam(FCookModificationDelegate, TArray<FString>&);
 // FAssignStreamingChunkDelegate(const FString& PackageToAdd, const FString& LastLoadedMapName, const TArray<int32>& AssetRegistryChunkIDs, const TArray<int32>& ExistingChunkIds, int32& OutChunkIndex);
 DECLARE_DELEGATE_FiveParams(FAssignStreamingChunkDelegate, const FString&, const FString&, const TArray<int32>&, const TArray<int32>&, int32&);
 
+/** Delegate to assign a disc layer to a chunk */
+// FAssignLayerChunkDelegate(const FChunkManifest* ChunkManifest, const FString& Platform, const int32 ChunkIndex, int32& OutChunkLayer);
+DECLARE_DELEGATE_FourParams(FAssignLayerChunkDelegate, const TSet<FString>*, const FString&, const int32, int32&);
+
 /** A delegate for platforms that need extra information to flesh out save data information (name of an icon, for instance) */
 // FExtendedSaveGameInfoDelegate(const TCHAR* SaveName, const EGameDelegates_SaveGame Key, FString& Value); 
 DECLARE_DELEGATE_ThreeParams(FExtendedSaveGameInfoDelegate, const TCHAR*, const EGameDelegates_SaveGame, FString&);
@@ -55,6 +59,7 @@ public:
 	// Implement all delegates declared above
 	DEFINE_GAME_DELEGATE(CookModificationDelegate);
 	DEFINE_GAME_DELEGATE(AssignStreamingChunkDelegate);
+	DEFINE_GAME_DELEGATE(AssignLayerChunkDelegate);
 	DEFINE_GAME_DELEGATE(ExtendedSaveGameInfoDelegate);
 	DEFINE_GAME_DELEGATE(WebServerActionDelegate);	
 };
