fix for finding a plist file in the game directory
first looks for Game-Info.plist
second looks for Info.plist
third looks for any .plist
TTP341222
#ue4
#ios

diff --git a/Engine/Source/Programs/AutomationTool/IOS/IOSPlatform.Automation.cs b/Engine/Source/Programs/AutomationTool/IOS/IOSPlatform.Automation.cs
--- a/Engine/Source/Programs/AutomationTool/IOS/IOSPlatform.Automation.cs
+++ b/Engine/Source/Programs/AutomationTool/IOS/IOSPlatform.Automation.cs
@@ -234,6 +234,19 @@ public class IOSPlatform : Platform
 		{
 			SourcePListFile = CombinePaths(ProjectRoot, "Build", "IOS", ShortProjectName + "-Info.plist");
 		}
+		else if (File.Exists(ProjectRoot + "/Build/IOS/Info.plist"))
+		{
+			SourcePListFile = CombinePaths(ProjectRoot, "Build", "IOS", "Info.plist");
+		}
+		else if (Directory.Exists(ProjectRoot + "/Build/IOS"))
+		{
+			// look for any plist file
+			string[] Plists = Directory.GetFiles(ProjectRoot + "/Build/IOS", "*.plist");
+			if (Plists.Length > 0)
+			{
+				SourcePListFile = Plists[0];
+			}
+		}
 
 		//@TODO: This is writing to the engine directory!
 		string SourcePath = CombinePaths((IsCodeBasedProject ? ProjectRoot : LocalRoot + "\\Engine"), "Intermediate", "IOS");
@@ -450,6 +463,8 @@ public class IOSPlatform : Platform
 
 	private static void CopyFileWithReplacements(string SourceFilename, string DestFilename, Dictionary<string, string> Replacements)
 	{
+		Console.WriteLine("Copying {0} to {1} with replacements...", SourceFilename, DestFilename);
+
 		if (!File.Exists(SourceFilename))
 		{
 			return;
@@ -526,6 +541,19 @@ public class IOSPlatform : Platform
 				{
 					SourcePListFile = CombinePaths(SC.ProjectRoot, "Build", "IOS", SC.ShortProjectName + "-Info.plist");
 				}
+				else if (File.Exists(SC.ProjectRoot + "/Build/IOS/Info.plist"))
+				{
+					SourcePListFile = CombinePaths(SC.ProjectRoot, "Build", "IOS", "Info.plist");
+				}
+				else if (Directory.Exists(SC.ProjectRoot + "/Build/IOS"))
+				{
+					// look for any plist file
+					string[] Plists = Directory.GetFiles(SC.ProjectRoot + "/Build/IOS", "*.plist");
+					if (Plists.Length > 0)
+					{
+						SourcePListFile = Plists[0];
+					}
+				}
 
 				//@TODO: This is writing to the engine directory!
 				string SourcePath = CombinePaths((SC.IsCodeBasedProject ? SC.ProjectRoot : SC.LocalRoot + "\\Engine"), "Intermediate", "IOS");
