build script, Change NonUnity Mac compile to timeindex 8.  Begin logging agent memory requirements.

diff --git a/Engine/Source/Programs/AutomationTool/Scripts/GUBP.Automation.cs b/Engine/Source/Programs/AutomationTool/Scripts/GUBP.Automation.cs
--- a/Engine/Source/Programs/AutomationTool/Scripts/GUBP.Automation.cs
+++ b/Engine/Source/Programs/AutomationTool/Scripts/GUBP.Automation.cs
@@ -2860,7 +2860,21 @@ public class GUBP : BuildCommand
         }
         public override int CISFrequencyQuantumShift(GUBP bp)
         {
-            return base.CISFrequencyQuantumShift(bp) + 2;
+            int Result = base.CISFrequencyQuantumShift(bp) + 2;
+            if(HostPlatform == UnrealTargetPlatform.Mac)
+            {
+                Result += 1;
+            }
+            return Result;
+        }
+        public override int AgentMemoryRequirement(GUBP bp)
+        {
+            int Result = base.AgentMemoryRequirement(bp);
+            if(HostPlatform == UnrealTargetPlatform.Mac)
+            {
+                Result = 32;
+            }
+            return Result;
         }
         public override void DoTest(GUBP bp)
         {
@@ -3600,6 +3614,11 @@ public class GUBP : BuildCommand
             {
                 Agent = "[" + Agent + "]";
             }
+            string MemoryReq = "[" + GUBPNodes[NodeToDo].AgentMemoryRequirement(bp).ToString() + "]";
+            if(MemoryReq == "[0]")
+            {
+                MemoryReq = "";                
+            }
             string FrequencyString = CISFrequencyQuantumShiftString(NodeToDo);
 
             Log("      {0}{1}{2}{3}{4}{5}{6} {7}  {8}",
@@ -3610,6 +3629,7 @@ public class GUBP : BuildCommand
                 GUBPNodes[NodeToDo].TriggerNode() ? " - (TriggerNode)" : "",
                 GUBPNodes[NodeToDo].IsSticky() ? " - (Sticky)" : "",
                 Agent,
+                MemoryReq,
                 EMails,
                 ECProc ? GUBPNodes[NodeToDo].ECProcedure() : ""
                 );
