fix for packaging error
TTP341413
#ue4
#ios

diff --git a/Engine/Source/Programs/AutomationTool/IOS/IOSPlatform.Automation.cs b/Engine/Source/Programs/AutomationTool/IOS/IOSPlatform.Automation.cs
--- a/Engine/Source/Programs/AutomationTool/IOS/IOSPlatform.Automation.cs
+++ b/Engine/Source/Programs/AutomationTool/IOS/IOSPlatform.Automation.cs
@@ -224,7 +224,7 @@ public class IOSPlatform : Platform
 
 				IPPArguments += (cookonthefly ? " -cookonthefly" : "");
 				IPPArguments += " -stagedir \"" + CombinePaths(Params.BaseStageDirectory, "IOS") + "\"";
-				IPPArguments += " -projectdir \"" + Path.GetDirectoryName(Params.RawProjectPath) + "\"";
+				IPPArguments += " -project \"" + Params.RawProjectPath + "\"";
 
 				RunAndLog(CmdEnv, IPPExe, IPPArguments);
 			}
