UAT: When no target platform is explicitly specified in the command line, BuildCookRun will pick the currently running platform instead of always Win64 (so that Mac doesn't try to compile Win64 by default)

diff --git a/Engine/Source/Programs/AutomationTool/ProjectParams.cs b/Engine/Source/Programs/AutomationTool/ProjectParams.cs
--- a/Engine/Source/Programs/AutomationTool/ProjectParams.cs
+++ b/Engine/Source/Programs/AutomationTool/ProjectParams.cs
@@ -164,7 +164,7 @@ namespace AutomationTool
 			}
 			if (CommandUtils.IsNullOrEmpty(TargetPlatforms))
 			{
-				// Revert to single default platform: Win64
+				// Revert to single default platform - the current platform we're running
 				TargetPlatforms = DefaultTargetPlatforms;
 			}
 			return TargetPlatforms;
@@ -363,7 +363,7 @@ namespace AutomationTool
 			}
 
 			// Parse command line params for client platforms "-TargetPlatform=", "-Platform=" and also "-Win64", "-Mac" etc.
-			this.ClientTargetPlatforms = SetupTargetPlatforms(Command, ClientTargetPlatforms, new ParamList<UnrealTargetPlatform>() {UnrealTargetPlatform.Win64}, true, "TargetPlatform", "Platform");
+			this.ClientTargetPlatforms = SetupTargetPlatforms(Command, ClientTargetPlatforms, new ParamList<UnrealTargetPlatform>() {HostPlatform.Current.HostEditorPlatform}, true, "TargetPlatform", "Platform");
 
 			// Parse command line params for server paltforms "-ServerTargetPlatform", "-ServerPlatform"; "-Win64" etc is not allowed here
 			this.ServerTargetPlatforms = SetupTargetPlatforms(Command, ServerTargetPlatforms, this.ClientTargetPlatforms, false, "ServerTargetPlatform", "ServerPlatform");
