Fixing iOS.Automation.csproj on the Mac

diff --git a/Engine/Source/Programs/AutomationTool/IOS/IOS.Automation.csproj b/Engine/Source/Programs/AutomationTool/IOS/IOS.Automation.csproj
--- a/Engine/Source/Programs/AutomationTool/IOS/IOS.Automation.csproj
+++ b/Engine/Source/Programs/AutomationTool/IOS/IOS.Automation.csproj
@@ -90,13 +90,5 @@
       <Private>False</Private>
     </ProjectReference>
   </ItemGroup>
-  <ItemGroup>
-    <EmbeddedResource Include="Properties\Resources.resx">
-      <SubType>Designer</SubType>
-    </EmbeddedResource>
-  </ItemGroup>
-  <ItemGroup>
-    <None Include="Resources\Image1.bmp" />
-  </ItemGroup>
   <Import Project="$(MSBuildToolsPath)\Microsoft.CSharp.targets" />
 </Project>
\ No newline at end of file
