#HTML5

      Correctly launch firefox browser on mac.

diff --git a/Engine/Source/Programs/AutomationTool/HTML5/HTML5Platform.Automation.cs b/Engine/Source/Programs/AutomationTool/HTML5/HTML5Platform.Automation.cs
--- a/Engine/Source/Programs/AutomationTool/HTML5/HTML5Platform.Automation.cs
+++ b/Engine/Source/Programs/AutomationTool/HTML5/HTML5Platform.Automation.cs
@@ -229,13 +229,24 @@ public class HTML5Platform : Platform
 			FirefoxProfileCommand = " ";
 		}
 
-   		// open the webpage
+		// open the webpage
 		// what to do with the commandline!??
-        string HTMLPath = Path.ChangeExtension(ClientApp, "html");
-        if  ( !File.Exists(HTMLPath) )  // its probably a content only game - then it exists in the UE4 directory and not in the game directory. 
-            HTMLPath = Path.Combine(CombinePaths(CmdEnv.LocalRoot, "Engine"), "Binaries", "HTML5", Path.GetFileName(HTMLPath)); 
+		string HTMLPath = Path.ChangeExtension(ClientApp, "html");
+		if  ( !File.Exists(HTMLPath) )  // its probably a content only game - then it exists in the UE4 directory and not in the game directory. 
+			HTMLPath = Path.Combine(CombinePaths(CmdEnv.LocalRoot, "Engine"), "Binaries", "HTML5", Path.GetFileName(HTMLPath)); 
 
-		ProcessResult ClientProcess = Run(FirefoxPath, FirefoxProfileCommand + "\"" + HTMLPath + "\"", null, ClientRunFlags | ERunOptions.NoWaitForExit);
+
+		if (ExternalExecution.GetRuntimePlatform () == UnrealTargetPlatform.Mac) 
+		{
+			FirefoxPath = "/bin/bash"; 
+			FirefoxProfileCommand = " -c  \" open -a Firefox.app  \"" + HTMLPath + "\"   \" "; 
+		} 
+		else if (ExternalExecution.GetRuntimePlatform () == UnrealTargetPlatform.Win64) 
+		{
+			FirefoxProfileCommand = FirefoxProfileCommand + "\"" + HTMLPath + "\" "; 
+		}
+
+		ProcessResult ClientProcess = Run(FirefoxPath, FirefoxProfileCommand, null, ClientRunFlags | ERunOptions.NoWaitForExit);
 
 		return ClientProcess;
 	}
