Fix Rocket build.

diff --git a/Engine/Source/Programs/AutomationTool/Scripts/GUBP.Automation.cs b/Engine/Source/Programs/AutomationTool/Scripts/GUBP.Automation.cs
--- a/Engine/Source/Programs/AutomationTool/Scripts/GUBP.Automation.cs
+++ b/Engine/Source/Programs/AutomationTool/Scripts/GUBP.Automation.cs
@@ -927,7 +927,7 @@ public class GUBP : BuildCommand
                 bool bInternalOnly;
                 bool SeparateNode;
 
-                if (ProgramTarget.Rules.GUBP_AlwaysBuildWithTools(HostPlatform, out bInternalOnly, out SeparateNode) && ProgramTarget.Rules.SupportsPlatform(HostPlatform) && !bInternalOnly && !SeparateNode)
+                if (ProgramTarget.Rules.GUBP_AlwaysBuildWithTools(HostPlatform, GUBP.bBuildRocket, out bInternalOnly, out SeparateNode) && ProgramTarget.Rules.SupportsPlatform(HostPlatform) && !bInternalOnly && !SeparateNode)
                 {
                     foreach (var Plat in ProgramTarget.Rules.GUBP_ToolPlatforms(HostPlatform))
                     {
@@ -1059,7 +1059,7 @@ public class GUBP : BuildCommand
                 bool bInternalOnly;
                 bool SeparateNode;
 
-                if (ProgramTarget.Rules.GUBP_AlwaysBuildWithTools(HostPlatform, out bInternalOnly, out SeparateNode) && ProgramTarget.Rules.SupportsPlatform(HostPlatform) && bInternalOnly && !SeparateNode)
+                if (ProgramTarget.Rules.GUBP_AlwaysBuildWithTools(HostPlatform, GUBP.bBuildRocket, out bInternalOnly, out SeparateNode) && ProgramTarget.Rules.SupportsPlatform(HostPlatform) && bInternalOnly && !SeparateNode)
                 {
                     foreach (var Plat in ProgramTarget.Rules.GUBP_ToolPlatforms(HostPlatform))
                     {
@@ -4778,7 +4778,7 @@ public class GUBP : BuildCommand
                 bool bInternalOnly;
                 bool SeparateNode;
 
-                if (ProgramTarget.Rules.GUBP_AlwaysBuildWithTools(HostPlatform, out bInternalOnly, out SeparateNode) && ProgramTarget.Rules.SupportsPlatform(HostPlatform) && SeparateNode)
+                if (ProgramTarget.Rules.GUBP_AlwaysBuildWithTools(HostPlatform, GUBP.bBuildRocket, out bInternalOnly, out SeparateNode) && ProgramTarget.Rules.SupportsPlatform(HostPlatform) && SeparateNode)
                 {
                     if (bInternalOnly)
                     {
diff --git a/Engine/Source/Programs/CrashReporter/CrashReportClient/CrashReportClient.Target.cs b/Engine/Source/Programs/CrashReporter/CrashReportClient/CrashReportClient.Target.cs
--- a/Engine/Source/Programs/CrashReporter/CrashReportClient/CrashReportClient.Target.cs
+++ b/Engine/Source/Programs/CrashReporter/CrashReportClient/CrashReportClient.Target.cs
@@ -91,7 +91,7 @@ public class CrashReportClientTarget : TargetRules
 
 		OutCPPEnvironmentConfiguration.Definitions.Add( "USE_CHECKS_IN_SHIPPING=1" );
 	}
-	public override bool GUBP_AlwaysBuildWithTools(UnrealTargetPlatform InHostPlatform, out bool bInternalToolOnly, out bool SeparateNode)
+    public override bool GUBP_AlwaysBuildWithTools(UnrealTargetPlatform InHostPlatform, bool bBuildingRocket, out bool bInternalToolOnly, out bool SeparateNode)
 	{
 		bInternalToolOnly = false;
 		SeparateNode = false;
diff --git a/Engine/Source/Programs/CrashReporter/MinidumpDiagnostics/MinidumpDiagnostics.Target.cs b/Engine/Source/Programs/CrashReporter/MinidumpDiagnostics/MinidumpDiagnostics.Target.cs
--- a/Engine/Source/Programs/CrashReporter/MinidumpDiagnostics/MinidumpDiagnostics.Target.cs
+++ b/Engine/Source/Programs/CrashReporter/MinidumpDiagnostics/MinidumpDiagnostics.Target.cs
@@ -60,7 +60,7 @@ public class MinidumpDiagnosticsTarget : TargetRules
 
 		OutCPPEnvironmentConfiguration.Definitions.Add("MINIDUMPDIAGNOSTICS=1");
 	}
-    public override bool GUBP_AlwaysBuildWithTools(UnrealTargetPlatform InHostPlatform, out bool bInternalToolOnly, out bool SeparateNode)
+    public override bool GUBP_AlwaysBuildWithTools(UnrealTargetPlatform InHostPlatform, bool bBuildingRocket, out bool bInternalToolOnly, out bool SeparateNode)
     {
         bInternalToolOnly = true;
         SeparateNode = false;
diff --git a/Engine/Source/Programs/Mac/UE4EditorServices/UE4EditorServices.Target.cs b/Engine/Source/Programs/Mac/UE4EditorServices/UE4EditorServices.Target.cs
--- a/Engine/Source/Programs/Mac/UE4EditorServices/UE4EditorServices.Target.cs
+++ b/Engine/Source/Programs/Mac/UE4EditorServices/UE4EditorServices.Target.cs
@@ -49,7 +49,7 @@ public class UE4EditorServicesTarget : TargetRules
 		UEBuildConfiguration.bCompileAgainstEngine = false;
 		UEBuildConfiguration.bCompileAgainstCoreUObject = false;
 	}
-	public override bool GUBP_AlwaysBuildWithTools(UnrealTargetPlatform InHostPlatform, out bool bInternalToolOnly, out bool SeparateNode)
+    public override bool GUBP_AlwaysBuildWithTools(UnrealTargetPlatform InHostPlatform, bool bBuildingRocket, out bool bInternalToolOnly, out bool SeparateNode)
 	{
 		bInternalToolOnly = false;
 		SeparateNode = false;
diff --git a/Engine/Source/Programs/UnrealBuildTool/System/RulesCompiler.cs b/Engine/Source/Programs/UnrealBuildTool/System/RulesCompiler.cs
--- a/Engine/Source/Programs/UnrealBuildTool/System/RulesCompiler.cs
+++ b/Engine/Source/Programs/UnrealBuildTool/System/RulesCompiler.cs
@@ -605,6 +605,7 @@ namespace UnrealBuildTool
         /// <param name="SeparateNode">If this is set to true, the program will get its own node</param>
         /// </summary>
         /// <returns>true if this target should always be built with the base editor.</returns>
+        [Obsolete]
         public virtual bool GUBP_AlwaysBuildWithTools(UnrealTargetPlatform InHostPlatform, out bool bInternalToolOnly, out bool SeparateNode)
         {
             bInternalToolOnly = false;
@@ -612,6 +613,17 @@ namespace UnrealBuildTool
             return false;
         }
         /// <summary>
+        /// Return true if this target should always be built with the tools. Usually programs like unrealpak.
+        /// <param name="SeparateNode">If this is set to true, the program will get its own node</param>
+        /// </summary>
+        /// <returns>true if this target should always be built with the base editor.</returns>        
+        public virtual bool GUBP_AlwaysBuildWithTools(UnrealTargetPlatform InHostPlatform, bool bBuildingRocket, out bool bInternalToolOnly, out bool SeparateNode)
+        {
+            bInternalToolOnly = false;
+            SeparateNode = false;
+            return false;
+        }
+        /// <summary>
         /// Return a list of platforms to build a tool for
         /// </summary>
         /// <returns>a list of platforms to build a tool for</returns>
diff --git a/Engine/Source/Programs/UnrealFileServer/UnrealFileServer.Target.cs b/Engine/Source/Programs/UnrealFileServer/UnrealFileServer.Target.cs
--- a/Engine/Source/Programs/UnrealFileServer/UnrealFileServer.Target.cs
+++ b/Engine/Source/Programs/UnrealFileServer/UnrealFileServer.Target.cs
@@ -56,7 +56,7 @@ public class UnrealFileServerTarget : TargetRules
 		// UnrealHeaderTool is a console application, not a Windows app (sets entry point to main(), instead of WinMain())
 		OutLinkEnvironmentConfiguration.bIsBuildingConsoleApplication = true;
 	}
-    public override bool GUBP_AlwaysBuildWithTools(UnrealTargetPlatform InHostPlatform, out bool bInternalToolOnly, out bool SeparateNode)
+    public override bool GUBP_AlwaysBuildWithTools(UnrealTargetPlatform InHostPlatform, bool bBuildingRocket, out bool bInternalToolOnly, out bool SeparateNode)
     {
         bInternalToolOnly = false;
         SeparateNode = false;
diff --git a/Engine/Source/Programs/UnrealLightmass/UnrealLightmass.Target.cs b/Engine/Source/Programs/UnrealLightmass/UnrealLightmass.Target.cs
--- a/Engine/Source/Programs/UnrealLightmass/UnrealLightmass.Target.cs
+++ b/Engine/Source/Programs/UnrealLightmass/UnrealLightmass.Target.cs
@@ -75,7 +75,7 @@ public class UnrealLightmassTarget : TargetRules
 		// Disable logging, lightmass will create its own unique logging file
 		OutCPPEnvironmentConfiguration.Definitions.Add("ALLOW_LOG_FILE=0");
 	}
-    public override bool GUBP_AlwaysBuildWithTools(UnrealTargetPlatform InHostPlatform, out bool bInternalToolOnly, out bool SeparateNode)
+    public override bool GUBP_AlwaysBuildWithTools(UnrealTargetPlatform InHostPlatform, bool bBuildingRocket, out bool bInternalToolOnly, out bool SeparateNode)
     {
         bInternalToolOnly = false;
         SeparateNode = false;
diff --git a/Engine/Source/Programs/UnrealPak/UnrealPak.Target.cs b/Engine/Source/Programs/UnrealPak/UnrealPak.Target.cs
--- a/Engine/Source/Programs/UnrealPak/UnrealPak.Target.cs
+++ b/Engine/Source/Programs/UnrealPak/UnrealPak.Target.cs
@@ -51,7 +51,7 @@ public class UnrealPakTarget : TargetRules
 		// UnrealHeaderTool is a console application, not a Windows app (sets entry point to main(), instead of WinMain())
 		OutLinkEnvironmentConfiguration.bIsBuildingConsoleApplication = true;
 	}
-    public override bool GUBP_AlwaysBuildWithTools(UnrealTargetPlatform InHostPlatform, out bool bInternalToolOnly, out bool SeparateNode)
+    public override bool GUBP_AlwaysBuildWithTools(UnrealTargetPlatform InHostPlatform, bool bBuildingRocket, out bool bInternalToolOnly, out bool SeparateNode)
     {
         bInternalToolOnly = false;
         SeparateNode = false;
diff --git a/Engine/Source/Programs/UnrealSync/UnrealSync.Target.cs b/Engine/Source/Programs/UnrealSync/UnrealSync.Target.cs
--- a/Engine/Source/Programs/UnrealSync/UnrealSync.Target.cs
+++ b/Engine/Source/Programs/UnrealSync/UnrealSync.Target.cs
@@ -56,7 +56,7 @@ public class UnrealSyncTarget : TargetRules
 		OutLinkEnvironmentConfiguration.bIsBuildingConsoleApplication = false;
 	}
 
-	public override bool GUBP_AlwaysBuildWithTools(UnrealTargetPlatform InHostPlatform, out bool bInternalToolOnly, out bool SeparateNode)
+    public override bool GUBP_AlwaysBuildWithTools(UnrealTargetPlatform InHostPlatform, bool bBuildingRocket, out bool bInternalToolOnly, out bool SeparateNode)
 	{
 		if (InHostPlatform == UnrealTargetPlatform.Win32 || InHostPlatform == UnrealTargetPlatform.Win64)
 		{
diff --git a/Engine/Source/Programs/UnrealVersionSelector/UnrealVersionSelector.Target.cs b/Engine/Source/Programs/UnrealVersionSelector/UnrealVersionSelector.Target.cs
--- a/Engine/Source/Programs/UnrealVersionSelector/UnrealVersionSelector.Target.cs
+++ b/Engine/Source/Programs/UnrealVersionSelector/UnrealVersionSelector.Target.cs
@@ -52,7 +52,7 @@ public class UnrealVersionSelectorTarget : TargetRules
 		// UnrealHeaderTool is a console application, not a Windows app (sets entry point to main(), instead of WinMain())
 		//OutLinkEnvironmentConfiguration.bIsBuildingConsoleApplication = true;
 	}
-	public override bool GUBP_AlwaysBuildWithTools(UnrealTargetPlatform InHostPlatform, out bool bInternalToolOnly, out bool SeparateNode)
+    public override bool GUBP_AlwaysBuildWithTools(UnrealTargetPlatform InHostPlatform, bool bBuildingRocket, out bool bInternalToolOnly, out bool SeparateNode)
 	{
 		bInternalToolOnly = true;
 		SeparateNode = true;
