build script, separate promotable cooked tests no longer happen after promotable labelling

diff --git a/Engine/Source/Programs/AutomationTool/Scripts/GUBP.Automation.cs b/Engine/Source/Programs/AutomationTool/Scripts/GUBP.Automation.cs
--- a/Engine/Source/Programs/AutomationTool/Scripts/GUBP.Automation.cs
+++ b/Engine/Source/Programs/AutomationTool/Scripts/GUBP.Automation.cs
@@ -2980,42 +2980,34 @@ public class GUBP : BuildCommand
                 AddDependency(GamePlatformCookedAndCompiledNode.StaticGetFullName(HostPlatform, GameProj, Plat));
             }
 
-            var Options = InGameProj.Options(HostPlatform);
-
-            // defer tests until later
-            if (Options.bSeparateGamePromotion)
+            var Options = InGameProj.Options(HostPlatform);           
+           
+            AddPseudodependency(WaitForTestShared.StaticGetFullName());
+            // If the same test fails for the base engine, don't bother trying
+            if (InGameProj.GameName != bp.Branch.BaseEngineProject.GameName)
             {
-                AddPseudodependency(GameLabelPromotableNode.StaticGetFullName(GameProj, false));
-            }
-            else
-            {
-                AddPseudodependency(WaitForTestShared.StaticGetFullName());
-                // If the same test fails for the base engine, don't bother trying
-                if (InGameProj.GameName != bp.Branch.BaseEngineProject.GameName)
+                if (bp.HasNode(UATTestNode.StaticGetFullName(HostPlatform, bp.Branch.BaseEngineProject, TestName)))
                 {
-                    if (bp.HasNode(UATTestNode.StaticGetFullName(HostPlatform, bp.Branch.BaseEngineProject, TestName)))
-                    {
-                        AddPseudodependency(UATTestNode.StaticGetFullName(HostPlatform, bp.Branch.BaseEngineProject, InTestName));
-                    }
-                    else
+                    AddPseudodependency(UATTestNode.StaticGetFullName(HostPlatform, bp.Branch.BaseEngineProject, InTestName));
+                }
+                else
+                {
+                    bool bFoundACook = false;
+                    foreach (var Plat in DependsOnCooked)
                     {
-                        bool bFoundACook = false;
-                        foreach (var Plat in DependsOnCooked)
+                        var PlatTestName = "CookedGameTest_"  + Plat.ToString();
+                        if (bp.HasNode(UATTestNode.StaticGetFullName(HostPlatform, bp.Branch.BaseEngineProject, PlatTestName)))
                         {
-                            var PlatTestName = "CookedGameTest_"  + Plat.ToString();
-                            if (bp.HasNode(UATTestNode.StaticGetFullName(HostPlatform, bp.Branch.BaseEngineProject, PlatTestName)))
-                            {
-                                AddPseudodependency(UATTestNode.StaticGetFullName(HostPlatform, bp.Branch.BaseEngineProject, PlatTestName));
-                                bFoundACook = true;
-                            }
-                        }
-
-                        if (!bFoundACook && bp.HasNode(UATTestNode.StaticGetFullName(HostPlatform, bp.Branch.BaseEngineProject, "EditorTest")))
-                        {
-                            AddPseudodependency(UATTestNode.StaticGetFullName(HostPlatform, bp.Branch.BaseEngineProject, "EditorTest"));
+                            AddPseudodependency(UATTestNode.StaticGetFullName(HostPlatform, bp.Branch.BaseEngineProject, PlatTestName));
+                            bFoundACook = true;
                         }
+                    }
 
+                    if (!bFoundACook && bp.HasNode(UATTestNode.StaticGetFullName(HostPlatform, bp.Branch.BaseEngineProject, "EditorTest")))
+                    {
+                        AddPseudodependency(UATTestNode.StaticGetFullName(HostPlatform, bp.Branch.BaseEngineProject, "EditorTest"));
                     }
+
                 }
             }
 
