build script, fix uat compile error on mac

diff --git a/Engine/Source/Programs/AutomationTool/Scripts/GUBP.Automation.cs b/Engine/Source/Programs/AutomationTool/Scripts/GUBP.Automation.cs
--- a/Engine/Source/Programs/AutomationTool/Scripts/GUBP.Automation.cs
+++ b/Engine/Source/Programs/AutomationTool/Scripts/GUBP.Automation.cs
@@ -2978,9 +2978,7 @@ public class GUBP : BuildCommand
             foreach (var Plat in DependsOnCooked)
             {
                 AddDependency(GamePlatformCookedAndCompiledNode.StaticGetFullName(HostPlatform, GameProj, Plat));
-            }
-
-            var Options = InGameProj.Options(HostPlatform);           
+            }         
            
             AddPseudodependency(WaitForTestShared.StaticGetFullName());
             // If the same test fails for the base engine, don't bother trying
