build script, disable iosonpc test until crash can be fixed.

diff --git a/Engine/Source/Programs/AutomationTool/Scripts/GUBP.Automation.cs b/Engine/Source/Programs/AutomationTool/Scripts/GUBP.Automation.cs
--- a/Engine/Source/Programs/AutomationTool/Scripts/GUBP.Automation.cs
+++ b/Engine/Source/Programs/AutomationTool/Scripts/GUBP.Automation.cs
@@ -5336,7 +5336,7 @@ if (HostPlatform == UnrealTargetPlatform.Mac) continue; //temp hack till mac aut
         {
             if (HasNode(GamePlatformMonolithicsNode.StaticGetFullName(UnrealTargetPlatform.Mac, Branch.BaseEngineProject, UnrealTargetPlatform.IOS)) && HasNode(ToolsNode.StaticGetFullName(UnrealTargetPlatform.Win64)))
             {
-                AddNode(new IOSOnPCTestNode(this));
+                //AddNode(new IOSOnPCTestNode(this)); - Disable IOSOnPCTest until a1011 crash is fixed
             }
             if (!bPreflightBuild)
             {
