UAT: If killing a descendant process failes, log a warning instead of an error.

diff --git a/Engine/Source/Programs/AutomationTool/ProcessUtils.cs b/Engine/Source/Programs/AutomationTool/ProcessUtils.cs
--- a/Engine/Source/Programs/AutomationTool/ProcessUtils.cs
+++ b/Engine/Source/Programs/AutomationTool/ProcessUtils.cs
@@ -508,8 +508,8 @@ namespace AutomationTool
 						}
 						catch (Exception Ex)
 						{
-							CommandUtils.Log(TraceEventType.Error, "Failed to kill descendant:");
-							CommandUtils.Log(TraceEventType.Error, Ex);
+							CommandUtils.Log(TraceEventType.Warning, "Failed to kill descendant:");
+							CommandUtils.Log(TraceEventType.Warning, Ex);
 						}
 						break;  // exit the loop as who knows what else died, so let's get processes anew
 					}
