#UE4 #HTML5

     #TTP 338960 CRITICAL: UFE: CRASH: Quick Launching project for HTML5

      - Fixup some logic change which went in a recent refactor for nfs.
      - Correctly launch firefox with the right http protocol url instead of file:// when cook on the fly is set.
      - Re- enable http nfs for win64.

diff --git a/Engine/Build/HTML5/Game.html.template b/Engine/Build/HTML5/Game.html.template
--- a/Engine/Build/HTML5/Game.html.template
+++ b/Engine/Build/HTML5/Game.html.template
@@ -92,6 +92,20 @@ input[type=button]:hover {
 
     <script type="text/javascript">
 
+// http://stackoverflow.com/questions/4750015/regular-expression-to-find-urls-within-a-string
+function getHTMLGetParam(name){
+   if(name=(new RegExp('[?&]'+encodeURIComponent(name)+'=([^&]*)')).exec(location.search))
+      return decodeURIComponent(name[1]);
+}
+
+var filehostargument = " ' ' "; 
+
+// we are serving via a server and it is unreal file server.
+if ( location.host != "" && getHTMLGetParam("cookonthefly") == "true" )
+{
+    filehostargument = "' -filehostIp=http://" + location.host + " '"; 
+}
+
 var UE4 = {
   get resize_game() {
     var fn = Module.cwrap('resize_game', null, ['number'],['number'] );
@@ -164,7 +178,7 @@ var Module = {
   TOTAL_MEMORY: %HEAPSIZE%,
   noImageDecoding: true,
   noAudioDecoding: true,
-        arguments: [%UE4CMDLINE%],
+        arguments: [ filehostargument, %UE4CMDLINE%],
   print: (function() {
     var element = document.getElementById('output');
     element.value = ''; // clear browser cache
diff --git a/Engine/Source/Programs/AutomationTool/HTML5/HTML5Platform.Automation.cs b/Engine/Source/Programs/AutomationTool/HTML5/HTML5Platform.Automation.cs
--- a/Engine/Source/Programs/AutomationTool/HTML5/HTML5Platform.Automation.cs
+++ b/Engine/Source/Programs/AutomationTool/HTML5/HTML5Platform.Automation.cs
@@ -233,19 +233,48 @@ public class HTML5Platform : Platform
 		// what to do with the commandline!??
 		string HTMLPath = Path.ChangeExtension(ClientApp, "html");
 		if  ( !File.Exists(HTMLPath) )  // its probably a content only game - then it exists in the UE4 directory and not in the game directory. 
-			HTMLPath = Path.Combine(CombinePaths(CmdEnv.LocalRoot, "Engine"), "Binaries", "HTML5", Path.GetFileName(HTMLPath)); 
-
-
-		if (ExternalExecution.GetRuntimePlatform () == UnrealTargetPlatform.Mac) 
-		{
-			FirefoxPath = "/bin/bash"; 
-			FirefoxProfileCommand = " -c  \" open -a Firefox.app  \"" + HTMLPath + "\"   \" "; 
-		} 
-		else if (ExternalExecution.GetRuntimePlatform () == UnrealTargetPlatform.Win64) 
-		{
-			FirefoxProfileCommand = FirefoxProfileCommand + "\"" + HTMLPath + "\" "; 
-		}
-
+			HTMLPath = Path.Combine(CombinePaths(CmdEnv.LocalRoot, "Engine"), "Binaries", "HTML5", Path.GetFileName(HTMLPath));
+
+	    // Are we running via cook on the fly server? 
+        // find our http url - This is awkward because RunClient doesn't have real information that NFS is running or not. 
+        bool IsCookOnTheFly = false; 
+        string[] Arguments = ClientCmdLine.Split(' ');
+        foreach (var Argument in Arguments)
+        {
+            string[] KeyVal = Argument.Split('=');
+            if (KeyVal[0].ToLower().Contains("filehostip"))
+            {
+                // split urls. 
+                string[] Urls =  KeyVal[1].Split('+');
+                foreach (var Url in Urls)
+                {
+                    if (Url.Contains("http"))
+                    {
+                        HTMLPath = Url + "/" + Path.GetFileName(HTMLPath);
+                        IsCookOnTheFly = true; 
+                    }
+                }
+            }
+            if (IsCookOnTheFly)
+                break; 
+        }
+
+        if (IsCookOnTheFly)
+        {
+            HTMLPath += "?cookonthefly=true";
+        }
+
+        if (ExternalExecution.GetRuntimePlatform() == UnrealTargetPlatform.Mac)
+        {
+            FirefoxPath = "/bin/bash";
+            FirefoxProfileCommand = " -c  \" open -a Firefox.app  \"" + HTMLPath + "\"   \" ";
+        }
+        else if (ExternalExecution.GetRuntimePlatform() == UnrealTargetPlatform.Win64)
+        {
+            FirefoxProfileCommand = FirefoxProfileCommand + "\"" + HTMLPath + "\" ";
+        }
+
+        System.Console.WriteLine("Firefox command line: " + FirefoxProfileCommand);
 		ProcessResult ClientProcess = Run(FirefoxPath, FirefoxProfileCommand, null, ClientRunFlags | ERunOptions.NoWaitForExit);
 
 		return ClientProcess;
diff --git a/Engine/Source/Programs/UnrealBuildTool/HTML5/UEBuildHTML5.cs b/Engine/Source/Programs/UnrealBuildTool/HTML5/UEBuildHTML5.cs
--- a/Engine/Source/Programs/UnrealBuildTool/HTML5/UEBuildHTML5.cs
+++ b/Engine/Source/Programs/UnrealBuildTool/HTML5/UEBuildHTML5.cs
@@ -16,7 +16,7 @@ namespace UnrealBuildTool
 
         // Enable and link in code for serving HTTP request via NFS. 
         [XmlConfig]
-        public static bool EnableHTTPForNFS = false;
+        public static bool EnableHTTPForNFS = true;
 
         static private string TargetPlatformName = "HTML5";
 
@@ -376,7 +376,8 @@ namespace UnrealBuildTool
                 {
                     InModule.AddPlatformSpecificDynamicallyLoadedModule("HTML5TargetPlatform");
                 }
-                if (EnableHTTPForNFS)
+
+                if (EnableHTTPForNFS && Target.Platform == UnrealTargetPlatform.Win64 )
                 {
                     if (InModule.ToString() == "NetworkFile") // client
                     {
diff --git a/Engine/Source/Runtime/NetworkFile/Private/HTTPTransport.cpp b/Engine/Source/Runtime/NetworkFile/Private/HTTPTransport.cpp
--- a/Engine/Source/Runtime/NetworkFile/Private/HTTPTransport.cpp
+++ b/Engine/Source/Runtime/NetworkFile/Private/HTTPTransport.cpp
@@ -11,6 +11,7 @@
 
 #if PLATFORM_HTML5_BROWSER
 #include "HTML5JavaScriptFx.h"
+#include <emscripten.h>
 #endif 
 
 
@@ -33,11 +34,8 @@ bool FHTTPTransport::Initialize(const TCHAR* InHostIp)
 		// no port put the default one on
 		HostIp = FString::Printf(TEXT("%s:%d"), *HostIp, (int)(DEFAULT_HTTP_FILE_SERVING_PORT) );
 	}
-	else
-	{
-		// make sure that our string is correctly formated
-		FString::Printf(TEXT("http://%s"),*HostIp);
-	}
+		// make sure that our string is again correctly formated
+	HostIp = FString::Printf(TEXT("http://%s"),*HostIp);
 
 	FCString::Sprintf(Url, *HostIp);
 
@@ -50,6 +48,10 @@ bool FHTTPTransport::Initialize(const TCHAR* InHostIp)
 	HTML5Win32::NFSHttp::Init(TCHAR_TO_ANSI(Url));
 #endif 
 
+#if PLATFORM_HTML5_BROWSER
+	emscripten_log(EM_LOG_CONSOLE , "Unreal File Server URL : %s ", TCHAR_TO_ANSI(Url)); 
+#endif 
+
 	TArray<uint8> In,Out; 
 	bool RetResult = SendPayloadAndReceiveResponse(In,Out); 
 	return RetResult; 
diff --git a/Engine/Source/Runtime/NetworkFileSystem/Private/NetworkFileServerHttp.cpp b/Engine/Source/Runtime/NetworkFileSystem/Private/NetworkFileServerHttp.cpp
--- a/Engine/Source/Runtime/NetworkFileSystem/Private/NetworkFileServerHttp.cpp
+++ b/Engine/Source/Runtime/NetworkFileSystem/Private/NetworkFileServerHttp.cpp
@@ -296,21 +296,17 @@ int FNetworkFileServerHttp::CallBack_HTTP(
 				FString FilePath = FPaths::GameDir() / TEXT("Binaries/HTML5") +  FString((ANSICHAR*)In); 
 				TCHAR *Mime = NULL; 
 
-				if ( FilePath.Contains(".html") )
-				{
-					Mime = TEXT("text/html;charset=UTF-8");
-				}
-				else if ( FilePath.Contains(".js"))
+				
+				if ( FilePath.Contains(".js"))
 				{
 					Mime = TEXT("application/javascript;charset=UTF-8");
 				}
 				else
 				{
-					// bad client. bail. 
-					return -1;
+						Mime = TEXT("text/html;charset=UTF-8");
 				}
 
-				UE_LOG(LogFileServer, Log, TEXT("Serving file %s with mime %s "), *FilePath, (Mime));
+				UE_LOG(LogFileServer, Warning, TEXT("HTTP Serving file %s with mime %s "), *FilePath, (Mime));
 
 				FString AbsoluteFilePath = FPaths::ConvertRelativePathToFull(FilePath);
 				AbsoluteFilePath.ReplaceInline(TEXT("/"),TEXT("\\"));
