build script, PS4 builds now depend on PS4MapFileUtil, which is now a separate node, also add hastests property for ec searches

diff --git a/Engine/Source/Programs/AutomationTool/Scripts/GUBP.Automation.cs b/Engine/Source/Programs/AutomationTool/Scripts/GUBP.Automation.cs
--- a/Engine/Source/Programs/AutomationTool/Scripts/GUBP.Automation.cs
+++ b/Engine/Source/Programs/AutomationTool/Scripts/GUBP.Automation.cs
@@ -105,6 +105,7 @@ public class GUBP : BuildCommand
     public int CL = 0;
     public int TimeIndex = 0;
     public bool bSignBuildProducts = false;
+    public bool bHasTests = false;
     public List<UnrealTargetPlatform> ActivePlatforms = null;
     public BranchInfo Branch = null;
     public bool bOrthogonalizeEditorPlatforms = false;
@@ -1250,6 +1251,15 @@ public class GUBP : BuildCommand
         {
             GameProj = InGameProj;
             TargetPlatform = InTargetPlatform;
+            if (TargetPlatform == UnrealTargetPlatform.PS4)
+            {
+                var PS4MapFileUtil = bp.Branch.FindProgram("PS4MapFileUtil");
+                if(PS4MapFileUtil.Rules == null)
+                {
+                    throw new AutomationException("PS4MapFileUtil is not is this branch, but is required to build PS4 monolithics");
+                }
+                AddDependency(SingleToolsNode.StaticGetFullName(HostPlatform, PS4MapFileUtil));
+            }
             if (InGameProj.GameName != bp.Branch.BaseEngineProject.GameName && GameProj.Properties.Targets.ContainsKey(TargetRules.TargetType.Editor))
             {
                 AddPseudodependency(EditorGameNode.StaticGetFullName(InHostPlatform, GameProj));
@@ -1473,7 +1483,7 @@ public class GUBP : BuildCommand
         {
             TargetPlatform = InTargetPlatform;
             AgentSharingGroup = "UE4_" + InTargetPlatform + "_Mono" + StaticGetHostPlatformSuffix(InHostPlatform);
-            AddDependency(GamePlatformMonolithicsNode.StaticGetFullName(HostPlatform, bp.Branch.BaseEngineProject, TargetPlatform));
+            AddDependency(GamePlatformMonolithicsNode.StaticGetFullName(HostPlatform, bp.Branch.BaseEngineProject, TargetPlatform));            
         }
         public static string StaticGetFullName(UnrealTargetPlatform InHostPlatform, UnrealTargetPlatform InTargetPlatform)
         {
@@ -3049,7 +3059,7 @@ public class GUBP : BuildCommand
                 }
             }
             return Result;
-        }
+        }      
 
         public static string StaticGetFullName(UnrealTargetPlatform InHostPlatform, BranchInfo.BranchUProject InGameProj, string InTestName)
         {
@@ -4665,7 +4675,7 @@ public class GUBP : BuildCommand
         Log("************************* bFake:           		        {0}", bFake);
         Log("************************* bFakeEC:           		        {0}", bFakeEC);
         Log("************************* bHistory:           		        {0}", bHistory);
-        Log("************************* TimeIndex:           		    {0}", TimeIndex);
+        Log("************************* TimeIndex:           		    {0}", TimeIndex);        
         Log("************************* bBuildRocket:           		    {0}", bBuildRocket);
 
 
@@ -4999,6 +5009,7 @@ public class GUBP : BuildCommand
                                 if (!bNoAutomatedTesting)
                                 {
                                     var ThisMonoGameTestNodes = new List<string>();
+                                    
                                     foreach (var Test in GameTests)
                                     {
                                         var TestName = Test.Key + "_" + Plat.ToString();
@@ -5059,6 +5070,7 @@ public class GUBP : BuildCommand
                                     }
                                     if (!bNoAutomatedTesting)
                                     {
+                                        
                                     if (HostPlatform == UnrealTargetPlatform.Mac) continue; //temp hack till mac automated testing works
                                     var ThisMonoGameTestNodes = new List<string>();
                                     foreach (var Test in GameTests)
@@ -5067,7 +5079,7 @@ public class GUBP : BuildCommand
                                         ThisMonoGameTestNodes.Add(AddNode(new UATTestNode(this, HostPlatform, NonCodeProject, TestName, Test.Value, CookedAgentSharingGroup, false, RequiredPlatforms)));
                                     }
                                     if (ThisMonoGameTestNodes.Count > 0)
-                                    {
+                                    {                                        
                                         GameTestNodes.Add(AddNode(new GameAggregateNode(this, HostPlatform, NonCodeProject, "CookedTests_" + Plat.ToString() + "_" + Kind.ToString() + HostPlatformNode.StaticGetHostPlatformSuffix(HostPlatform), ThisMonoGameTestNodes, 0.0f)));
                                     }
                                 }
@@ -5142,7 +5154,7 @@ public class GUBP : BuildCommand
                         EditorTestNodes.Add(AddNode(new UATTestNode(this, HostPlatform, CodeProj, Test.Key, Test.Value, AgentSharingGroup)));
                     }
                     if (EditorTestNodes.Count > 0)
-                    {
+                    {                        
                         AddNode(new GameAggregateNode(this, HostPlatform, CodeProj, "AllEditorTests", EditorTestNodes, 0.0f));
                     }
                 }
@@ -5845,8 +5857,7 @@ if (HostPlatform == UnrealTargetPlatform.Mac) continue; //temp hack till mac aut
             foreach (var NodePair in FullNodeListSortKey)
             {
                 ECProps.Add(string.Format("SortKey/{0}={1}", NodePair.Key, NodePair.Value));
-            }
-
+            }            
             var ECJobProps = new List<string>();
             if (ExplicitTrigger != "")
             {
@@ -5964,6 +5975,10 @@ if (HostPlatform == UnrealTargetPlatform.Mac) continue; //temp hack till mac aut
             {
                 if (GUBPNodes[NodeToDo].RunInEC() && !NodeIsAlreadyComplete(NodeToDo, LocalOnly)) // if something is already finished, we don't put it into EC  
                 {
+                    if ((NodeToDo.Contains("Test")) && !(NodeToDo.Contains("MakeBuild")) && !(NodeToDo.Contains("Unity")))
+                    {
+                        bHasTests = true;
+                    }                    
                     string EMails;
                     var NodeProps = GetECPropsForNode(NodeToDo, CLString, out EMails);
                     ECProps.AddRange(NodeProps);
@@ -6202,6 +6217,7 @@ if (HostPlatform == UnrealTargetPlatform.Mac) continue; //temp hack till mac aut
                     }
                 }
             }
+            ECProps.Add(String.Format("HasTests={0}", bHasTests));
             {
                 ECProps.Add("GUBP_LoadedProps=1");
                 string BranchDefFile = CommandUtils.CombinePaths(CommandUtils.CmdEnv.LogFolder, "BranchDef.properties");
