build script, fix for hastests check

diff --git a/Engine/Source/Programs/AutomationTool/Scripts/GUBP.Automation.cs b/Engine/Source/Programs/AutomationTool/Scripts/GUBP.Automation.cs
--- a/Engine/Source/Programs/AutomationTool/Scripts/GUBP.Automation.cs
+++ b/Engine/Source/Programs/AutomationTool/Scripts/GUBP.Automation.cs
@@ -5975,7 +5975,7 @@ if (HostPlatform == UnrealTargetPlatform.Mac) continue; //temp hack till mac aut
             {
                 if (GUBPNodes[NodeToDo].RunInEC() && !NodeIsAlreadyComplete(NodeToDo, LocalOnly)) // if something is already finished, we don't put it into EC  
                 {
-                    if ((NodeToDo.Contains("Test")) && !(NodeToDo.Contains("MakeBuild")) && !(NodeToDo.Contains("Unity")))
+                    if ((NodeToDo.Contains("Test")) && !(NodeToDo.Contains("MakeBuild")) && !(NodeToDo.Contains("Unity")) && !(NodesToDo.Contains("TestBuild")))
                     {
                         bHasTests = true;
                     }                    
