Fixes bug where AutomationTool's PushDir_NoExceptions achieved different result to PushDir in the success case.

Previously PushDir_NoExceptions would change the current working directory, and push the NEW working directory onto the stack,
thereby losing the ability to pop the final directory and return to the original working directory before anything was pushed.

This old behavior contradicted the method's comment, which reads "Pushes the current working directory onto a stack and sets CWD to a new value."

diff --git a/Engine/Source/Programs/AutomationTool/CommandUtils.cs b/Engine/Source/Programs/AutomationTool/CommandUtils.cs
--- a/Engine/Source/Programs/AutomationTool/CommandUtils.cs
+++ b/Engine/Source/Programs/AutomationTool/CommandUtils.cs
@@ -739,11 +739,12 @@ namespace AutomationTool
 		public static bool PushDir_NoExceptions(string WorkingDirectory)
 		{
 			bool Result = true;
+			string OrigCurrentDirectory = Environment.CurrentDirectory;
 			WorkingDirectory = ConvertSeparators(PathSeparator.Default, WorkingDirectory);
 			try
 			{
 				Environment.CurrentDirectory = WorkingDirectory;
-				WorkingDirectoryStack.Push(Environment.CurrentDirectory);
+				WorkingDirectoryStack.Push(OrigCurrentDirectory);
 			}
 			catch
 			{
