fix for not always copying provisions if we do not find a game provision or UE4Game.mobileprovision
TTP341967
TTP341968

diff --git a/Engine/Source/Programs/AutomationTool/IOS/IOSPlatform.Automation.cs b/Engine/Source/Programs/AutomationTool/IOS/IOSPlatform.Automation.cs
--- a/Engine/Source/Programs/AutomationTool/IOS/IOSPlatform.Automation.cs
+++ b/Engine/Source/Programs/AutomationTool/IOS/IOSPlatform.Automation.cs
@@ -364,56 +364,110 @@ public class IOSPlatform : Platform
 		CopyFileWithReplacements(SourcePListFile, TargetPListFile, Replacements);
 
 		// Now do the .mobileprovision
-		//@TODO: Remove this mobileprovision copy, and move to a library approach like Xcode/codesign does
-		string SourceProvision = CombinePaths(LocalRoot, "Engine", "Build", "IOS", "UE4Game.mobileprovision");
-		string GameSourceProvision = CombinePaths(ProjectRoot, "Build", "IOS", ShortProjectName + ".mobileprovision");
-		if (!File.Exists(GameSourceProvision))
+		// install the provision
+		FileInfo DestFileInfo;
+		string InEngineDir = CombinePaths(LocalRoot, "Engine");
+		string ProvisionWithPrefix = CombinePaths(InEngineDir, "Build", "IOS","UE4Game.mobileprovision");
+		string BuildDirectory = CombinePaths(ProjectRoot, "Build", "IOS");
+		if (File.Exists(BuildDirectory + "/" + ShortProjectName + ".mobileprovision"))
 		{
-			GameSourceProvision = CombinePaths(ProjectRoot, "Build", "IOS", "NotForLicensees", ShortProjectName + ".mobileprovision");
-			if (File.Exists(GameSourceProvision))
+			ProvisionWithPrefix = BuildDirectory + "/" + ShortProjectName + ".mobileprovision";
+		}
+		else
+		{
+			if (File.Exists(BuildDirectory + "/NotForLicensees/" + ShortProjectName + ".mobileprovision"))
 			{
-				SourceProvision = GameSourceProvision;
+				ProvisionWithPrefix = BuildDirectory + "/NotForLicensees/" + ShortProjectName + ".mobileprovision";
 			}
-			else if (!File.Exists(SourceProvision))
+			else if (!File.Exists(ProvisionWithPrefix))
 			{
-				SourceProvision = CombinePaths(LocalRoot, "Engine", "Build", "IOS", "NotForLicensees", "UE4Game.mobileprovision");
+				ProvisionWithPrefix = InEngineDir + "/Build/IOS/NotForLicensees/UE4Game.mobileprovision";
 			}
 		}
-		else
-		{
-			SourceProvision = GameSourceProvision;
-		}
-		if (File.Exists (SourceProvision))
+		if (File.Exists(ProvisionWithPrefix))
 		{
-			Directory.CreateDirectory (Environment.GetEnvironmentVariable ("HOME") + "/Library/MobileDevice/Provisioning Profiles/");
-			File.Copy (SourceProvision, Environment.GetEnvironmentVariable ("HOME") + "/Library/MobileDevice/Provisioning Profiles/" + ShortProjectName + ".mobileprovision", true);
-			FileInfo DestFileInfo = new FileInfo (Environment.GetEnvironmentVariable ("HOME") + "/Library/MobileDevice/Provisioning Profiles/" + ShortProjectName + ".mobileprovision");
+			Directory.CreateDirectory(Environment.GetEnvironmentVariable("HOME") + "/Library/MobileDevice/Provisioning Profiles/");
+			if (File.Exists(Environment.GetEnvironmentVariable("HOME") + "/Library/MobileDevice/Provisioning Profiles/" + ShortProjectName + ".mobileprovision"))
+			{
+				DestFileInfo = new FileInfo(Environment.GetEnvironmentVariable("HOME") + "/Library/MobileDevice/Provisioning Profiles/" + ShortProjectName + ".mobileprovision");
+				DestFileInfo.Attributes = DestFileInfo.Attributes & ~FileAttributes.ReadOnly;
+			}
+			File.Copy(ProvisionWithPrefix, Environment.GetEnvironmentVariable("HOME") + "/Library/MobileDevice/Provisioning Profiles/" + ShortProjectName + ".mobileprovision", true);
+			DestFileInfo = new FileInfo(Environment.GetEnvironmentVariable("HOME") + "/Library/MobileDevice/Provisioning Profiles/" + ShortProjectName + ".mobileprovision");
 			DestFileInfo.Attributes = DestFileInfo.Attributes & ~FileAttributes.ReadOnly;
 		}
-
-		// install the distribution provision
-		SourceProvision = CombinePaths(LocalRoot, "Engine", "Build", "IOS", "UE4Game_Distro.mobileprovision");
-		GameSourceProvision = CombinePaths(ProjectRoot, "Build", "IOS", ShortProjectName + "_Distro.mobileprovision");
-		if (!File.Exists(GameSourceProvision))
+		else
 		{
-			GameSourceProvision = CombinePaths(ProjectRoot, "Build", "IOS", "NotForLicensees", ShortProjectName + "_Distro.mobileprovision");
-			if (File.Exists(GameSourceProvision))
+			// copy all provisions from the game directory, the engine directory, and the notforlicensees directory
+			// copy all of the provisions from the game directory to the library
 			{
-				SourceProvision = GameSourceProvision;
+				if (Directory.Exists(BuildDirectory))
+				{
+					foreach (string Provision in Directory.EnumerateFiles(BuildDirectory, "*.mobileprovision", SearchOption.AllDirectories))
+					{
+						if (!File.Exists(Environment.GetEnvironmentVariable("HOME") + "/Library/MobileDevice/Provisioning Profiles/" + Path.GetFileName(Provision)) || File.GetLastWriteTime(Environment.GetEnvironmentVariable("HOME") + "/Library/MobileDevice/Provisioning Profiles/" + Path.GetFileName(Provision)) < File.GetLastWriteTime(Provision))
+						{
+							if (File.Exists(Environment.GetEnvironmentVariable("HOME") + "/Library/MobileDevice/Provisioning Profiles/" + Path.GetFileName(Provision)))
+							{
+								DestFileInfo = new FileInfo(Environment.GetEnvironmentVariable("HOME") + "/Library/MobileDevice/Provisioning Profiles/" + Path.GetFileName(Provision));
+								DestFileInfo.Attributes = DestFileInfo.Attributes & ~FileAttributes.ReadOnly;
+							}
+							File.Copy(Provision, Environment.GetEnvironmentVariable("HOME") + "/Library/MobileDevice/Provisioning Profiles/" + Path.GetFileName(Provision), true);
+							DestFileInfo = new FileInfo(Environment.GetEnvironmentVariable("HOME") + "/Library/MobileDevice/Provisioning Profiles/" + Path.GetFileName(Provision));
+							DestFileInfo.Attributes = DestFileInfo.Attributes & ~FileAttributes.ReadOnly;
+						}
+					}
+				}
 			}
-			else if (!File.Exists(SourceProvision))
+
+			// copy all of the provisions from the engine directory to the library
 			{
-				SourceProvision = CombinePaths(LocalRoot, "Engine", "Build", "IOS", "NotForLicensees", "UE4Game_Distro.mobileprovision");
+				if (Directory.Exists(InEngineDir + "/Build/IOS"))
+				{
+					foreach (string Provision in Directory.EnumerateFiles(InEngineDir + "/Build/IOS", "*.mobileprovision", SearchOption.AllDirectories))
+					{
+						if (!File.Exists(Environment.GetEnvironmentVariable("HOME") + "/Library/MobileDevice/Provisioning Profiles/" + Path.GetFileName(Provision)) || File.GetLastWriteTime(Environment.GetEnvironmentVariable("HOME") + "/Library/MobileDevice/Provisioning Profiles/" + Path.GetFileName(Provision)) < File.GetLastWriteTime(Provision))
+						{
+							if (File.Exists(Environment.GetEnvironmentVariable("HOME") + "/Library/MobileDevice/Provisioning Profiles/" + Path.GetFileName(Provision)))
+							{
+								DestFileInfo = new FileInfo(Environment.GetEnvironmentVariable("HOME") + "/Library/MobileDevice/Provisioning Profiles/" + Path.GetFileName(Provision));
+								DestFileInfo.Attributes = DestFileInfo.Attributes & ~FileAttributes.ReadOnly;
+							}
+							File.Copy(Provision, Environment.GetEnvironmentVariable("HOME") + "/Library/MobileDevice/Provisioning Profiles/" + Path.GetFileName(Provision), true);
+							DestFileInfo = new FileInfo(Environment.GetEnvironmentVariable("HOME") + "/Library/MobileDevice/Provisioning Profiles/" + Path.GetFileName(Provision));
+							DestFileInfo.Attributes = DestFileInfo.Attributes & ~FileAttributes.ReadOnly;
+						}
+					}
+				}
 			}
 		}
+
+		// install the distribution provision
+		ProvisionWithPrefix = InEngineDir + "/Build/IOS/UE4Game_Distro.mobileprovision";
+		if (File.Exists(BuildDirectory + "/" + ShortProjectName + "_Distro.mobileprovision"))
+		{
+			ProvisionWithPrefix = BuildDirectory + "/" + ShortProjectName + "_Distro.mobileprovision";
+		}
 		else
 		{
-			SourceProvision = GameSourceProvision;
+			if (File.Exists(BuildDirectory + "/NotForLicensees/" + ShortProjectName + "_Distro.mobileprovision"))
+			{
+				ProvisionWithPrefix = BuildDirectory + "/NotForLicensees/" + ShortProjectName + "_Distro.mobileprovision";
+			}
+			else if (!File.Exists(ProvisionWithPrefix))
+			{
+				ProvisionWithPrefix = InEngineDir + "/Build/IOS/NotForLicensees/UE4Game_Distro.mobileprovision";
+			}
 		}
-		if (File.Exists(SourceProvision))
+		if (File.Exists(ProvisionWithPrefix))
 		{
-			File.Copy (SourceProvision, Environment.GetEnvironmentVariable ("HOME") + "/Library/MobileDevice/Provisioning Profiles/" + ShortProjectName + "_Distro.mobileprovision", true);
-			FileInfo DestFileInfo = new FileInfo (Environment.GetEnvironmentVariable ("HOME") + "/Library/MobileDevice/Provisioning Profiles/" + ShortProjectName + "_Distro.mobileprovision");
+			if (File.Exists(Environment.GetEnvironmentVariable("HOME") + "/Library/MobileDevice/Provisioning Profiles/" + ShortProjectName + "_Distro.mobileprovision"))
+			{
+				DestFileInfo = new FileInfo(Environment.GetEnvironmentVariable("HOME") + "/Library/MobileDevice/Provisioning Profiles/" + ShortProjectName + "_Distro.mobileprovision");
+				DestFileInfo.Attributes = DestFileInfo.Attributes & ~FileAttributes.ReadOnly;
+			}
+			File.Copy(ProvisionWithPrefix, Environment.GetEnvironmentVariable("HOME") + "/Library/MobileDevice/Provisioning Profiles/" + ShortProjectName + "_Distro.mobileprovision", true);
+			DestFileInfo = new FileInfo(Environment.GetEnvironmentVariable("HOME") + "/Library/MobileDevice/Provisioning Profiles/" + ShortProjectName + "_Distro.mobileprovision");
 			DestFileInfo.Attributes = DestFileInfo.Attributes & ~FileAttributes.ReadOnly;
 		}
 
diff --git a/Engine/Source/Programs/UnrealBuildTool/IOS/UEDeployIOS.cs b/Engine/Source/Programs/UnrealBuildTool/IOS/UEDeployIOS.cs
--- a/Engine/Source/Programs/UnrealBuildTool/IOS/UEDeployIOS.cs
+++ b/Engine/Source/Programs/UnrealBuildTool/IOS/UEDeployIOS.cs
@@ -111,6 +111,7 @@ namespace UnrealBuildTool.IOS
 			}
 
 			// install the provision
+			FileInfo DestFileInfo;
 			string ProvisionWithPrefix = InEngineDir + "/Build/IOS/UE4Game.mobileprovision";
 			if (File.Exists(BuildDirectory + "/" + InProjectName + ".mobileprovision"))
 			{
@@ -130,10 +131,60 @@ namespace UnrealBuildTool.IOS
 			if (File.Exists (ProvisionWithPrefix))
 			{
 				Directory.CreateDirectory (Environment.GetEnvironmentVariable ("HOME") + "/Library/MobileDevice/Provisioning Profiles/");
+				if (File.Exists(Environment.GetEnvironmentVariable("HOME") + "/Library/MobileDevice/Provisioning Profiles/" + InProjectName + ".mobileprovision"))
+				{
+					DestFileInfo = new FileInfo(Environment.GetEnvironmentVariable("HOME") + "/Library/MobileDevice/Provisioning Profiles/" + InProjectName + ".mobileprovision");
+					DestFileInfo.Attributes = DestFileInfo.Attributes & ~FileAttributes.ReadOnly;
+				}
 				File.Copy (ProvisionWithPrefix, Environment.GetEnvironmentVariable ("HOME") + "/Library/MobileDevice/Provisioning Profiles/" + InProjectName + ".mobileprovision", true);
-				FileInfo DestFileInfo = new FileInfo (Environment.GetEnvironmentVariable ("HOME") + "/Library/MobileDevice/Provisioning Profiles/" + InProjectName + ".mobileprovision");
+				DestFileInfo = new FileInfo (Environment.GetEnvironmentVariable ("HOME") + "/Library/MobileDevice/Provisioning Profiles/" + InProjectName + ".mobileprovision");
 				DestFileInfo.Attributes = DestFileInfo.Attributes & ~FileAttributes.ReadOnly;
 			}
+			else
+			{
+				// copy all provisions from the game directory, the engine directory, and the notforlicensees directory
+				// copy all of the provisions from the game directory to the library
+				{
+					if (Directory.Exists(BuildDirectory))
+					{
+						foreach (string Provision in Directory.EnumerateFiles(BuildDirectory, "*.mobileprovision", SearchOption.AllDirectories))
+						{
+							if (!File.Exists(Environment.GetEnvironmentVariable("HOME") + "/Library/MobileDevice/Provisioning Profiles/" + Path.GetFileName(Provision)) || File.GetLastWriteTime(Environment.GetEnvironmentVariable("HOME") + "/Library/MobileDevice/Provisioning Profiles/" + Path.GetFileName(Provision)) < File.GetLastWriteTime(Provision))
+							{
+								if (File.Exists(Environment.GetEnvironmentVariable("HOME") + "/Library/MobileDevice/Provisioning Profiles/" + Path.GetFileName(Provision)))
+								{
+									DestFileInfo = new FileInfo(Environment.GetEnvironmentVariable("HOME") + "/Library/MobileDevice/Provisioning Profiles/" + Path.GetFileName(Provision));
+									DestFileInfo.Attributes = DestFileInfo.Attributes & ~FileAttributes.ReadOnly;
+								}
+								File.Copy(Provision, Environment.GetEnvironmentVariable("HOME") + "/Library/MobileDevice/Provisioning Profiles/" + Path.GetFileName(Provision), true);
+								DestFileInfo = new FileInfo(Environment.GetEnvironmentVariable("HOME") + "/Library/MobileDevice/Provisioning Profiles/" + Path.GetFileName(Provision));
+								DestFileInfo.Attributes = DestFileInfo.Attributes & ~FileAttributes.ReadOnly;
+							}
+						}
+					}
+				}
+
+				// copy all of the provisions from the engine directory to the library
+				{
+					if (Directory.Exists(InEngineDir + "/Build/IOS"))
+					{
+						foreach (string Provision in Directory.EnumerateFiles(InEngineDir + "/Build/IOS", "*.mobileprovision", SearchOption.AllDirectories))
+						{
+							if (!File.Exists(Environment.GetEnvironmentVariable("HOME") + "/Library/MobileDevice/Provisioning Profiles/" + Path.GetFileName(Provision)) || File.GetLastWriteTime(Environment.GetEnvironmentVariable("HOME") + "/Library/MobileDevice/Provisioning Profiles/" + Path.GetFileName(Provision)) < File.GetLastWriteTime(Provision))
+							{
+								if (File.Exists(Environment.GetEnvironmentVariable("HOME") + "/Library/MobileDevice/Provisioning Profiles/" + Path.GetFileName(Provision)))
+								{
+									DestFileInfo = new FileInfo(Environment.GetEnvironmentVariable("HOME") + "/Library/MobileDevice/Provisioning Profiles/" + Path.GetFileName(Provision));
+									DestFileInfo.Attributes = DestFileInfo.Attributes & ~FileAttributes.ReadOnly;
+								}
+								File.Copy(Provision, Environment.GetEnvironmentVariable("HOME") + "/Library/MobileDevice/Provisioning Profiles/" + Path.GetFileName(Provision), true);
+								DestFileInfo = new FileInfo(Environment.GetEnvironmentVariable("HOME") + "/Library/MobileDevice/Provisioning Profiles/" + Path.GetFileName(Provision));
+								DestFileInfo.Attributes = DestFileInfo.Attributes & ~FileAttributes.ReadOnly;
+							}
+						}
+					}
+				}
+			}
 
 			// install the distribution provision
 			ProvisionWithPrefix = InEngineDir + "/Build/IOS/UE4Game_Distro.mobileprovision";
@@ -154,8 +205,13 @@ namespace UnrealBuildTool.IOS
 			}
 			if (File.Exists(ProvisionWithPrefix))
 			{
-				File.Copy (ProvisionWithPrefix, Environment.GetEnvironmentVariable ("HOME") + "/Library/MobileDevice/Provisioning Profiles/" + InProjectName + "_Distro.mobileprovision", true);
-				FileInfo DestFileInfo = new FileInfo (Environment.GetEnvironmentVariable ("HOME") + "/Library/MobileDevice/Provisioning Profiles/" + InProjectName + "_Distro.mobileprovision");
+				if (File.Exists(Environment.GetEnvironmentVariable("HOME") + "/Library/MobileDevice/Provisioning Profiles/" + InProjectName + "_Distro.mobileprovision"))
+				{
+					DestFileInfo = new FileInfo(Environment.GetEnvironmentVariable("HOME") + "/Library/MobileDevice/Provisioning Profiles/" + InProjectName + "_Distro.mobileprovision");
+					DestFileInfo.Attributes = DestFileInfo.Attributes & ~FileAttributes.ReadOnly;
+				}
+				File.Copy(ProvisionWithPrefix, Environment.GetEnvironmentVariable("HOME") + "/Library/MobileDevice/Provisioning Profiles/" + InProjectName + "_Distro.mobileprovision", true);
+				DestFileInfo = new FileInfo (Environment.GetEnvironmentVariable ("HOME") + "/Library/MobileDevice/Provisioning Profiles/" + InProjectName + "_Distro.mobileprovision");
 				DestFileInfo.Attributes = DestFileInfo.Attributes & ~FileAttributes.ReadOnly;
 			}
 
