Fix inability to compile Android code projects in Rocket. Header intermediate directory was renamed from "Android-armv7" to "Android/-armv7" as a sideeffect of CL 2221537, so headers were not being included in the distribution.

diff --git a/Engine/Source/Programs/AutomationTool/Scripts/GUBP.Automation.cs b/Engine/Source/Programs/AutomationTool/Scripts/GUBP.Automation.cs
--- a/Engine/Source/Programs/AutomationTool/Scripts/GUBP.Automation.cs
+++ b/Engine/Source/Programs/AutomationTool/Scripts/GUBP.Automation.cs
@@ -1510,7 +1510,7 @@ public class GUBP : BuildCommand
 				var PlatformDir = TargetPlatform.ToString();
 				if (TargetPlatform == UnrealTargetPlatform.Android)
 				{
-					PlatformDir = "Android-armv7";
+					PlatformDir = "Android/-armv7";
 				}
 				foreach (var FileToCopy in CommandUtils.FindFiles("*.h", true, CommandUtils.CombinePaths(CmdEnv.LocalRoot, @"Engine\Intermediate\Build\", PlatformDir, "Inc")))
 				{
