build script, shootergame formal builds, fix formal build logic

diff --git a/Engine/Source/Programs/AutomationTool/Scripts/GUBP.Automation.cs b/Engine/Source/Programs/AutomationTool/Scripts/GUBP.Automation.cs
--- a/Engine/Source/Programs/AutomationTool/Scripts/GUBP.Automation.cs
+++ b/Engine/Source/Programs/AutomationTool/Scripts/GUBP.Automation.cs
@@ -5210,37 +5210,49 @@ public class GUBP : BuildCommand
                                     string FormalNodeName = null;
                                     if (Kind == TargetRules.TargetType.Client)
                                     {
-                                        FormalNodeName = AddNode(new FormalBuildNode(this, CodeProj, HostPlatform, InClientTargetPlatforms: new List<UnrealTargetPlatform>() { Plat }, InClientConfigs: new List<UnrealTargetConfiguration>() { Config.TargetConfig }, InClientNotGame : true));
+                                        if (Plat == Config.TargetPlatform)
+                                        {
+                                            FormalNodeName = AddNode(new FormalBuildNode(this, CodeProj, HostPlatform, InClientTargetPlatforms: new List<UnrealTargetPlatform>() { Config.TargetPlatform }, InClientConfigs: new List<UnrealTargetConfiguration>() { Config.TargetConfig }, InClientNotGame: true));
+                                        }
                                     }
                                     else if (Kind == TargetRules.TargetType.Server)
                                     {
-                                        FormalNodeName = AddNode(new FormalBuildNode(this, CodeProj, HostPlatform, InServerTargetPlatforms: new List<UnrealTargetPlatform>() { Plat }, InServerConfigs: new List<UnrealTargetConfiguration>() { Config.TargetConfig }));
+                                        if (Plat == Config.TargetPlatform)
+                                        {
+                                            FormalNodeName = AddNode(new FormalBuildNode(this, CodeProj, HostPlatform, InServerTargetPlatforms: new List<UnrealTargetPlatform>() { Config.TargetPlatform }, InServerConfigs: new List<UnrealTargetConfiguration>() { Config.TargetConfig }));
+                                        }
                                     }
                                     else if (Kind == TargetRules.TargetType.Game)
                                     {
-                                        FormalNodeName = AddNode(new FormalBuildNode(this, CodeProj, HostPlatform, InClientTargetPlatforms: new List<UnrealTargetPlatform>() { Plat }, InClientConfigs: new List<UnrealTargetConfiguration>() { Config.TargetConfig }));
+                                        if (Plat == Config.TargetPlatform)
+                                        {
+                                            FormalNodeName = AddNode(new FormalBuildNode(this, CodeProj, HostPlatform, InClientTargetPlatforms: new List<UnrealTargetPlatform>() { Config.TargetPlatform }, InClientConfigs: new List<UnrealTargetConfiguration>() { Config.TargetConfig }));
+                                        }
                                     }
-                                    // we don't want this delayed
-                                    // this would normally wait for the testing phase, we just want to build it right away
-                                    RemovePseudodependencyFromNode(
-                                        CookNode.StaticGetFullName(HostPlatform, CodeProj, CookedPlatform),
-                                        WaitForTestShared.StaticGetFullName());
-                                    RemovePseudodependencyFromNode(
-                                        CookNode.StaticGetFullName(HostPlatform, CodeProj, CookedPlatform),
-                                        CookNode.StaticGetFullName(HostPlatform, Branch.BaseEngineProject, CookedPlatform));
-
-                                    string BuildAgentSharingGroup = CodeProj.GameName + "_MakeFormalBuild_" + Plat.ToString() + HostPlatformNode.StaticGetHostPlatformSuffix(HostPlatform);
-                                    if (Plat == UnrealTargetPlatform.IOS || Plat == UnrealTargetPlatform.Android) // These trash build products, so we need to use different agents
+                                    if (FormalNodeName != null)
                                     {
-                                        BuildAgentSharingGroup = "";
-                                    }
-                                    GUBPNodes[CookNode.StaticGetFullName(HostPlatform, CodeProj, CookedPlatform)].AgentSharingGroup = BuildAgentSharingGroup;
-                                    GUBPNodes[GamePlatformCookedAndCompiledNode.StaticGetFullName(HostPlatform, CodeProj, Plat)].AgentSharingGroup = BuildAgentSharingGroup;
-                                    GUBPNodes[FormalNodeName].AgentSharingGroup = BuildAgentSharingGroup;
+                                        // we don't want this delayed
+                                        // this would normally wait for the testing phase, we just want to build it right away
+                                        RemovePseudodependencyFromNode(
+                                            CookNode.StaticGetFullName(HostPlatform, CodeProj, CookedPlatform),
+                                            WaitForTestShared.StaticGetFullName());
+                                        RemovePseudodependencyFromNode(
+                                            CookNode.StaticGetFullName(HostPlatform, CodeProj, CookedPlatform),
+                                            CookNode.StaticGetFullName(HostPlatform, Branch.BaseEngineProject, CookedPlatform));
+
+                                        string BuildAgentSharingGroup = CodeProj.GameName + "_MakeFormalBuild_" + Plat.ToString() + HostPlatformNode.StaticGetHostPlatformSuffix(HostPlatform);
+                                        if (Plat == UnrealTargetPlatform.IOS || Plat == UnrealTargetPlatform.Android) // These trash build products, so we need to use different agents
+                                        {
+                                            BuildAgentSharingGroup = "";
+                                        }
+                                        GUBPNodes[CookNode.StaticGetFullName(HostPlatform, CodeProj, CookedPlatform)].AgentSharingGroup = BuildAgentSharingGroup;
+                                        GUBPNodes[GamePlatformCookedAndCompiledNode.StaticGetFullName(HostPlatform, CodeProj, Plat)].AgentSharingGroup = BuildAgentSharingGroup;
+                                        GUBPNodes[FormalNodeName].AgentSharingGroup = BuildAgentSharingGroup;
 
-                                    if (Config.bTest)
-                                    {
-                                        AddNode(new FormalBuildTestNode(this, CodeProj, HostPlatform, Plat, Config.TargetConfig));
+                                        if (Config.bTest)
+                                        {
+                                            AddNode(new FormalBuildTestNode(this, CodeProj, HostPlatform, Plat, Config.TargetConfig));
+                                        }
                                     }
                                 }
                                 if (!bNoAutomatedTesting)
@@ -5975,7 +5987,7 @@ if (HostPlatform == UnrealTargetPlatform.Mac) continue; //temp hack till mac aut
             {
                 if (GUBPNodes[NodeToDo].RunInEC() && !NodeIsAlreadyComplete(NodeToDo, LocalOnly)) // if something is already finished, we don't put it into EC  
                 {
-                    if ((NodeToDo.Contains("Test")) && !(NodeToDo.Contains("MakeBuild")) && !(NodeToDo.Contains("Unity")) && !(NodesToDo.Contains("TestBuild")))
+                    if ((NodeToDo.Contains("Test")) && !(NodeToDo.Contains("MakeBuild")) && !(NodeToDo.Contains("Unity")) && !(NodesToDo.Contains("TestBuild")) && !(NodesToDo.Contains("Compile")))
                     {
                         bHasTests = true;
                     }                    
