build script, xbox one formal builds now have agent requirement

diff --git a/Engine/Source/Programs/AutomationTool/Scripts/GUBP.Automation.cs b/Engine/Source/Programs/AutomationTool/Scripts/GUBP.Automation.cs
--- a/Engine/Source/Programs/AutomationTool/Scripts/GUBP.Automation.cs
+++ b/Engine/Source/Programs/AutomationTool/Scripts/GUBP.Automation.cs
@@ -2613,6 +2613,18 @@ public class GUBP : BuildCommand
         {
             return true;
         }
+        public override string ECAgentString()
+        {
+            string Result = base.ECAgentString();
+            foreach (UnrealTargetPlatform Plat in ClientTargetPlatforms)
+            {
+                if (Plat == UnrealTargetPlatform.XboxOne)
+                {
+                    Result = MergeSpaceStrings(Result, Plat.ToString());
+                }
+            }
+            return Result;
+        }      
         public override float Priority()
         {
             return base.Priority() + 20.0f;
@@ -3614,7 +3626,7 @@ public class GUBP : BuildCommand
             {
                 Agent = "[" + Agent + "]";
             }
-            string MemoryReq = "[" + GUBPNodes[NodeToDo].AgentMemoryRequirement(bp).ToString() + "]";
+            string MemoryReq = "[" + GUBPNodes[NodeToDo].AgentMemoryRequirement(bp).ToString() + "]";            
             if(MemoryReq == "[0]")
             {
                 MemoryReq = "";                
@@ -5207,7 +5219,7 @@ public class GUBP : BuildCommand
 
                                 foreach (var Config in FormalBuildConfigs)
                                 {
-                                    string FormalNodeName = null;
+                                    string FormalNodeName = null;                                    
                                     if (Kind == TargetRules.TargetType.Client)
                                     {
                                         if (Plat == Config.TargetPlatform)
@@ -5241,7 +5253,7 @@ public class GUBP : BuildCommand
                                             CookNode.StaticGetFullName(HostPlatform, Branch.BaseEngineProject, CookedPlatform));
 
                                         string BuildAgentSharingGroup = CodeProj.GameName + "_MakeFormalBuild_" + Plat.ToString() + HostPlatformNode.StaticGetHostPlatformSuffix(HostPlatform);
-                                        if (Plat == UnrealTargetPlatform.IOS || Plat == UnrealTargetPlatform.Android) // These trash build products, so we need to use different agents
+                                        if (Plat == UnrealTargetPlatform.IOS || Plat == UnrealTargetPlatform.Android || Plat == UnrealTargetPlatform.XboxOne) // These trash build products, so we need to use different agents
                                         {
                                             BuildAgentSharingGroup = "";
                                         }
