fix for trying to stage crash reporter files when platform does not support crash reporter
#ue4
#uat

diff --git a/Engine/Source/Programs/AutomationTool/Scripts/CopyBuildToStagingDirectory.Automation.cs b/Engine/Source/Programs/AutomationTool/Scripts/CopyBuildToStagingDirectory.Automation.cs
--- a/Engine/Source/Programs/AutomationTool/Scripts/CopyBuildToStagingDirectory.Automation.cs
+++ b/Engine/Source/Programs/AutomationTool/Scripts/CopyBuildToStagingDirectory.Automation.cs
@@ -253,7 +253,7 @@ public partial class Project : CommandUtils
 
 			// CrashReportClient is a standalone slate app that does not look in the generated pak file, so it needs the Content/Slate and Shaders/StandaloneRenderer folders Non-UFS
 			// @todo Make CrashReportClient more portable so we don't have to do this
-			if (SC.bStageCrashReporter && !SC.DedicatedServer)
+			if (SC.bStageCrashReporter && UnrealBuildTool.UnrealBuildTool.PlatformSupportsCrashReporter(SC.StageTargetPlatform.PlatformType) && !SC.DedicatedServer)
 			{
 				//If the .dat file needs to be staged as NonUFS for non-Windows/Linux hosts we need to change the casing as we do with the build properties file above.
 				SC.StageFiles(StagedFileType.NonUFS, CombinePaths(SC.LocalRoot, "Engine/Content/Slate"));
