Add support for Build Info Service v2 to UAT BuildInfoPublisher.

diff --git a/Engine/Source/Programs/AutomationTool/MCPPublic.cs b/Engine/Source/Programs/AutomationTool/MCPPublic.cs
--- a/Engine/Source/Programs/AutomationTool/MCPPublic.cs
+++ b/Engine/Source/Programs/AutomationTool/MCPPublic.cs
@@ -1,12 +1,14 @@
 // Copyright 1998-2014 Epic Games, Inc. All Rights Reserved.
 using System;
 using System.Collections.Generic;
+using System.Linq;
 using System.Text;
 using System.IO;
 using AutomationTool;
 using System.Runtime.Serialization;
 using System.Net;
 using System.Reflection;
+using System.Text.RegularExpressions;
 using UnrealBuildTool;
 
 namespace EpicGames.MCP.Automation
@@ -143,7 +145,8 @@ namespace EpicGames.MCP.Automation
         {
             get
             {
-                return AppName + BuildVersion + "-" + Platform.ToString() + ".manifest";
+                var BaseFilename = AppName + BuildVersion + "-" + Platform.ToString() + ".manifest";
+                return Regex.Replace(BaseFilename, @"\s+", ""); // Strip out whitespace in order to be compatible with BuildPatchTool
             }
         }
 
@@ -581,6 +584,35 @@ namespace EpicGames.MCP.Config
             CommandUtils.Log("ClientId : {0}", ClientId);
             // we don't really want this in logs CommandUtils.Log("ClientSecret : {0}", ClientSecret);
         }
+
+		/// <summary>
+		/// Returns Base Urls of build info service(s).
+		/// Will only return properties which have been populated.
+		/// </summary>
+		public IEnumerable<string> BuildInfoBaseUrls
+		{
+			get
+			{
+				return new List<string> { BuildInfoBaseUrl, BuildInfoV2BaseUrl }.Where(x => !string.IsNullOrEmpty(x));
+			}
+		}
+
+		/// <summary>
+		/// Returns the build info base URL to use for get requests.
+		/// Switching to version 2 can be achieved by overriding UseV2BuildInfoService in concrete subclasses and setting
+		/// it to true.
+		/// </summary>
+		public string DefaultBuildInfoBaseUrl
+		{
+			get
+			{
+				return UseV2BuildInfoService ? BuildInfoV2BaseUrl : BuildInfoBaseUrl;
+			}
+		}
+		protected virtual bool UseV2BuildInfoService
+		{
+			get { return false; }
+		}
     }
 
     public class McpConfigMapper
