build script, add branch hacker support for removing platforms from branches

diff --git a/Engine/Source/Programs/AutomationTool/Scripts/GUBP.Automation.cs b/Engine/Source/Programs/AutomationTool/Scripts/GUBP.Automation.cs
--- a/Engine/Source/Programs/AutomationTool/Scripts/GUBP.Automation.cs
+++ b/Engine/Source/Programs/AutomationTool/Scripts/GUBP.Automation.cs
@@ -188,8 +188,8 @@ public class GUBP : BuildCommand
     public abstract class GUBPBranchHacker
     {
         public class BranchOptions
-        {
-            public bool NoMac = false;
+        {            
+            public List<UnrealTargetPlatform> PlatformsToRemove = new List<UnrealTargetPlatform>();
         }
         public virtual void ModifyOptions(GUBP bp, ref BranchOptions Options, string Branch)
         {
@@ -340,7 +340,7 @@ public class GUBP : BuildCommand
            Frequency = FrequencyHacker.GetNodeFrequency(bp, Branch, NodeName, BaseFrequency);            
         }
         return Frequency;
-    }
+    }    
     public abstract class GUBPNode
     {
         public List<string> FullNamesOfDependencies = new List<string>();
@@ -4453,7 +4453,7 @@ public class GUBP : BuildCommand
             BranchForOptions = P4Env.BuildRootP4;
         }
         var BranchOptions = GetBranchOptions(BranchForOptions);
-        bool WithMac = !BranchOptions.NoMac;
+        bool WithMac = !BranchOptions.PlatformsToRemove.Contains(UnrealTargetPlatform.Mac);        
         if (ParseParam("NoMac"))
         {
             WithMac = false;
@@ -4742,6 +4742,15 @@ public class GUBP : BuildCommand
             }
             ActivePlatforms = FilteredActivePlatforms;
         }
+        var SupportedPlatforms = new List<UnrealTargetPlatform>();
+        foreach(var Plat in ActivePlatforms)
+        {
+            if(!BranchOptions.PlatformsToRemove.Contains(Plat))
+            {
+                SupportedPlatforms.Add(Plat);
+            }
+        }
+        ActivePlatforms = SupportedPlatforms;
         foreach (var Plat in ActivePlatforms)
         {
             Log("Active Platform: {0}", Plat.ToString());
