Corrected issue with ICU internationalization data and UE4 localization data not being properly staged.

diff --git a/Engine/Source/Programs/AutomationTool/Scripts/CopyBuildToStagingDirectory.Automation.cs b/Engine/Source/Programs/AutomationTool/Scripts/CopyBuildToStagingDirectory.Automation.cs
--- a/Engine/Source/Programs/AutomationTool/Scripts/CopyBuildToStagingDirectory.Automation.cs
+++ b/Engine/Source/Programs/AutomationTool/Scripts/CopyBuildToStagingDirectory.Automation.cs
@@ -1,4 +1,4 @@
-// Copyright 1998-2014 Epic Games, Inc. All Rights Reserved.
+﻿// Copyright 1998-2014 Epic Games, Inc. All Rights Reserved.
 
 using AutomationTool;
 using System;
@@ -189,14 +189,14 @@ public partial class Project : CommandUtils
 
 			}
 			SC.StageFiles(StagedFileType.UFS, CombinePaths(SC.LocalRoot, "Engine/Content/Localization/Engine"), "*.locres", true, null, null, false, !Params.Pak);
+            SC.StageFiles(StagedFileType.UFS, CombinePaths(SC.LocalRoot, "Engine/Content/Localization/ICU"), "*", true, null, null, false, !Params.Pak);
 			SC.StageFiles(StagedFileType.UFS, CombinePaths(SC.LocalRoot, "Engine/Plugins"), "*.uplugin", true, null, null, true, !Params.Pak);
 
 			// Game ufs (content)
 			SC.StageFiles(StagedFileType.UFS, CombinePaths(SC.ProjectRoot), "*.uproject", false, null, CombinePaths(SC.RelativeProjectRootForStage), true, !Params.Pak);
 			SC.StageFiles(StagedFileType.UFS, CombinePaths(SC.ProjectRoot, "Config"), "*", true, null, CombinePaths(SC.RelativeProjectRootForStage, "Config"), true, !Params.Pak);
 			SC.StageFiles(StagedFileType.UFS, CombinePaths(SC.ProjectRoot, "Plugins"), "*.uplugin", true, null, null, true, !Params.Pak);
-
-			SC.StageFiles(StagedFileType.UFS, CombinePaths(SC.ProjectRoot, "Content/Localization/Engine"), "*.locres", true, null, CombinePaths(SC.RelativeProjectRootForStage, "Content/Localization/Engine"), true, !Params.Pak);
+            SC.StageFiles(StagedFileType.UFS, CombinePaths(SC.ProjectRoot, "Content/Localization/Game"), "*.locres", true, null, CombinePaths(SC.RelativeProjectRootForStage, "Content/Localization/Game"), true, !Params.Pak);
 
 			// Stage any additional UFS and NonUFS paths specified in the project ini files; these dirs are relative to the game content directory
 			ConfigCacheIni PlatformGameConfig = null;
@@ -260,7 +260,13 @@ public partial class Project : CommandUtils
 				SC.StageFiles(StagedFileType.NonUFS, CombinePaths(SC.LocalRoot, "Engine/Shaders/StandaloneRenderer"));
 
 				SC.StageFiles( StagedFileType.NonUFS, CombinePaths( SC.LocalRoot, "Engine/Content/Localization/ICU" ) );
-				SC.StageFiles( StagedFileType.NonUFS, CombinePaths( SC.LocalRoot, "Engine/Binaries/ThirdParty/ICU" ) );
+                if (SC.StageTargetPlatform.PlatformType == UnrealTargetPlatform.Win64 ||
+                    SC.StageTargetPlatform.PlatformType == UnrealTargetPlatform.Win32 ||
+                    SC.StageTargetPlatform.PlatformType == UnrealTargetPlatform.Linux ||
+                    SC.StageTargetPlatform.PlatformType == UnrealTargetPlatform.Mac)
+                {
+                    SC.StageFiles(StagedFileType.NonUFS, CombinePaths(SC.LocalRoot, "Engine/Binaries/ThirdParty/ICU"));
+                }
 
 				// SSL libraries are only available for Win64 builds.
 				// @see FPerforceSourceControlProvider::LoadSSLLibraries
