Use a local temp directory for GUBP packaging on Xbox One to avoid issues with remote drive and VHD. Resulting package is then copied into place.

diff --git a/Engine/Source/Programs/AutomationTool/CommandUtils.cs b/Engine/Source/Programs/AutomationTool/CommandUtils.cs
--- a/Engine/Source/Programs/AutomationTool/CommandUtils.cs
+++ b/Engine/Source/Programs/AutomationTool/CommandUtils.cs
@@ -1282,6 +1282,53 @@ namespace AutomationTool
 			}
 			File.SetLastWriteTimeUtc(Dest, File.GetLastWriteTimeUtc(Source));
 		}
+
+		/// <summary>
+		/// Copies a directory and all of it's contents recursively. Does not throw exceptions.
+		/// </summary>
+		/// <param name="Source"></param>
+		/// <param name="Dest"></param>
+		/// <returns>True if the operation was successful, false otherwise.</returns>
+		public static bool CopyDirectory_NoExceptions(string Source, string Dest, bool bQuiet = false)
+		{
+			Source = ConvertSeparators(PathSeparator.Default, Source);
+			Dest = ConvertSeparators(PathSeparator.Default, Dest);
+			Dest.TrimEnd(PathSeparator.Default.ToString().ToCharArray());
+
+			if (InternalUtils.SafeDirectoryExists(Dest))
+			{
+				InternalUtils.SafeDeleteDirectory(Dest, bQuiet);
+				if (InternalUtils.SafeDirectoryExists(Dest, true))
+				{
+					return false;
+				}
+			}
+
+			if (!InternalUtils.SafeCreateDirectory(Dest, bQuiet))
+			{
+				return false;
+			}
+			foreach (var SourceSubDirectory in Directory.GetDirectories(Source))
+			{
+				string DestPath = Dest + PathSeparator.Default.ToString() + GetDirectoryName(SourceSubDirectory);
+				if (!CopyDirectory_NoExceptions(SourceSubDirectory, DestPath, bQuiet))
+				{
+					return false;
+				}
+			}
+			foreach (var SourceFile in Directory.GetFiles(Source))
+			{
+				int FilenameStart = SourceFile.LastIndexOf(PathSeparator.Default.ToString());
+				string DestPath = Dest + SourceFile.Substring(FilenameStart);
+				if (!CopyFile_NoExceptions(SourceFile, DestPath, bQuiet))
+				{
+					return false;
+				}
+			}
+
+			return true;
+		}
+
 		/// <summary>
 		/// Returns directory name without filename. 
 		/// The difference between this and Path.GetDirectoryName is that this
diff --git a/Engine/Source/Programs/AutomationTool/Scripts/GUBP.Automation.cs b/Engine/Source/Programs/AutomationTool/Scripts/GUBP.Automation.cs
--- a/Engine/Source/Programs/AutomationTool/Scripts/GUBP.Automation.cs
+++ b/Engine/Source/Programs/AutomationTool/Scripts/GUBP.Automation.cs
@@ -2655,11 +2655,16 @@ public class GUBP : BuildCommand
             }
             string Args = String.Format("BuildCookRun{0} -SkipBuild -SkipCook -Stage -Pak -Package -NoSubmit", ProjectArg);
 
+			bool bXboxOneTarget = false;
             if (ClientTargetPlatforms != null)
             {
                 bool bFirstClient = true;
                 foreach (var Plat in ClientTargetPlatforms)
                 {
+					if (Plat == UnrealTargetPlatform.XboxOne)
+					{
+						bXboxOneTarget = true;
+					}
                     if (bFirstClient)
                     {
                         bFirstClient = false;
@@ -2698,7 +2703,11 @@ public class GUBP : BuildCommand
                 bool bFirstServer = true;
                 foreach (var Plat in ServerTargetPlatforms)
                 {
-                    if (bFirstServer)
+					if (Plat == UnrealTargetPlatform.XboxOne)
+					{
+						bXboxOneTarget = true;
+					}
+					if (bFirstServer)
                     {
                         bFirstServer = false;
                         Args += String.Format(" -serverplatform={0}", Plat.ToString());
@@ -2723,22 +2732,45 @@ public class GUBP : BuildCommand
                 }
             }
 
+            string FinalArchiveDirectory = "";
+			string IntermediateArchiveDirectory = FinalArchiveDirectory;
             if (P4Enabled)
             {
-                string ArchiveDirectory = GetArchiveDirectory(GameProj, HostPlatform, ClientTargetPlatforms, ClientConfigs, ServerTargetPlatforms, ServerConfigs, ClientNotGame);
-                CleanFormalBuilds(ArchiveDirectory);
-                if (DirectoryExists_NoExceptions(ArchiveDirectory))
+                FinalArchiveDirectory = GetArchiveDirectory(GameProj, HostPlatform, ClientTargetPlatforms, ClientConfigs, ServerTargetPlatforms, ServerConfigs, ClientNotGame);
+				IntermediateArchiveDirectory = FinalArchiveDirectory;
+				// Xbox One packaging does not function with remote file systems. Use a temp local directory to package and then move files into final location.
+				if (bXboxOneTarget)
+				{
+					IntermediateArchiveDirectory = Path.Combine(Path.GetTempPath(), "GUBP.XboxOne");
+					if (DirectoryExists_NoExceptions(IntermediateArchiveDirectory))
+					{
+						DeleteDirectory_NoExceptions(IntermediateArchiveDirectory);
+					}
+					CreateDirectory_NoExceptions(IntermediateArchiveDirectory);
+				}
+				CleanFormalBuilds(FinalArchiveDirectory);
+				if (DirectoryExists_NoExceptions(FinalArchiveDirectory))
                 {
                     if (IsBuildMachine)
                     {
-                        throw new AutomationException("Archive directory already exists {0}", ArchiveDirectory);
+						throw new AutomationException("Archive directory already exists {0}", FinalArchiveDirectory);
                     }
-                    DeleteDirectory_NoExceptions(ArchiveDirectory);
+					DeleteDirectory_NoExceptions(FinalArchiveDirectory);
                 }
-                Args += String.Format(" -Archive -archivedirectory={0}", CommandUtils.MakePathSafeToUseWithCommandLine(ArchiveDirectory));
+				Args += String.Format(" -Archive -archivedirectory={0}", CommandUtils.MakePathSafeToUseWithCommandLine(IntermediateArchiveDirectory));
             }
 
             string LogFile = CommandUtils.RunUAT(CommandUtils.CmdEnv, Args);
+
+            if (P4Enabled)
+            {
+				if (!FinalArchiveDirectory.Equals(IntermediateArchiveDirectory, StringComparison.InvariantCultureIgnoreCase))
+				{
+					CopyDirectory_NoExceptions(IntermediateArchiveDirectory, FinalArchiveDirectory);
+					DeleteDirectory_NoExceptions(IntermediateArchiveDirectory);
+				}
+			}
+
             SaveRecordOfSuccessAndAddToBuildProducts(CommandUtils.ReadAllText(LogFile));
         }
     }
