UAT: Support for duplicate entries when staging files.

diff --git a/Engine/Source/Programs/AutomationTool/DeploymentContext.cs b/Engine/Source/Programs/AutomationTool/DeploymentContext.cs
--- a/Engine/Source/Programs/AutomationTool/DeploymentContext.cs
+++ b/Engine/Source/Programs/AutomationTool/DeploymentContext.cs
@@ -400,15 +400,15 @@ public class DeploymentContext //: ProjectParams
 
 				if (FileType == StagedFileType.UFS)
 				{
-					UFSStagingFiles.Add(FileToCopy, Dest);
+					AddUniqueStagingFile(UFSStagingFiles, FileToCopy, Dest);
 				}
 				else if (FileType == StagedFileType.NonUFS)
 				{
-					NonUFSStagingFiles.Add(FileToCopy, Dest);
+					AddUniqueStagingFile(NonUFSStagingFiles, FileToCopy, Dest);
 				}
 				else if (FileType == StagedFileType.DebugNonUFS)
 				{
-					NonUFSStagingFilesDebug.Add(FileToCopy, Dest);
+					AddUniqueStagingFile(NonUFSStagingFilesDebug, FileToCopy, Dest);
 				}
 				FilesAdded++;
 			}
@@ -423,6 +423,24 @@ public class DeploymentContext //: ProjectParams
 		return FilesAdded;
 	}
 
+	private void AddUniqueStagingFile(Dictionary<string, string> FilesToStage, string FileToCopy, string Dest)
+	{
+		string ExistingDest;
+		if (FilesToStage.TryGetValue(FileToCopy, out ExistingDest))
+		{
+			// If the file already exists, it must have the same dest
+			if (String.Compare(ExistingDest, Dest, true) != 0)
+			{
+				throw new AutomationException("Attempting to add \"{0}\" to staging map but it already exists with a different destination (existing: \"{1}\", new: \"{2}\"",
+					FileToCopy, ExistingDest, Dest);
+			}
+		}
+		else
+		{
+			FilesToStage.Add(FileToCopy, Dest);
+		}
+	}
+
 	public int ArchiveFiles(string InPath, string Wildcard = "*", bool bRecursive = true, string[] ExcludeWildcard = null)
 	{
 		int FilesAdded = 0;
