Ensure we don't double-submit to the same build info service if "v1" and "v2" have identical URLs.

diff --git a/Engine/Source/Programs/AutomationTool/MCPPublic.cs b/Engine/Source/Programs/AutomationTool/MCPPublic.cs
--- a/Engine/Source/Programs/AutomationTool/MCPPublic.cs
+++ b/Engine/Source/Programs/AutomationTool/MCPPublic.cs
@@ -593,7 +593,7 @@ namespace EpicGames.MCP.Config
 		{
 			get
 			{
-				return new List<string> { BuildInfoBaseUrl, BuildInfoV2BaseUrl }.Where(x => !string.IsNullOrEmpty(x));
+				return new List<string> { BuildInfoBaseUrl, BuildInfoV2BaseUrl }.Where(x => !string.IsNullOrEmpty(x)).Distinct();
 			}
 		}
 
