Removed staging of localization data generation configuration files. Removed staging of ICU files from CopyBuildToStagingDirectory.Automation.cs. Individual platforms can elect to copy them.

diff --git a/Engine/Source/Programs/AutomationTool/Scripts/CopyBuildToStagingDirectory.Automation.cs b/Engine/Source/Programs/AutomationTool/Scripts/CopyBuildToStagingDirectory.Automation.cs
--- a/Engine/Source/Programs/AutomationTool/Scripts/CopyBuildToStagingDirectory.Automation.cs
+++ b/Engine/Source/Programs/AutomationTool/Scripts/CopyBuildToStagingDirectory.Automation.cs
@@ -176,7 +176,7 @@ public partial class Project : CommandUtils
 		if (!Params.CookOnTheFly && !Params.SkipCookOnTheFly) // only stage the UFS files if we are not using cook on the fly
 		{
 			// Engine ufs (content)
-			SC.StageFiles(StagedFileType.UFS, CombinePaths(SC.LocalRoot, "Engine/Config"), "*", true, null, null, false, !Params.Pak);
+            SC.StageFiles(StagedFileType.UFS, CombinePaths(SC.LocalRoot, "Engine/Config"), "*", true, new string[] { CombinePaths(SC.LocalRoot, "Engine/Config/Localization") }, null, false, !Params.Pak);
 
 			if (Params.bUsesSlate)
 			{
@@ -189,12 +189,11 @@ public partial class Project : CommandUtils
 
 			}
 			SC.StageFiles(StagedFileType.UFS, CombinePaths(SC.LocalRoot, "Engine/Content/Localization/Engine"), "*.locres", true, null, null, false, !Params.Pak);
-            SC.StageFiles(StagedFileType.UFS, CombinePaths(SC.LocalRoot, "Engine/Content/Localization/ICU"), "*", true, null, null, false, !Params.Pak);
 			SC.StageFiles(StagedFileType.UFS, CombinePaths(SC.LocalRoot, "Engine/Plugins"), "*.uplugin", true, null, null, true, !Params.Pak);
 
 			// Game ufs (content)
 			SC.StageFiles(StagedFileType.UFS, CombinePaths(SC.ProjectRoot), "*.uproject", false, null, CombinePaths(SC.RelativeProjectRootForStage), true, !Params.Pak);
-			SC.StageFiles(StagedFileType.UFS, CombinePaths(SC.ProjectRoot, "Config"), "*", true, null, CombinePaths(SC.RelativeProjectRootForStage, "Config"), true, !Params.Pak);
+            SC.StageFiles(StagedFileType.UFS, CombinePaths(SC.ProjectRoot, "Config"), "*", true, new string[] { CombinePaths(SC.ProjectRoot, "Config/Localization") }, CombinePaths(SC.RelativeProjectRootForStage, "Config"), true, !Params.Pak);
 			SC.StageFiles(StagedFileType.UFS, CombinePaths(SC.ProjectRoot, "Plugins"), "*.uplugin", true, null, null, true, !Params.Pak);
             SC.StageFiles(StagedFileType.UFS, CombinePaths(SC.ProjectRoot, "Content/Localization/Game"), "*.locres", true, null, CombinePaths(SC.RelativeProjectRootForStage, "Content/Localization/Game"), true, !Params.Pak);
 
