Undid a problematic change to avoid staging localization data generation configuration files. To be properly done later.

diff --git a/Engine/Source/Programs/AutomationTool/Scripts/CopyBuildToStagingDirectory.Automation.cs b/Engine/Source/Programs/AutomationTool/Scripts/CopyBuildToStagingDirectory.Automation.cs
--- a/Engine/Source/Programs/AutomationTool/Scripts/CopyBuildToStagingDirectory.Automation.cs
+++ b/Engine/Source/Programs/AutomationTool/Scripts/CopyBuildToStagingDirectory.Automation.cs
@@ -176,7 +176,7 @@ public partial class Project : CommandUtils
 		if (!Params.CookOnTheFly && !Params.SkipCookOnTheFly) // only stage the UFS files if we are not using cook on the fly
 		{
 			// Engine ufs (content)
-            SC.StageFiles(StagedFileType.UFS, CombinePaths(SC.LocalRoot, "Engine/Config"), "*", true, new string[] { CombinePaths(SC.LocalRoot, "Engine/Config/Localization") }, null, false, !Params.Pak);
+            SC.StageFiles(StagedFileType.UFS, CombinePaths(SC.LocalRoot, "Engine/Config"), "*", true, null, null, false, !Params.Pak); // TODO: Exclude localization data generation config files.
 
 			if (Params.bUsesSlate)
 			{
@@ -193,7 +193,7 @@ public partial class Project : CommandUtils
 
 			// Game ufs (content)
 			SC.StageFiles(StagedFileType.UFS, CombinePaths(SC.ProjectRoot), "*.uproject", false, null, CombinePaths(SC.RelativeProjectRootForStage), true, !Params.Pak);
-            SC.StageFiles(StagedFileType.UFS, CombinePaths(SC.ProjectRoot, "Config"), "*", true, new string[] { CombinePaths(SC.ProjectRoot, "Config/Localization") }, CombinePaths(SC.RelativeProjectRootForStage, "Config"), true, !Params.Pak);
+            SC.StageFiles(StagedFileType.UFS, CombinePaths(SC.ProjectRoot, "Config"), "*", true, null, CombinePaths(SC.RelativeProjectRootForStage, "Config"), true, !Params.Pak); // TODO: Exclude localization data generation config files.
 			SC.StageFiles(StagedFileType.UFS, CombinePaths(SC.ProjectRoot, "Plugins"), "*.uplugin", true, null, null, true, !Params.Pak);
             SC.StageFiles(StagedFileType.UFS, CombinePaths(SC.ProjectRoot, "Content/Localization/Game"), "*.locres", true, null, CombinePaths(SC.RelativeProjectRootForStage, "Content/Localization/Game"), true, !Params.Pak);
 
