build script, add branch hacker to disable adder nodes

diff --git a/Engine/Source/Programs/AutomationTool/Scripts/GUBP.Automation.cs b/Engine/Source/Programs/AutomationTool/Scripts/GUBP.Automation.cs
--- a/Engine/Source/Programs/AutomationTool/Scripts/GUBP.Automation.cs
+++ b/Engine/Source/Programs/AutomationTool/Scripts/GUBP.Automation.cs
@@ -119,6 +119,7 @@ public class GUBP : BuildCommand
     static public bool bPreflightBuild = false;
     public int PreflightShelveCL = 0;
     static public string PreflightMangleSuffix = "";
+    public GUBPBranchHacker.BranchOptions BranchOptions = null;
 
     Dictionary<string, GUBPNode> GUBPNodes;
     Dictionary<string, bool> GUBPNodesCompleted;
@@ -190,6 +191,7 @@ public class GUBP : BuildCommand
         public class BranchOptions
         {            
             public List<UnrealTargetPlatform> PlatformsToRemove = new List<UnrealTargetPlatform>();
+            public List<string> ExcludeNodes = new List<string>();            
         }
         public virtual void ModifyOptions(GUBP bp, ref BranchOptions Options, string Branch)
         {
@@ -2655,16 +2657,11 @@ public class GUBP : BuildCommand
             }
             string Args = String.Format("BuildCookRun{0} -SkipBuild -SkipCook -Stage -Pak -Package -NoSubmit", ProjectArg);
 
-			bool bXboxOneTarget = false;
             if (ClientTargetPlatforms != null)
             {
                 bool bFirstClient = true;
                 foreach (var Plat in ClientTargetPlatforms)
                 {
-					if (Plat == UnrealTargetPlatform.XboxOne)
-					{
-						bXboxOneTarget = true;
-					}
                     if (bFirstClient)
                     {
                         bFirstClient = false;
@@ -2703,11 +2700,7 @@ public class GUBP : BuildCommand
                 bool bFirstServer = true;
                 foreach (var Plat in ServerTargetPlatforms)
                 {
-					if (Plat == UnrealTargetPlatform.XboxOne)
-					{
-						bXboxOneTarget = true;
-					}
-					if (bFirstServer)
+                    if (bFirstServer)
                     {
                         bFirstServer = false;
                         Args += String.Format(" -serverplatform={0}", Plat.ToString());
@@ -2732,45 +2725,22 @@ public class GUBP : BuildCommand
                 }
             }
 
-            string FinalArchiveDirectory = "";
-			string IntermediateArchiveDirectory = FinalArchiveDirectory;
             if (P4Enabled)
             {
-                FinalArchiveDirectory = GetArchiveDirectory(GameProj, HostPlatform, ClientTargetPlatforms, ClientConfigs, ServerTargetPlatforms, ServerConfigs, ClientNotGame);
-				IntermediateArchiveDirectory = FinalArchiveDirectory;
-				// Xbox One packaging does not function with remote file systems. Use a temp local directory to package and then move files into final location.
-				if (bXboxOneTarget)
-				{
-					IntermediateArchiveDirectory = Path.Combine(Path.GetTempPath(), "GUBP.XboxOne");
-					if (DirectoryExists_NoExceptions(IntermediateArchiveDirectory))
-					{
-						DeleteDirectory_NoExceptions(IntermediateArchiveDirectory);
-					}
-					CreateDirectory_NoExceptions(IntermediateArchiveDirectory);
-				}
-				CleanFormalBuilds(FinalArchiveDirectory);
-				if (DirectoryExists_NoExceptions(FinalArchiveDirectory))
+                string ArchiveDirectory = GetArchiveDirectory(GameProj, HostPlatform, ClientTargetPlatforms, ClientConfigs, ServerTargetPlatforms, ServerConfigs, ClientNotGame);
+                CleanFormalBuilds(ArchiveDirectory);
+                if (DirectoryExists_NoExceptions(ArchiveDirectory))
                 {
                     if (IsBuildMachine)
                     {
-						throw new AutomationException("Archive directory already exists {0}", FinalArchiveDirectory);
+                        throw new AutomationException("Archive directory already exists {0}", ArchiveDirectory);
                     }
-					DeleteDirectory_NoExceptions(FinalArchiveDirectory);
+                    DeleteDirectory_NoExceptions(ArchiveDirectory);
                 }
-				Args += String.Format(" -Archive -archivedirectory={0}", CommandUtils.MakePathSafeToUseWithCommandLine(IntermediateArchiveDirectory));
+                Args += String.Format(" -Archive -archivedirectory={0}", CommandUtils.MakePathSafeToUseWithCommandLine(ArchiveDirectory));
             }
 
             string LogFile = CommandUtils.RunUAT(CommandUtils.CmdEnv, Args);
-
-            if (P4Enabled)
-            {
-				if (!FinalArchiveDirectory.Equals(IntermediateArchiveDirectory, StringComparison.InvariantCultureIgnoreCase))
-				{
-					CopyDirectory_NoExceptions(IntermediateArchiveDirectory, FinalArchiveDirectory);
-					DeleteDirectory_NoExceptions(IntermediateArchiveDirectory);
-				}
-			}
-
             SaveRecordOfSuccessAndAddToBuildProducts(CommandUtils.ReadAllText(LogFile));
         }
     }
@@ -4484,7 +4454,7 @@ public class GUBP : BuildCommand
         {
             BranchForOptions = P4Env.BuildRootP4;
         }
-        var BranchOptions = GetBranchOptions(BranchForOptions);
+        BranchOptions = GetBranchOptions(BranchForOptions);
         bool WithMac = !BranchOptions.PlatformsToRemove.Contains(UnrealTargetPlatform.Mac);        
         if (ParseParam("NoMac"))
         {
@@ -5406,7 +5376,7 @@ if (HostPlatform == UnrealTargetPlatform.Mac) continue; //temp hack till mac aut
         }
         foreach (var HostPlatform in HostPlatforms)
         {
-            AddCustomNodes(HostPlatform);
+            AddCustomNodes(HostPlatform);            
         }
         
         if (HasNode(ToolsForCompileNode.StaticGetFullName(UnrealTargetPlatform.Win64)))
