build script, fixing stomped changes from my last change

diff --git a/Engine/Source/Programs/AutomationTool/Scripts/GUBP.Automation.cs b/Engine/Source/Programs/AutomationTool/Scripts/GUBP.Automation.cs
--- a/Engine/Source/Programs/AutomationTool/Scripts/GUBP.Automation.cs
+++ b/Engine/Source/Programs/AutomationTool/Scripts/GUBP.Automation.cs
@@ -2657,11 +2657,16 @@ public class GUBP : BuildCommand
             }
             string Args = String.Format("BuildCookRun{0} -SkipBuild -SkipCook -Stage -Pak -Package -NoSubmit", ProjectArg);
 
+			bool bXboxOneTarget = false;
             if (ClientTargetPlatforms != null)
             {
                 bool bFirstClient = true;
                 foreach (var Plat in ClientTargetPlatforms)
                 {
+					if (Plat == UnrealTargetPlatform.XboxOne)
+					{
+						bXboxOneTarget = true;
+					}
                     if (bFirstClient)
                     {
                         bFirstClient = false;
@@ -2700,7 +2705,11 @@ public class GUBP : BuildCommand
                 bool bFirstServer = true;
                 foreach (var Plat in ServerTargetPlatforms)
                 {
-                    if (bFirstServer)
+					if (Plat == UnrealTargetPlatform.XboxOne)
+					{
+						bXboxOneTarget = true;
+					}
+					if (bFirstServer)
                     {
                         bFirstServer = false;
                         Args += String.Format(" -serverplatform={0}", Plat.ToString());
@@ -2725,22 +2734,45 @@ public class GUBP : BuildCommand
                 }
             }
 
+            string FinalArchiveDirectory = "";
+			string IntermediateArchiveDirectory = FinalArchiveDirectory;
             if (P4Enabled)
             {
-                string ArchiveDirectory = GetArchiveDirectory(GameProj, HostPlatform, ClientTargetPlatforms, ClientConfigs, ServerTargetPlatforms, ServerConfigs, ClientNotGame);
-                CleanFormalBuilds(ArchiveDirectory);
-                if (DirectoryExists_NoExceptions(ArchiveDirectory))
+                FinalArchiveDirectory = GetArchiveDirectory(GameProj, HostPlatform, ClientTargetPlatforms, ClientConfigs, ServerTargetPlatforms, ServerConfigs, ClientNotGame);
+				IntermediateArchiveDirectory = FinalArchiveDirectory;
+				// Xbox One packaging does not function with remote file systems. Use a temp local directory to package and then move files into final location.
+				if (bXboxOneTarget)
+				{
+					IntermediateArchiveDirectory = Path.Combine(Path.GetTempPath(), "GUBP.XboxOne");
+					if (DirectoryExists_NoExceptions(IntermediateArchiveDirectory))
+					{
+						DeleteDirectory_NoExceptions(IntermediateArchiveDirectory);
+					}
+					CreateDirectory_NoExceptions(IntermediateArchiveDirectory);
+				}
+				CleanFormalBuilds(FinalArchiveDirectory);
+				if (DirectoryExists_NoExceptions(FinalArchiveDirectory))
                 {
                     if (IsBuildMachine)
                     {
-                        throw new AutomationException("Archive directory already exists {0}", ArchiveDirectory);
+						throw new AutomationException("Archive directory already exists {0}", FinalArchiveDirectory);
                     }
-                    DeleteDirectory_NoExceptions(ArchiveDirectory);
+					DeleteDirectory_NoExceptions(FinalArchiveDirectory);
                 }
-                Args += String.Format(" -Archive -archivedirectory={0}", CommandUtils.MakePathSafeToUseWithCommandLine(ArchiveDirectory));
+				Args += String.Format(" -Archive -archivedirectory={0}", CommandUtils.MakePathSafeToUseWithCommandLine(IntermediateArchiveDirectory));
             }
 
             string LogFile = CommandUtils.RunUAT(CommandUtils.CmdEnv, Args);
+
+            if (P4Enabled)
+            {
+				if (!FinalArchiveDirectory.Equals(IntermediateArchiveDirectory, StringComparison.InvariantCultureIgnoreCase))
+				{
+					CopyDirectory_NoExceptions(IntermediateArchiveDirectory, FinalArchiveDirectory);
+					DeleteDirectory_NoExceptions(IntermediateArchiveDirectory);
+				}
+			}
+
             SaveRecordOfSuccessAndAddToBuildProducts(CommandUtils.ReadAllText(LogFile));
         }
     }
@@ -5376,7 +5408,7 @@ if (HostPlatform == UnrealTargetPlatform.Mac) continue; //temp hack till mac aut
         }
         foreach (var HostPlatform in HostPlatforms)
         {
-            AddCustomNodes(HostPlatform);            
+            AddCustomNodes(HostPlatform);
         }
         
         if (HasNode(ToolsForCompileNode.StaticGetFullName(UnrealTargetPlatform.Win64)))
