UAT: Log process exit code, always Dispose processes.

diff --git a/Engine/Source/Programs/AutomationTool/ProcessUtils.cs b/Engine/Source/Programs/AutomationTool/ProcessUtils.cs
--- a/Engine/Source/Programs/AutomationTool/ProcessUtils.cs
+++ b/Engine/Source/Programs/AutomationTool/ProcessUtils.cs
@@ -698,16 +698,11 @@ namespace AutomationTool
 				AddRunTime(App, (int)(BuildDuration));
                 if (!Options.HasFlag(ERunOptions.NoLoggingOfRunCommand))
                 {
-                    Log(SpewVerbosity,"Run: Took {0}s to run " + Path.GetFileName(App), BuildDuration / 1000);
-
+                    Log(SpewVerbosity,"Run: Took {0}s to run {1}, ExitCode={2}", BuildDuration / 1000, Path.GetFileName(App), Result.ExitCode);
                 }
 				Result.ExitCode = Proc.ExitCode;
 				Result.OnProcessExited();
-				if (UnrealBuildTool.Utils.IsRunningOnMono)
-				{
-					// Mono's detection of process exit status is broken. It doesn't clean up after the process and it doesn't call Exited callback.
-					Proc.Dispose();					
-				}
+				Proc.Dispose();
 			}
 			else
 			{
