fix for not having the distribution parameter being passed to IPP
#ue4
#ios

diff --git a/Engine/Source/Programs/AutomationTool/IOS/IOSPlatform.Automation.cs b/Engine/Source/Programs/AutomationTool/IOS/IOSPlatform.Automation.cs
--- a/Engine/Source/Programs/AutomationTool/IOS/IOSPlatform.Automation.cs
+++ b/Engine/Source/Programs/AutomationTool/IOS/IOSPlatform.Automation.cs
@@ -187,10 +187,10 @@ public class IOSPlatform : Platform
 			// rename the .ipa if not code based
 			if (!Params.IsCodeBasedProject)
 			{
-				ProjectIPA = Path.Combine(Path.GetDirectoryName(Params.RawProjectPath), "Binaries", "IOS", Params.ShortProjectName + ".ipa");
+				ProjectIPA = Path.Combine(Path.GetDirectoryName(Params.RawProjectPath), "Binaries", "IOS", (Params.Distribution ? "Distro_" : "") + Params.ShortProjectName + ".ipa");
 				if (TargetConfiguration != UnrealTargetConfiguration.Development)
 				{
-					ProjectIPA = Path.Combine(Path.GetDirectoryName(Params.RawProjectPath), "Binaries", "IOS", Params.ShortProjectName + "-" + PlatformType.ToString() + "-" + TargetConfiguration.ToString() + ".ipa");
+					ProjectIPA = Path.Combine(Path.GetDirectoryName(Params.RawProjectPath), "Binaries", "IOS", (Params.Distribution ? "Distro_" : "") + Params.ShortProjectName + "-" + PlatformType.ToString() + "-" + TargetConfiguration.ToString() + ".ipa");
 				}
 			}
 
@@ -220,6 +220,10 @@ public class IOSPlatform : Platform
 				if (bNeedToSign)
 				{
 					IPPArguments += " -sign";
+					if (Params.Distribution)
+					{
+						IPPArguments += " -distribution";
+					}
 				}
 
 				IPPArguments += (cookonthefly ? " -cookonthefly" : "");
@@ -297,10 +301,10 @@ public class IOSPlatform : Platform
 		else
 		{
 			// code sign the app
-			CodeSign(Path.GetDirectoryName(Params.ProjectGameExeFilename), Path.GetFileNameWithoutExtension(Params.ProjectGameExeFilename), Params.RawProjectPath, SC.StageTargetConfigurations[0], SC.LocalRoot, Params.ShortProjectName, Path.GetDirectoryName(Params.RawProjectPath), SC.IsCodeBasedProject, Params.Distribution);
+			CodeSign(Path.GetDirectoryName(Params.ProjectGameExeFilename), Params.IsCodeBasedProject ? Params.ShortProjectName : Path.GetFileNameWithoutExtension(Params.ProjectGameExeFilename), Params.RawProjectPath, SC.StageTargetConfigurations[0], SC.LocalRoot, Params.ShortProjectName, Path.GetDirectoryName(Params.RawProjectPath), SC.IsCodeBasedProject, Params.Distribution);
 
 			// now generate the ipa
-			PackageIPA(Path.GetDirectoryName(Params.ProjectGameExeFilename), Path.GetFileNameWithoutExtension(Params.ProjectGameExeFilename), Params.ShortProjectName, Path.GetDirectoryName(Params.RawProjectPath), SC.StageTargetConfigurations[0], Params.Distribution);
+			PackageIPA(Path.GetDirectoryName(Params.ProjectGameExeFilename), Params.IsCodeBasedProject ? Params.ShortProjectName : Path.GetFileNameWithoutExtension(Params.ProjectGameExeFilename), Params.ShortProjectName, Path.GetDirectoryName(Params.RawProjectPath), SC.StageTargetConfigurations[0], Params.Distribution);
 		}
 
 		PrintRunTime();
@@ -757,10 +761,10 @@ public class IOSPlatform : Platform
 		// rename the .ipa if not code based
 		if (!Params.IsCodeBasedProject)
 		{
-			ProjectIPA = Path.Combine(Path.GetDirectoryName(Params.RawProjectPath), "Binaries", "IOS", Params.ShortProjectName + ".ipa");
+			ProjectIPA = Path.Combine(Path.GetDirectoryName(Params.RawProjectPath), "Binaries", "IOS", (Params.Distribution ? "Distro_" : "") + Params.ShortProjectName + ".ipa");
 			if (TargetConfiguration != UnrealTargetConfiguration.Development)
 			{
-				ProjectIPA = Path.Combine(Path.GetDirectoryName(Params.RawProjectPath), "Binaries", "IOS", Params.ShortProjectName + "-" + PlatformType.ToString() + "-" + TargetConfiguration.ToString() + ".ipa");
+				ProjectIPA = Path.Combine(Path.GetDirectoryName(Params.RawProjectPath), "Binaries", "IOS", (Params.Distribution ? "Distro_" : "") + Params.ShortProjectName + "-" + PlatformType.ToString() + "-" + TargetConfiguration.ToString() + ".ipa");
 			}
 		}
 
@@ -781,6 +785,10 @@ public class IOSPlatform : Platform
 			{
 				throw new AutomationException ("iOS is currently only able to package one target configuration at a time, but StageTargetConfigurations contained {0} configurations", SC.StageTargetConfigurations.Count);
 			}
+			if (Params.Distribution)
+			{
+				throw new AutomationException("iOS cannot deploy a package made for distribution.");
+			}
 			var TargetConfiguration = SC.StageTargetConfigurations[0];
 
 //			var ProjectStub = Params.ProjectGameExeFilename;	
