fix for name of distribution ipa on mac
TTP343694
#ue4
#ios

diff --git a/Engine/Source/Programs/AutomationTool/IOS/IOSPlatform.Automation.cs b/Engine/Source/Programs/AutomationTool/IOS/IOSPlatform.Automation.cs
--- a/Engine/Source/Programs/AutomationTool/IOS/IOSPlatform.Automation.cs
+++ b/Engine/Source/Programs/AutomationTool/IOS/IOSPlatform.Automation.cs
@@ -506,7 +506,7 @@ public class IOSPlatform : Platform
 	private void PackageIPA(string BaseDirectory, string GameName, string ProjectName, string ProjectDirectory, UnrealTargetConfiguration TargetConfig, bool Distribution = false)
 	{
 		// create the ipa
-		string IPAName = CombinePaths(ProjectDirectory, "Binaries", "IOS", ProjectName + (TargetConfig != UnrealTargetConfiguration.Development ? ("-IOS-" + TargetConfig.ToString()) : "") + ".ipa");
+		string IPAName = CombinePaths(ProjectDirectory, "Binaries", "IOS", (Distribution ? "Distro_" : "") + ProjectName + (TargetConfig != UnrealTargetConfiguration.Development ? ("-IOS-" + TargetConfig.ToString()) : "") + ".ipa");
 		// delete the old one
 		if (File.Exists(IPAName))
 		{
