Added Cultures-to-Stage list to Project Packaging Settings. Default value includes only English. Packaging now only stages specified localization data as per the previous comment or by passing "-CookCultures=".

diff --git a/Engine/Source/Editor/UnrealEd/Classes/Settings/ProjectPackagingSettings.h b/Engine/Source/Editor/UnrealEd/Classes/Settings/ProjectPackagingSettings.h
--- a/Engine/Source/Editor/UnrealEd/Classes/Settings/ProjectPackagingSettings.h
+++ b/Engine/Source/Editor/UnrealEd/Classes/Settings/ProjectPackagingSettings.h
@@ -61,11 +61,17 @@ public:
 	UPROPERTY(config, EditAnywhere, Category=Packaging)
 	bool UsePakFile;
 
-
 	/** If enabled, on Android platforms, .pak files are placed inside the APK. */
 	UPROPERTY(config, EditAnywhere, Category = Packaging)
 	bool UseOBB_InAPK;
 
+	/**
+	 * ISO codes of cultures whose data should be packaged to be cooked, staged, and packaged.
+	 * Note: These paths are relative to your project Content directory
+	 */
+	UPROPERTY(config, EditAnywhere, Category=Packaging, AdvancedDisplay, meta=(DisplayName="Cultures to Package"))
+	TArray<FString> CulturesToStage;
+
 	/**
 	 * Directories containing .uasset files that should always be cooked regardless of whether they're referenced by anything in your project
 	 * Note: These paths are relative to your project Content directory
diff --git a/Engine/Source/Editor/UnrealEd/Private/Commandlets/CookCommandlet.cpp b/Engine/Source/Editor/UnrealEd/Private/Commandlets/CookCommandlet.cpp
--- a/Engine/Source/Editor/UnrealEd/Private/Commandlets/CookCommandlet.cpp
+++ b/Engine/Source/Editor/UnrealEd/Private/Commandlets/CookCommandlet.cpp
@@ -759,50 +759,44 @@ void UCookCommandlet::CollectFilesToCook(TArray<FString>& FilesInPath)
 
 	TArray<FString> CmdLineMapEntries;
 	TArray<FString> CmdLineDirEntries;
+	TArray<FString> CmdLineCultEntries;
 	for (int32 SwitchIdx = 0; SwitchIdx < Switches.Num(); SwitchIdx++)
 	{
 		const FString& Switch = Switches[SwitchIdx];
 
-		// Check for -MAP=<name of map> entries
-		if (Switch.StartsWith(TEXT("MAP=")) == true)
+		auto GetSwitchValueElements = [&Switch](const FString SwitchKey) -> TArray<FString>
 		{
-			FString MapToCook = Switch.Right(Switch.Len() - 4);
-			
-			// Allow support for -MAP=Map1+Map2+Map3 as well as -MAP=Map1 -MAP=Map2
-			for (int32 PlusIdx = MapToCook.Find(TEXT("+")); PlusIdx != INDEX_NONE; PlusIdx = MapToCook.Find(TEXT("+")))
+			TArray<FString> ValueElements;
+			if (Switch.StartsWith(SwitchKey + TEXT("=")) == true)
 			{
-				FString MapName = MapToCook.Left(PlusIdx);
-				CmdLineMapEntries.Add(MapName);
+				FString ValuesList = Switch.Right(Switch.Len() - (SwitchKey + TEXT("=")).Len());
+
+				// Allow support for -KEY=Value1+Value2+Value3 as well as -KEY=Value1 -KEY=Value2
+				for (int32 PlusIdx = ValuesList.Find(TEXT("+")); PlusIdx != INDEX_NONE; PlusIdx = ValuesList.Find(TEXT("+")))
+				{
+					const FString ValueElement = ValuesList.Left(PlusIdx);
+					ValueElements.Add(ValueElement);
 
-				MapToCook = MapToCook.Right(MapToCook.Len() - (PlusIdx + 1));
+					ValuesList = ValuesList.Right(ValuesList.Len() - (PlusIdx + 1));
+				}
+				ValueElements.Add(ValuesList);
 			}
+			return ValueElements;
+		};
 
-			CmdLineMapEntries.Add(MapToCook);
-		}
+		// Check for -MAP=<name of map> entries
+		CmdLineMapEntries = GetSwitchValueElements(TEXT("MAP"));
 
 		// Check for -COOKDIR=<path to directory> entries
-		if (Switch.StartsWith(TEXT("COOKDIR=")) == true)
+		CmdLineDirEntries = GetSwitchValueElements(TEXT("COOKDIR"));
+		for(FString& Entry : CmdLineDirEntries)
 		{
-			FString DirToCook = Switch.Right(Switch.Len() - 8);
-			
-			// Allow support for -COOKDIR=Dir1+Dir2+Dir3 as well as -COOKDIR=Dir1 -COOKDIR=Dir2
-			for (int32 PlusIdx = DirToCook.Find(TEXT("+")); PlusIdx != INDEX_NONE; PlusIdx = DirToCook.Find(TEXT("+")))
-			{
-				FString DirName = DirToCook.Left(PlusIdx);
-				
-				// The dir may be contained within quotes
-				DirName = DirName.TrimQuotes();
-				FPaths::NormalizeDirectoryName(DirName);
-				CmdLineDirEntries.Add(DirName);
-
-				DirToCook = DirToCook.Right(DirToCook.Len() - (PlusIdx + 1));
-			}
-			
-			// The dir may be contained within quotes
-			DirToCook = DirToCook.TrimQuotes();
-			FPaths::NormalizeDirectoryName(DirToCook);
-			CmdLineDirEntries.Add(DirToCook);
+			Entry = Entry.TrimQuotes();
+			FPaths::NormalizeDirectoryName(Entry);
 		}
+
+		// Check for -COOKCULTURES=<culture name> entries
+		CmdLineCultEntries = GetSwitchValueElements(TEXT("COOKCULTURES"));
 	}
 
 	// Also append any cookdirs from the project ini files; these dirs are relative to the game content directory
diff --git a/Engine/Source/Editor/UnrealEd/Private/Settings/SettingsClasses.cpp b/Engine/Source/Editor/UnrealEd/Private/Settings/SettingsClasses.cpp
--- a/Engine/Source/Editor/UnrealEd/Private/Settings/SettingsClasses.cpp
+++ b/Engine/Source/Editor/UnrealEd/Private/Settings/SettingsClasses.cpp
@@ -277,6 +277,7 @@ UProjectPackagingSettings::UProjectPackagingSettings( const class FPostConstruct
 	FullRebuild = true;
 	UsePakFile = true;
 	UseOBB_InAPK = false;
+	CulturesToStage.Add("en");
 }
 
 
diff --git a/Engine/Source/Programs/AutomationTool/CommandletUtils.cs b/Engine/Source/Programs/AutomationTool/CommandletUtils.cs
--- a/Engine/Source/Programs/AutomationTool/CommandletUtils.cs
+++ b/Engine/Source/Programs/AutomationTool/CommandletUtils.cs
@@ -21,7 +21,7 @@ namespace AutomationTool
 		/// <param name="Dirs">List of directories to cook, can be null</param>
 		/// <param name="TargetPlatform">Target platform.</param>
 		/// <param name="Parameters">List of additional parameters.</param>
-		public static void CookCommandlet(string ProjectName, string UE4Exe = "UE4Editor-Cmd.exe", string[] Maps = null, string[] Dirs = null, string TargetPlatform = "WindowsNoEditor", string Parameters = "-Unversioned")
+		public static void CookCommandlet(string ProjectName, string UE4Exe = "UE4Editor-Cmd.exe", string[] Maps = null, string[] Dirs = null, string[] Cultures = null, string TargetPlatform = "WindowsNoEditor", string Parameters = "-Unversioned")
 		{
 			string MapsToCook = "";
 			if (IsNullOrEmpty(Maps))
@@ -41,7 +41,14 @@ namespace AutomationTool
 				DirsToCook.Trim();		
 			}
 
-			RunCommandlet(ProjectName, UE4Exe, "Cook", String.Format("{0} {1} -TargetPlatform={2} {3}", MapsToCook, DirsToCook, TargetPlatform, Parameters));
+            string CulturesToCook = "";
+            if (!IsNullOrEmpty(Cultures))
+            {
+                CulturesToCook = "-CookCultures=" + CombineCommandletParams(Cultures);
+                CulturesToCook.Trim();
+            }
+
+            RunCommandlet(ProjectName, UE4Exe, "Cook", String.Format("{0} {1} {2} -TargetPlatform={3} {4}", MapsToCook, DirsToCook, CulturesToCook, TargetPlatform, Parameters));
 		}
 
         /// <summary>
diff --git a/Engine/Source/Programs/AutomationTool/ProjectParams.cs b/Engine/Source/Programs/AutomationTool/ProjectParams.cs
--- a/Engine/Source/Programs/AutomationTool/ProjectParams.cs
+++ b/Engine/Source/Programs/AutomationTool/ProjectParams.cs
@@ -179,6 +179,7 @@ namespace AutomationTool
 			this.RawProjectPath = InParams.RawProjectPath;
 			this.MapsToCook = InParams.MapsToCook;
 			this.DirectoriesToCook = InParams.DirectoriesToCook;
+            this.CulturesToCook = InParams.CulturesToCook;
 			this.ClientCookedTargets = InParams.ClientCookedTargets;
 			this.ServerCookedTargets = InParams.ServerCookedTargets;
 			this.EditorTargets = InParams.EditorTargets;
@@ -279,6 +280,7 @@ namespace AutomationTool
 			List<UnrealTargetConfiguration> ServerConfigsToBuild = null,
 			ParamList<string> MapsToCook = null,
 			ParamList<string> DirectoriesToCook = null,
+            ParamList<string> CulturesToCook = null,
 			ParamList<string> ClientCookedTargets = null,
 			ParamList<string> EditorTargets = null,
 			ParamList<string> ServerCookedTargets = null,
@@ -347,6 +349,10 @@ namespace AutomationTool
 			{
 				this.DirectoriesToCook = DirectoriesToCook;
 			}
+            if (CulturesToCook != null)
+			{
+                this.CulturesToCook = CulturesToCook;
+			}
 			if (ClientCookedTargets != null)
 			{
 				this.ClientCookedTargets = ClientCookedTargets;
@@ -825,6 +831,11 @@ namespace AutomationTool
 		/// </summary>
 		public ParamList<string> DirectoriesToCook = new ParamList<string>();
 
+        /// <summary>
+        /// Cook: List of cultures to cook.
+        /// </summary>
+        public ParamList<string> CulturesToCook = new ParamList<string>();
+
         /// <summary>
         /// Compress packages during cook.
         /// </summary>
@@ -1354,6 +1365,11 @@ namespace AutomationTool
 			get { return !CommandUtils.IsNullOrEmpty(DirectoriesToCook); }
 		}
 
+        public bool HasCulturesToCook
+        {
+            get { return !CommandUtils.IsNullOrEmpty(CulturesToCook); }
+        }
+
 		public bool HasGameTargetDetected
 		{
 			get { return ProjectTargets.ContainsKey(TargetRules.TargetType.Game); }
@@ -1618,6 +1634,7 @@ namespace AutomationTool
 				CommandUtils.Log("CookOnTheFly={0}", CookOnTheFly);
 				CommandUtils.Log("DedicatedServer={0}", DedicatedServer);
 				CommandUtils.Log("DirectoriesToCook={0}", DirectoriesToCook.ToString());
+                CommandUtils.Log("CulturesToCook={0}", CulturesToCook.ToString());
 				CommandUtils.Log("EditorTargets={0}", EditorTargets.ToString());
 				CommandUtils.Log("Foreign={0}", Foreign);
 				CommandUtils.Log("IsCodeBasedProject={0}", IsCodeBasedProject.ToString());
diff --git a/Engine/Source/Programs/AutomationTool/Scripts/BuildCookRun.Automation.cs b/Engine/Source/Programs/AutomationTool/Scripts/BuildCookRun.Automation.cs
--- a/Engine/Source/Programs/AutomationTool/Scripts/BuildCookRun.Automation.cs
+++ b/Engine/Source/Programs/AutomationTool/Scripts/BuildCookRun.Automation.cs
@@ -102,6 +102,12 @@ public class BuildCookRun : BuildCommand
 			Params.DirectoriesToCook = new ParamList<string>(DirectoriesToCook.Split('+'));
 		}
 
+        var CulturesToCook = ParseParamValue("cookcultures");
+        if (!String.IsNullOrEmpty(CulturesToCook))
+        {
+            Params.CulturesToCook = new ParamList<string>(CulturesToCook.Split('+'));
+        }
+
 		if (Params.DedicatedServer)
 		{
 			foreach ( var ServerPlatformInstance in Params.ServerTargetPlatformInstances )
diff --git a/Engine/Source/Programs/AutomationTool/Scripts/CookCommand.Automation.cs b/Engine/Source/Programs/AutomationTool/Scripts/CookCommand.Automation.cs
--- a/Engine/Source/Programs/AutomationTool/Scripts/CookCommand.Automation.cs
+++ b/Engine/Source/Programs/AutomationTool/Scripts/CookCommand.Automation.cs
@@ -68,6 +68,12 @@ public partial class Project : CommandUtils
 			Dirs = Params.DirectoriesToCook.ToArray();
 		}
 
+        string[] Cultures = null;
+        if (Params.HasCulturesToCook)
+        {
+            Cultures = Params.CulturesToCook.ToArray();
+        }
+
 		try
 		{
 			var CommandletParams = "-buildmachine -Unversioned -fileopenlog";
@@ -83,7 +89,7 @@ public partial class Project : CommandUtils
 			{
 				CommandletParams += " -iterate";
 			}
-			CookCommandlet(Params.RawProjectPath, Params.UE4Exe, Maps, Dirs, CombineCommandletParams(PlatformsToCook.ToArray()), CommandletParams);
+			CookCommandlet(Params.RawProjectPath, Params.UE4Exe, Maps, Dirs, Cultures, CombineCommandletParams(PlatformsToCook.ToArray()), CommandletParams);
 		}
 		catch (Exception Ex)
 		{
diff --git a/Engine/Source/Programs/AutomationTool/Scripts/CopyBuildToStagingDirectory.Automation.cs b/Engine/Source/Programs/AutomationTool/Scripts/CopyBuildToStagingDirectory.Automation.cs
--- a/Engine/Source/Programs/AutomationTool/Scripts/CopyBuildToStagingDirectory.Automation.cs
+++ b/Engine/Source/Programs/AutomationTool/Scripts/CopyBuildToStagingDirectory.Automation.cs
@@ -175,6 +175,29 @@ public partial class Project : CommandUtils
 
 		if (!Params.CookOnTheFly && !Params.SkipCookOnTheFly) // only stage the UFS files if we are not using cook on the fly
 		{
+            ConfigCacheIni PlatformGameConfig = null;
+            Params.GameConfigs.TryGetValue(SC.StageTargetPlatform.PlatformType, out PlatformGameConfig);
+
+            // Initialize cultures to stage.
+            List<string> CulturesToStage = null;
+            if (PlatformGameConfig != null)
+            {
+                PlatformGameConfig.GetArray("/Script/UnrealEd.ProjectPackagingSettings", "CulturesToStage", out CulturesToStage);
+            }
+            if (CulturesToStage != null)
+            {
+                CulturesToStage.AddRange(Params.CulturesToCook);
+            }
+            else
+            {
+                CulturesToStage = new List<string>(Params.CulturesToCook);
+            }
+
+            if (CulturesToStage == null || CulturesToStage.Count == 0)
+            {
+                throw new AutomationException("No cultures were specified for cooking and packaging. This will lead to fatal errors when launching. Specify culture codes via commandline (-CookCultures=) or using project packaging settings.");
+            }
+
 			// Engine ufs (content)
             SC.StageFiles(StagedFileType.UFS, CombinePaths(SC.LocalRoot, "Engine/Config"), "*", true, null, null, false, !Params.Pak); // TODO: Exclude localization data generation config files.
 
@@ -188,18 +211,24 @@ public partial class Project : CommandUtils
 				SC.StageFiles(StagedFileType.UFS, CombinePaths(SC.ProjectRoot, "Content/Slate"), "*", true, null, CombinePaths(SC.RelativeProjectRootForStage, "Content/Slate"), true, !Params.Pak);
 
 			}
-			SC.StageFiles(StagedFileType.UFS, CombinePaths(SC.LocalRoot, "Engine/Content/Localization/Engine"), "*.locres", true, null, null, false, !Params.Pak);
+            foreach (string Culture in CulturesToStage)
+            {
+                SC.StageFiles(StagedFileType.UFS, CombinePaths(SC.LocalRoot, "Engine/Content/Localization/Engine", Culture), "*.locres", true, null, null, true, !Params.Pak);
+            }
 			SC.StageFiles(StagedFileType.UFS, CombinePaths(SC.LocalRoot, "Engine/Plugins"), "*.uplugin", true, null, null, true, !Params.Pak);
 
 			// Game ufs (content)
+            
 			SC.StageFiles(StagedFileType.UFS, CombinePaths(SC.ProjectRoot), "*.uproject", false, null, CombinePaths(SC.RelativeProjectRootForStage), true, !Params.Pak);
             SC.StageFiles(StagedFileType.UFS, CombinePaths(SC.ProjectRoot, "Config"), "*", true, null, CombinePaths(SC.RelativeProjectRootForStage, "Config"), true, !Params.Pak); // TODO: Exclude localization data generation config files.
 			SC.StageFiles(StagedFileType.UFS, CombinePaths(SC.ProjectRoot, "Plugins"), "*.uplugin", true, null, null, true, !Params.Pak);
-            SC.StageFiles(StagedFileType.UFS, CombinePaths(SC.ProjectRoot, "Content/Localization/Game"), "*.locres", true, null, CombinePaths(SC.RelativeProjectRootForStage, "Content/Localization/Game"), true, !Params.Pak);
+            foreach (string Culture in CulturesToStage)
+            {
+                SC.StageFiles(StagedFileType.UFS, CombinePaths(SC.ProjectRoot, "Content/Localization/Game", Culture), "*.locres", true, null, CombinePaths(SC.RelativeProjectRootForStage, "Content/Localization/Game", Culture), true, !Params.Pak);
+            }
 
 			// Stage any additional UFS and NonUFS paths specified in the project ini files; these dirs are relative to the game content directory
-			ConfigCacheIni PlatformGameConfig = null;
-			if (Params.GameConfigs.TryGetValue(SC.StageTargetPlatform.PlatformType, out PlatformGameConfig))
+            if (PlatformGameConfig != null)
 			{
 				var ProjectContentRoot = CombinePaths(SC.ProjectRoot, "Content");
 				var StageContentRoot = CombinePaths(SC.RelativeProjectRootForStage, "Content");
diff --git a/Engine/Source/Programs/AutomationTool/Scripts/GUBP.Automation.cs b/Engine/Source/Programs/AutomationTool/Scripts/GUBP.Automation.cs
--- a/Engine/Source/Programs/AutomationTool/Scripts/GUBP.Automation.cs
+++ b/Engine/Source/Programs/AutomationTool/Scripts/GUBP.Automation.cs
@@ -2239,7 +2239,7 @@ public class GUBP : BuildCommand
             {
                 // not sure if we need something here or if the cook commandlet will automatically convert the exe name
             }
-            CommandUtils.CookCommandlet(GameProj.FilePath, "UE4Editor-Cmd.exe", null, null, CookPlatform);
+            CommandUtils.CookCommandlet(GameProj.FilePath, "UE4Editor-Cmd.exe", null, null, null, CookPlatform);
             var CookedPath = RootIfAnyForTempStorage();
             var CookedFiles = CommandUtils.FindFiles("*", true, CookedPath);
             if (CookedFiles.GetLength(0) < 1)
