add a default culture if one is not provided
#ue4

diff --git a/Engine/Source/Programs/AutomationTool/Scripts/BuildCookRun.Automation.cs b/Engine/Source/Programs/AutomationTool/Scripts/BuildCookRun.Automation.cs
--- a/Engine/Source/Programs/AutomationTool/Scripts/BuildCookRun.Automation.cs
+++ b/Engine/Source/Programs/AutomationTool/Scripts/BuildCookRun.Automation.cs
@@ -107,6 +107,10 @@ public class BuildCookRun : BuildCommand
         {
             Params.CulturesToCook = new ParamList<string>(CulturesToCook.Split('+'));
         }
+		else
+		{
+			Params.CulturesToCook = new ParamList<string>(new string[] { "en" });
+		}
 
 		if (Params.DedicatedServer)
 		{
