fix for search order for Info.plist
#ue4
#ios

diff --git a/Engine/Source/Programs/AutomationTool/IOS/IOSPlatform.Automation.cs b/Engine/Source/Programs/AutomationTool/IOS/IOSPlatform.Automation.cs
--- a/Engine/Source/Programs/AutomationTool/IOS/IOSPlatform.Automation.cs
+++ b/Engine/Source/Programs/AutomationTool/IOS/IOSPlatform.Automation.cs
@@ -341,13 +341,13 @@ public class IOSPlatform : Platform
 
 		// copy the appropriate plist file over
 		string SourcePListFile = CombinePaths(LocalRoot, "Engine", "Build", "IOS", "UE4Game-Info.plist");
-		if (File.Exists(ProjectRoot + "/Build/IOS/" + ShortProjectName + "-Info.plist"))
+		if (File.Exists(ProjectRoot + "/Build/IOS/Info.plist"))
 		{
-			SourcePListFile = CombinePaths(ProjectRoot, "Build", "IOS", ShortProjectName + "-Info.plist");
+			SourcePListFile = CombinePaths(ProjectRoot, "Build", "IOS", "Info.plist");
 		}
-		else if (File.Exists(ProjectRoot + "/Build/IOS/Info.plist"))
+		else if (File.Exists(ProjectRoot + "/Build/IOS/" + ShortProjectName + "-Info.plist"))
 		{
-			SourcePListFile = CombinePaths(ProjectRoot, "Build", "IOS", "Info.plist");
+			SourcePListFile = CombinePaths(ProjectRoot, "Build", "IOS", ShortProjectName + "-Info.plist");
 		}
 		else if (Directory.Exists(ProjectRoot + "/Build/IOS"))
 		{
@@ -703,13 +703,13 @@ public class IOSPlatform : Platform
 			if (GetCodeSignDesirability(Params))
 			{
 				string SourcePListFile = CombinePaths(SC.LocalRoot, "Engine", "Build", "IOS", "UE4Game-Info.plist");
-				if (File.Exists(SC.ProjectRoot + "/Build/IOS/" + SC.ShortProjectName + "-Info.plist"))
+				if (File.Exists(SC.ProjectRoot + "/Build/IOS/Info.plist"))
 				{
-					SourcePListFile = CombinePaths(SC.ProjectRoot, "Build", "IOS", SC.ShortProjectName + "-Info.plist");
+					SourcePListFile = CombinePaths(SC.ProjectRoot, "Build", "IOS", "Info.plist");
 				}
-				else if (File.Exists(SC.ProjectRoot + "/Build/IOS/Info.plist"))
+				else if (File.Exists(SC.ProjectRoot + "/Build/IOS/" + SC.ShortProjectName + "-Info.plist"))
 				{
-					SourcePListFile = CombinePaths(SC.ProjectRoot, "Build", "IOS", "Info.plist");
+					SourcePListFile = CombinePaths(SC.ProjectRoot, "Build", "IOS", SC.ShortProjectName + "-Info.plist");
 				}
 				else if (Directory.Exists(SC.ProjectRoot + "/Build/IOS"))
 				{
