build script, moving sample makebuilds behind promotable trigger

diff --git a/Engine/Source/Programs/AutomationTool/Scripts/GUBP.Automation.cs b/Engine/Source/Programs/AutomationTool/Scripts/GUBP.Automation.cs
--- a/Engine/Source/Programs/AutomationTool/Scripts/GUBP.Automation.cs
+++ b/Engine/Source/Programs/AutomationTool/Scripts/GUBP.Automation.cs
@@ -2457,7 +2457,7 @@ public class GUBP : BuildCommand
         List<UnrealTargetConfiguration> ClientConfigs;
         List<UnrealTargetConfiguration> ServerConfigs;
         bool ClientNotGame;
-        UnrealBuildTool.TargetRules.TargetType GameOrClient;
+        UnrealBuildTool.TargetRules.TargetType GameOrClient;		
 
         public FormalBuildNode(GUBP bp,
             BranchInfo.BranchUProject InGameProj,
@@ -2466,8 +2466,7 @@ public class GUBP : BuildCommand
             List<UnrealTargetConfiguration> InClientConfigs = null,
             List<UnrealTargetPlatform> InServerTargetPlatforms = null,
             List<UnrealTargetConfiguration> InServerConfigs = null,
-            bool InClientNotGame = false
-
+            bool InClientNotGame = false			
             )
             : base(InHostPlatform)
         {
@@ -2476,7 +2475,7 @@ public class GUBP : BuildCommand
             ServerTargetPlatforms = InServerTargetPlatforms;
             ClientConfigs = InClientConfigs;
             ServerConfigs = InServerConfigs;
-            ClientNotGame = InClientNotGame;
+            ClientNotGame = InClientNotGame;			
 
             GameOrClient = TargetRules.TargetType.Game;
 
@@ -2495,6 +2494,11 @@ public class GUBP : BuildCommand
             }
 
             var AllTargetPlatforms = new List<UnrealTargetPlatform>();
+			var Options = InGameProj.Options(HostPlatform);
+			if(!Options.bSeparateGamePromotion)
+			{
+				AddPseudodependency(SharedLabelPromotableNode.StaticGetFullName(false));				
+			}
             if (ClientTargetPlatforms != null)
             {
                 if (!WorkingGameProject.Properties.Targets.ContainsKey(GameOrClient))
@@ -2564,7 +2568,10 @@ public class GUBP : BuildCommand
             // add dependencies for cooked and compiled
             foreach (var Plat in AllTargetPlatforms)
             {
-                AddDependency(GamePlatformCookedAndCompiledNode.StaticGetFullName(HostPlatform, GameProj, Plat));
+				if (Options.bSeparateGamePromotion)
+				{
+					AddDependency(GamePlatformCookedAndCompiledNode.StaticGetFullName(HostPlatform, GameProj, Plat));
+				}
             }
         }
 
@@ -2629,7 +2636,7 @@ public class GUBP : BuildCommand
         }      
         public override float Priority()
         {
-            return base.Priority() + 20.0f;
+			return base.Priority() + 20.0f;
         }
         public override int CISFrequencyQuantumShift(GUBP bp)
         {
@@ -2820,6 +2827,11 @@ public class GUBP : BuildCommand
 
             // verify we actually built these
             var WorkingGameProject = InGameProj;
+			var Options = InGameProj.Options(HostPlatform);
+			if(!Options.bSeparateGamePromotion)
+			{
+				AddPseudodependency(SharedLabelPromotableNode.StaticGetFullName(false));
+			}
             if (!WorkingGameProject.Properties.Targets.ContainsKey(TargetRules.TargetType.Editor))
             {
                 // this is a codeless project, use the base project
@@ -5096,32 +5108,39 @@ public class GUBP : BuildCommand
                                             {
                                                 if (PlatPair.TargetPlatform == Plat)
                                                 {
-                                                    var NodeName = AddNode(new FormalBuildNode(this, NonCodeProject, HostPlatform, new List<UnrealTargetPlatform>() { Plat }, new List<UnrealTargetConfiguration>() {PlatPair.TargetConfig}));
-                                                    // we don't want this delayed
-                                                    // this would normally wait for the testing phase, we just want to build it right away
-                                                    RemovePseudodependencyFromNode(
-                                                        CookNode.StaticGetFullName(HostPlatform, NonCodeProject, CookedPlatform),
-                                                        WaitForTestShared.StaticGetFullName());
-                                                    RemovePseudodependencyFromNode(
-                                                        CookNode.StaticGetFullName(HostPlatform, NonCodeProject, CookedPlatform),
-                                                        CookNode.StaticGetFullName(HostPlatform, Branch.BaseEngineProject, CookedPlatform));
-
-                                                    string BuildAgentSharingGroup = NonCodeProject.GameName + "_MakeFormalBuild_" + Plat.ToString() + HostPlatformNode.StaticGetHostPlatformSuffix(HostPlatform);
-                                                    if (Plat == UnrealTargetPlatform.IOS || Plat == UnrealTargetPlatform.Android) // These trash build products, so we need to use different agents
-                                                    {
-                                                        BuildAgentSharingGroup = "";
-                                                    }
-                                                    GUBPNodes[CookNode.StaticGetFullName(HostPlatform, NonCodeProject, CookedPlatform)].AgentSharingGroup = BuildAgentSharingGroup;
-                                                    GUBPNodes[GamePlatformCookedAndCompiledNode.StaticGetFullName(HostPlatform, NonCodeProject, Plat)].AgentSharingGroup = BuildAgentSharingGroup;
-                                                    GUBPNodes[NodeName].AgentSharingGroup = BuildAgentSharingGroup;
-
-                                                    if (PlatPair.bTest)
-                                                    {
-                                                        AddNode(new FormalBuildTestNode(this, NonCodeProject, HostPlatform, Plat, PlatPair.TargetConfig));
-                                                    }
-                                                }
+													if (Options.bSeparateGamePromotion)
+													{
+														var NodeName = AddNode(new FormalBuildNode(this, NonCodeProject, HostPlatform, new List<UnrealTargetPlatform>() { Plat }, new List<UnrealTargetConfiguration>() { PlatPair.TargetConfig }));
+														// we don't want this delayed
+														// this would normally wait for the testing phase, we just want to build it right away
+														RemovePseudodependencyFromNode(
+															CookNode.StaticGetFullName(HostPlatform, NonCodeProject, CookedPlatform),
+															WaitForTestShared.StaticGetFullName());
+														RemovePseudodependencyFromNode(
+															CookNode.StaticGetFullName(HostPlatform, NonCodeProject, CookedPlatform),
+															CookNode.StaticGetFullName(HostPlatform, Branch.BaseEngineProject, CookedPlatform));
+
+														string BuildAgentSharingGroup = NonCodeProject.GameName + "_MakeFormalBuild_" + Plat.ToString() + HostPlatformNode.StaticGetHostPlatformSuffix(HostPlatform);
+														if (Plat == UnrealTargetPlatform.IOS || Plat == UnrealTargetPlatform.Android) // These trash build products, so we need to use different agents
+														{
+															BuildAgentSharingGroup = "";
+														}
+														GUBPNodes[CookNode.StaticGetFullName(HostPlatform, NonCodeProject, CookedPlatform)].AgentSharingGroup = BuildAgentSharingGroup;
+														GUBPNodes[GamePlatformCookedAndCompiledNode.StaticGetFullName(HostPlatform, NonCodeProject, Plat)].AgentSharingGroup = BuildAgentSharingGroup;
+														GUBPNodes[NodeName].AgentSharingGroup = BuildAgentSharingGroup;
+
+														if (PlatPair.bTest)
+														{
+															AddNode(new FormalBuildTestNode(this, NonCodeProject, HostPlatform, Plat, PlatPair.TargetConfig));
+														}
+													}
+													else
+													{														
+														AddNode(new FormalBuildNode(this, NonCodeProject, HostPlatform, new List<UnrealTargetPlatform>() { Plat }, new List<UnrealTargetConfiguration>() { PlatPair.TargetConfig }));
+													}
+                                                }												
                                             }
-                                    }
+										}
                                     }
                                     if (!bNoAutomatedTesting)
                                     {
@@ -5288,15 +5307,15 @@ public class GUBP : BuildCommand
                                     {
                                         // we don't want this delayed
                                         // this would normally wait for the testing phase, we just want to build it right away
-                                        RemovePseudodependencyFromNode(
-                                            CookNode.StaticGetFullName(HostPlatform, CodeProj, CookedPlatform),
-                                            WaitForTestShared.StaticGetFullName());
-                                        RemovePseudodependencyFromNode(
-                                            CookNode.StaticGetFullName(HostPlatform, CodeProj, CookedPlatform),
-                                            CookNode.StaticGetFullName(HostPlatform, Branch.BaseEngineProject, CookedPlatform));
+										RemovePseudodependencyFromNode(
+											CookNode.StaticGetFullName(HostPlatform, CodeProj, CookedPlatform),
+											WaitForTestShared.StaticGetFullName());
+										RemovePseudodependencyFromNode(
+											CookNode.StaticGetFullName(HostPlatform, CodeProj, CookedPlatform),
+											CookNode.StaticGetFullName(HostPlatform, Branch.BaseEngineProject, CookedPlatform));
 
                                         string BuildAgentSharingGroup = CodeProj.GameName + "_MakeFormalBuild_" + Plat.ToString() + HostPlatformNode.StaticGetHostPlatformSuffix(HostPlatform);
-                                        if (Plat == UnrealTargetPlatform.IOS || Plat == UnrealTargetPlatform.Android || Plat == UnrealTargetPlatform.XboxOne) // These trash build products, so we need to use different agents
+                                        if (Plat == UnrealTargetPlatform.IOS || Plat == UnrealTargetPlatform.Android || Plat == UnrealTargetPlatform.XboxOne || !Options.bSeparateGamePromotion) // These trash build products, so we need to use different agents
                                         {
                                             BuildAgentSharingGroup = "";
                                         }
@@ -5980,25 +5999,26 @@ if (HostPlatform == UnrealTargetPlatform.Mac) continue; //temp hack till mac aut
             bool HitNonSticky = false;
             bool bHaveECNodes = false;
             // sticky nodes are ones that we run on the main agent. We run then first and they must not be intermixed with parallel jobs
-            foreach (var NodeToDo in OrdereredToDo)
-            {
-                if (GUBPNodes[NodeToDo].RunInEC() && !NodeIsAlreadyComplete(NodeToDo, LocalOnly)) // if something is already finished, we don't put it into EC
-                {
-                    bHaveECNodes = true;
-                    if (GUBPNodes[NodeToDo].IsSticky())
-                    {
-                        LastSticky = NodeToDo;
-                        if (HitNonSticky && !bSkipTriggers)
-                        {
-                            throw new AutomationException("Sticky and non-sticky jobs did not sort right.");
-                        }
-                    }
-                    else
-                    {
-                        HitNonSticky = true;
-                    }
-                }
-            }
+			foreach (var NodeToDo in OrdereredToDo)
+			{
+				if (GUBPNodes[NodeToDo].RunInEC() && !NodeIsAlreadyComplete(NodeToDo, LocalOnly)) // if something is already finished, we don't put it into EC
+				{
+					bHaveECNodes = true;
+					if (GUBPNodes[NodeToDo].IsSticky())
+					{
+						LastSticky = NodeToDo;
+						if (HitNonSticky && !bSkipTriggers)
+						{
+							throw new AutomationException("Sticky and non-sticky jobs did not sort right.");
+						}
+					}
+					else
+					{
+						HitNonSticky = true;
+					}
+				}
+			}
+			
             string ParentPath = ParseParamValue("ParentPath");
             string BaseArgs = String.Format("createJobStep --parentPath {0}", ParentPath);
 
