Forcibly rebuild APIDocTool and UnrealDocTool targets before building documentation. Prevents issues with stale build products.

diff --git a/Engine/Source/Programs/AutomationTool/Scripts/BuildDocumentation.Automation.cs b/Engine/Source/Programs/AutomationTool/Scripts/BuildDocumentation.Automation.cs
--- a/Engine/Source/Programs/AutomationTool/Scripts/BuildDocumentation.Automation.cs
+++ b/Engine/Source/Programs/AutomationTool/Scripts/BuildDocumentation.Automation.cs
@@ -66,10 +66,10 @@ public class BuildDocumentation : BuildCommand
 		MsBuild(CmdEnv, UnrealBuildToolProject, "/verbosity:normal /target:Rebuild /property:Configuration=Development /property:Platform=AnyCPU", "BuildUnrealBuildTool");
 
 		// Compile APIDocTool
-		BuildSolution(CmdEnv, ApiDocToolSolution, "Release|x64", "BuildApiDocTool");
+		RunAndLog(CmdEnv, CmdEnv.MsDevExe, String.Format("\"{0}\" /rebuild \"Release|x64\"", ApiDocToolSolution), "BuildApiDocTool");
 
 		// Compile UnrealDocTool
-		BuildSolution(CmdEnv, UnrealDocToolSolution, "Release", "BuildUnrealDocTool");
+		RunAndLog(CmdEnv, CmdEnv.MsDevExe, String.Format("\"{0}\" /rebuild \"Release|Any CPU\"", UnrealDocToolSolution), "BuildUnrealDocTool");
 
 		// Prepare the shared doc tool command line, containing all the paths
 		string ApiToolCommandLine = "";
@@ -77,7 +77,6 @@ public class BuildDocumentation : BuildCommand
 		ApiToolCommandLine += " -targetinfo=\"" + TargetInfoPath + "\"";
 		ApiToolCommandLine += " -xmldir=\"" + DoxygenDir + "\"";
 		ApiToolCommandLine += " -metadatadir=\"" + MetadataDir + "\"";
-		ApiToolCommandLine += " -verbose";
 		if (Filter != null) ApiToolCommandLine += " -filter=" + Filter;
 
 		// Execute the clean
