Platform: Delete UAT merged copy of iPhonePackager as it is currently unused and already out of date
- Will be reintroduced once IPP standalone work is finished, probably as a shared assembly

diff --git a/Engine/Source/Programs/AutomationTool/IOS/IOS.Automation.csproj b/Engine/Source/Programs/AutomationTool/IOS/IOS.Automation.csproj
--- a/Engine/Source/Programs/AutomationTool/IOS/IOS.Automation.csproj
+++ b/Engine/Source/Programs/AutomationTool/IOS/IOS.Automation.csproj
@@ -60,16 +60,6 @@
       <Link>Properties\MetaData.cs</Link>
     </Compile>
     <Compile Include="IOSPlatform.Automation.cs" />
-    <Compile Include="IPhonePackager\CodeSigning.cs" />
-    <Compile Include="IPhonePackager\CompileTime.cs" />
-    <Compile Include="IPhonePackager\Config.cs" />
-    <Compile Include="IPhonePackager\CookTime.cs" />
-    <Compile Include="IPhonePackager\DeployTime.cs" />
-    <Compile Include="IPhonePackager\FileOperations.cs" />
-    <Compile Include="IPhonePackager\iPhonePackager.cs" />
-    <Compile Include="IPhonePackager\MachObjects.cs" />
-    <Compile Include="IPhonePackager\MobileProvisionUtilities.cs" />
-    <Compile Include="IPhonePackager\Utilities.cs" />
     <Compile Include="Properties\AssemblyInfo.cs" />
     <Compile Include="Properties\Resources.Designer.cs" />
   </ItemGroup>
diff --git a/Engine/Source/Programs/AutomationTool/IOS/IOSPlatform.Automation.cs b/Engine/Source/Programs/AutomationTool/IOS/IOSPlatform.Automation.cs
--- a/Engine/Source/Programs/AutomationTool/IOS/IOSPlatform.Automation.cs
+++ b/Engine/Source/Programs/AutomationTool/IOS/IOSPlatform.Automation.cs
@@ -67,7 +67,7 @@ public class IOSPlatform : Platform
 
 	protected static int RunIPP(string[] Args)
 	{
-		return iPhonePackager.Program.RealMain(Args);
+		return 5;//@TODO: Reintegrate IPP to IOS.Automation iPhonePackager.Program.RealMain(Args);
 	}
 
 	public override int RunCommand(string Command)
diff --git a/Engine/Source/Programs/AutomationTool/IOS/IPhonePackager/CodeSigning.cs b/Engine/Source/Programs/AutomationTool/IOS/IPhonePackager/CodeSigning.cs
deleted file mode 100644
diff --git a/Engine/Source/Programs/AutomationTool/IOS/IPhonePackager/CompileTime.cs b/Engine/Source/Programs/AutomationTool/IOS/IPhonePackager/CompileTime.cs
deleted file mode 100644
diff --git a/Engine/Source/Programs/AutomationTool/IOS/IPhonePackager/Config.cs b/Engine/Source/Programs/AutomationTool/IOS/IPhonePackager/Config.cs
deleted file mode 100644
diff --git a/Engine/Source/Programs/AutomationTool/IOS/IPhonePackager/CookTime.cs b/Engine/Source/Programs/AutomationTool/IOS/IPhonePackager/CookTime.cs
deleted file mode 100644
diff --git a/Engine/Source/Programs/AutomationTool/IOS/IPhonePackager/DeployTime.cs b/Engine/Source/Programs/AutomationTool/IOS/IPhonePackager/DeployTime.cs
deleted file mode 100644
diff --git a/Engine/Source/Programs/AutomationTool/IOS/IPhonePackager/FileOperations.cs b/Engine/Source/Programs/AutomationTool/IOS/IPhonePackager/FileOperations.cs
deleted file mode 100644
diff --git a/Engine/Source/Programs/AutomationTool/IOS/IPhonePackager/MachObjects.cs b/Engine/Source/Programs/AutomationTool/IOS/IPhonePackager/MachObjects.cs
deleted file mode 100644
diff --git a/Engine/Source/Programs/AutomationTool/IOS/IPhonePackager/MobileProvisionUtilities.cs b/Engine/Source/Programs/AutomationTool/IOS/IPhonePackager/MobileProvisionUtilities.cs
deleted file mode 100644
diff --git a/Engine/Source/Programs/AutomationTool/IOS/IPhonePackager/Utilities.cs b/Engine/Source/Programs/AutomationTool/IOS/IPhonePackager/Utilities.cs
deleted file mode 100644
diff --git a/Engine/Source/Programs/AutomationTool/IOS/IPhonePackager/iPhonePackager.cs b/Engine/Source/Programs/AutomationTool/IOS/IPhonePackager/iPhonePackager.cs
deleted file mode 100644
