- Fixed apk naming issue

diff --git a/Engine/Source/Programs/AutomationTool/Android/AndroidPlatform.Automation.cs b/Engine/Source/Programs/AutomationTool/Android/AndroidPlatform.Automation.cs
--- a/Engine/Source/Programs/AutomationTool/Android/AndroidPlatform.Automation.cs
+++ b/Engine/Source/Programs/AutomationTool/Android/AndroidPlatform.Automation.cs
@@ -181,27 +181,33 @@ public class AndroidPlatform : Platform
 			throw new AutomationException("Android is currently only able to package one target configuration at a time, but StageTargetConfigurations contained {0} configurations", SC.StageTargetConfigurations.Count);
 		}
 
-		string ApkName = GetFinalApkName(Params, SC.StageExecutables[0], true, "ERROR");
-		string ObbName = GetFinalObbName(ApkName);
-		string BatchName = GetFinalBatchName(ApkName, Params, "ERROR");
+		string[] Architectures = UnrealBuildTool.AndroidToolChain.GetAllArchitectures();
+		bool bMakeSeparateApks = UnrealBuildTool.Android.UEDeployAndroid.ShouldMakeSeparateApks();
 
-		// verify the files exist
-		if (!FileExists(ApkName))
-		{
-			throw new AutomationException("ARCHIVE FAILED - {0} was not found", ApkName);
-		}
-		if (!Params.OBBinAPK && !FileExists(ObbName))
+		foreach (string Architecture in Architectures)
 		{
-			throw new AutomationException("ARCHIVE FAILED - {0} was not found", ObbName);
-		}
+			string ApkName = GetFinalApkName(Params, SC.StageExecutables[0], true, bMakeSeparateApks ? Architecture : "");
+			string ObbName = GetFinalObbName(ApkName);
+			string BatchName = GetFinalBatchName(ApkName, Params, bMakeSeparateApks ? Architecture : "");
 
-		SC.ArchiveFiles(Path.GetDirectoryName(ApkName), Path.GetFileName(ApkName));
-		if (!Params.OBBinAPK)
-		{
-			SC.ArchiveFiles(Path.GetDirectoryName(ObbName), Path.GetFileName(ObbName));
-		}
+			// verify the files exist
+			if (!FileExists(ApkName))
+			{
+				throw new AutomationException("ARCHIVE FAILED - {0} was not found", ApkName);
+			}
+			if (!Params.OBBinAPK && !FileExists(ObbName))
+			{
+				throw new AutomationException("ARCHIVE FAILED - {0} was not found", ObbName);
+			}
+
+			SC.ArchiveFiles(Path.GetDirectoryName(ApkName), Path.GetFileName(ApkName));
+			if (!Params.OBBinAPK)
+			{
+				SC.ArchiveFiles(Path.GetDirectoryName(ObbName), Path.GetFileName(ObbName));
+			}
 
-		SC.ArchiveFiles(Path.GetDirectoryName(BatchName), Path.GetFileName(BatchName));
+			SC.ArchiveFiles(Path.GetDirectoryName(BatchName), Path.GetFileName(BatchName));
+		}
 	}
 
 	private string GetAdbCommand(ProjectParams Params)
@@ -299,8 +305,11 @@ public class AndroidPlatform : Platform
 
 	public override void Deploy(ProjectParams Params, DeploymentContext SC)
 	{
+		// @todo android: Get this from the connected device we are running on if separate Apks!
+		string DeviceArchitecture = "";
+
 		string AdbCommand = GetAdbCommand(Params);
-		string ApkName = GetFinalApkName(Params, SC.StageExecutables[0], true, "ERROR");
+		string ApkName = GetFinalApkName(Params, SC.StageExecutables[0], true, DeviceArchitecture);
 
 		// make sure APK is up to date (this is fast if so)
 		var Deploy = UEBuildDeploy.GetBuildDeploy(UnrealTargetPlatform.Android);
@@ -549,7 +558,9 @@ public class AndroidPlatform : Platform
 
 	public override ProcessResult RunClient(ERunOptions ClientRunFlags, string ClientApp, string ClientCmdLine, ProjectParams Params)
 	{
-		string DeviceArchitecture = "ERROR";
+		// @todo android: Get this from the connected device we are running on if separate Apks!
+		string DeviceArchitecture = "";
+
 		string ApkName = ClientApp + DeviceArchitecture + ".apk";
 		if (!File.Exists(ApkName))
 		{
diff --git a/Engine/Source/Programs/UnrealBuildTool/Android/UEDeployAndroid.cs b/Engine/Source/Programs/UnrealBuildTool/Android/UEDeployAndroid.cs
--- a/Engine/Source/Programs/UnrealBuildTool/Android/UEDeployAndroid.cs
+++ b/Engine/Source/Programs/UnrealBuildTool/Android/UEDeployAndroid.cs
@@ -311,9 +311,9 @@ namespace UnrealBuildTool.Android
 			foreach (string Arch in Arches)
 			{
 				string SourceSOName = AndroidToolChain.InlineArchName(OutputPath, Arch);
-				string DestApkName = Path.Combine(ProjectDirectory, "Binaries/Android/") + Path.GetFileNameWithoutExtension(OutputPath) + ".apk";
 				// if the source binary was UE4Game, replace it with the new project name, when re-packaging a binary only build
-				DestApkName = DestApkName.Replace("UE4Game-", ProjectName + "-");
+				string ApkFilename = Path.GetFileNameWithoutExtension(OutputPath).Replace("UE4Game", ProjectName);
+				string DestApkName = Path.Combine(ProjectDirectory, "Binaries/Android/") + ApkFilename + ".apk";
 
 				// if we making multiple Apks, we need to put the architecture into the name
 				if (bMakeSeparateApks)
@@ -410,9 +410,9 @@ namespace UnrealBuildTool.Android
 				string Arch = Arches[ArchIndex];
 
 				string SourceSOName = AndroidToolChain.InlineArchName(OutputPath, Arch);
-				string DestApkName = Path.Combine(ProjectDirectory, "Binaries/Android/") + Path.GetFileNameWithoutExtension(OutputPath) + ".apk";
 				// if the source binary was UE4Game, replace it with the new project name, when re-packaging a binary only build
-				DestApkName = DestApkName.Replace("UE4Game-", ProjectName + "-");
+				string ApkFilename = Path.GetFileNameWithoutExtension(OutputPath).Replace("UE4Game", ProjectName);
+				string DestApkName = Path.Combine(ProjectDirectory, "Binaries/Android/") + ApkFilename + ".apk";
 
 				// if we making multiple Apks, we need to put the architecture into the name
 				if (bMakeSeparateApks)
