buildscript, batch perl jobsteps

diff --git a/Engine/Source/Programs/AutomationTool/Scripts/GUBP.Automation.cs b/Engine/Source/Programs/AutomationTool/Scripts/GUBP.Automation.cs
--- a/Engine/Source/Programs/AutomationTool/Scripts/GUBP.Automation.cs
+++ b/Engine/Source/Programs/AutomationTool/Scripts/GUBP.Automation.cs
@@ -455,10 +455,10 @@ public class GUBP : BuildCommand
         }
         public virtual string ECProcedureParams()
         {
-            var Result = String.Format("--actualParameter Sticky={0}", IsSticky() ? 1 : 0);
+            var Result = String.Format(", {{actualParameterName => 'Sticky', value => '{0}'}}", IsSticky() ? 1 : 0);
             if (AgentSharingGroup != "")
             {
-                Result += String.Format(" --actualParameter AgentSharingGroup={0}", AgentSharingGroup);
+                Result += String.Format(", {{actualParameterName => 'AgentSharingGroup', value => '{0}'}}", AgentSharingGroup);
             }
             return Result;
         }
@@ -1819,7 +1819,8 @@ public class GUBP : BuildCommand
             var Result = base.ECProcedureParams();
             if (!bTriggerWasTriggered)
             {
-                Result += String.Format(" --actualParameter TriggerState={0} --actualParameter ActionText=\"{1}\" --actualParameter DescText=\"{2}\"", GetTriggerStateName(), GetTriggerActionText(), GetTriggerDescText());
+                Result += String.Format(", {{actualParameterName => 'TriggerState', value => '{0}'}}, {{actualParameterName => 'ActionText', value =>\"{1}\"}}, {{actualParameterName => 'DescText', value =>\"{2}\"}}", GetTriggerStateName(), GetTriggerActionText(), GetTriggerDescText());
+                //Result += String.Format(" --actualParameter TriggerState={0} --actualParameter ActionText=\"{1}\" --actualParameter DescText=\"{2}\"", GetTriggerStateName(), GetTriggerActionText(), GetTriggerDescText());
             }
             return Result;
         }
@@ -2460,7 +2461,7 @@ public class GUBP : BuildCommand
         List<UnrealTargetConfiguration> ClientConfigs;
         List<UnrealTargetConfiguration> ServerConfigs;
         bool ClientNotGame;
-        UnrealBuildTool.TargetRules.TargetType GameOrClient;		
+        UnrealBuildTool.TargetRules.TargetType GameOrClient;
 
         public FormalBuildNode(GUBP bp,
             BranchInfo.BranchUProject InGameProj,
@@ -2469,7 +2470,7 @@ public class GUBP : BuildCommand
             List<UnrealTargetConfiguration> InClientConfigs = null,
             List<UnrealTargetPlatform> InServerTargetPlatforms = null,
             List<UnrealTargetConfiguration> InServerConfigs = null,
-            bool InClientNotGame = false			
+            bool InClientNotGame = false
             )
             : base(InHostPlatform)
         {
@@ -2478,7 +2479,7 @@ public class GUBP : BuildCommand
             ServerTargetPlatforms = InServerTargetPlatforms;
             ClientConfigs = InClientConfigs;
             ServerConfigs = InServerConfigs;
-            ClientNotGame = InClientNotGame;			
+            ClientNotGame = InClientNotGame;
 
             GameOrClient = TargetRules.TargetType.Game;
 
@@ -2571,8 +2572,8 @@ public class GUBP : BuildCommand
             // add dependencies for cooked and compiled
             foreach (var Plat in AllTargetPlatforms)
             {
-					AddDependency(GamePlatformCookedAndCompiledNode.StaticGetFullName(HostPlatform, GameProj, Plat));
-			}
+                AddDependency(GamePlatformCookedAndCompiledNode.StaticGetFullName(HostPlatform, GameProj, Plat));
+            }
         }
 
         public static string StaticGetFullName(BranchInfo.BranchUProject InGameProj, UnrealTargetPlatform InHostPlatform, List<UnrealTargetPlatform> InClientTargetPlatforms = null, List<UnrealTargetConfiguration> InClientConfigs = null, List<UnrealTargetPlatform> InServerTargetPlatforms = null, List<UnrealTargetConfiguration> InServerConfigs = null, bool InClientNotGame = false)
@@ -2636,7 +2637,7 @@ public class GUBP : BuildCommand
         }      
         public override float Priority()
         {
-			return base.Priority() + 20.0f;
+            return base.Priority() + 20.0f;
         }
         public override int CISFrequencyQuantumShift(GUBP bp)
         {
@@ -3495,6 +3496,12 @@ public class GUBP : BuildCommand
             return RunAndLog("ectool", Args, Options: Opts);
         }
     }
+    void WriteECPerl(List<string> Args)
+    {
+        Args.Add("$batch->submit();");
+        string ECPerlFile = CommandUtils.CombinePaths(CommandUtils.CmdEnv.LogFolder, "jobsteps.pl");
+        WriteAllLines_NoExceptions(ECPerlFile, Args.ToArray());
+    }
     string GetEMailListForNode(GUBP bp, string NodeToDo, string Causers)
     {        
         var BranchForEmail = "";
@@ -5121,37 +5128,37 @@ public class GUBP : BuildCommand
                                                 {
 													if (Options.bSeparateGamePromotion)
 													{
-														var NodeName = AddNode(new FormalBuildNode(this, NonCodeProject, HostPlatform, new List<UnrealTargetPlatform>() { Plat }, new List<UnrealTargetConfiguration>() { PlatPair.TargetConfig }));
-														// we don't want this delayed
-														// this would normally wait for the testing phase, we just want to build it right away
-														RemovePseudodependencyFromNode(
-															CookNode.StaticGetFullName(HostPlatform, NonCodeProject, CookedPlatform),
-															WaitForTestShared.StaticGetFullName());
-														RemovePseudodependencyFromNode(
-															CookNode.StaticGetFullName(HostPlatform, NonCodeProject, CookedPlatform),
-															CookNode.StaticGetFullName(HostPlatform, Branch.BaseEngineProject, CookedPlatform));
-
-														string BuildAgentSharingGroup = NonCodeProject.GameName + "_MakeFormalBuild_" + Plat.ToString() + HostPlatformNode.StaticGetHostPlatformSuffix(HostPlatform);
-														if (Plat == UnrealTargetPlatform.IOS || Plat == UnrealTargetPlatform.Android) // These trash build products, so we need to use different agents
-														{
-															BuildAgentSharingGroup = "";
-														}
-														GUBPNodes[CookNode.StaticGetFullName(HostPlatform, NonCodeProject, CookedPlatform)].AgentSharingGroup = BuildAgentSharingGroup;
-														GUBPNodes[GamePlatformCookedAndCompiledNode.StaticGetFullName(HostPlatform, NonCodeProject, Plat)].AgentSharingGroup = BuildAgentSharingGroup;
-														GUBPNodes[NodeName].AgentSharingGroup = BuildAgentSharingGroup;
-
-														if (PlatPair.bTest)
-														{
-															AddNode(new FormalBuildTestNode(this, NonCodeProject, HostPlatform, Plat, PlatPair.TargetConfig));
-														}
-													}
+                                                    var NodeName = AddNode(new FormalBuildNode(this, NonCodeProject, HostPlatform, new List<UnrealTargetPlatform>() { Plat }, new List<UnrealTargetConfiguration>() {PlatPair.TargetConfig}));
+                                                    // we don't want this delayed
+                                                    // this would normally wait for the testing phase, we just want to build it right away
+                                                    RemovePseudodependencyFromNode(
+                                                        CookNode.StaticGetFullName(HostPlatform, NonCodeProject, CookedPlatform),
+                                                        WaitForTestShared.StaticGetFullName());
+                                                    RemovePseudodependencyFromNode(
+                                                        CookNode.StaticGetFullName(HostPlatform, NonCodeProject, CookedPlatform),
+                                                        CookNode.StaticGetFullName(HostPlatform, Branch.BaseEngineProject, CookedPlatform));
+
+                                                    string BuildAgentSharingGroup = NonCodeProject.GameName + "_MakeFormalBuild_" + Plat.ToString() + HostPlatformNode.StaticGetHostPlatformSuffix(HostPlatform);
+                                                    if (Plat == UnrealTargetPlatform.IOS || Plat == UnrealTargetPlatform.Android) // These trash build products, so we need to use different agents
+                                                    {
+                                                        BuildAgentSharingGroup = "";
+                                                    }
+                                                    GUBPNodes[CookNode.StaticGetFullName(HostPlatform, NonCodeProject, CookedPlatform)].AgentSharingGroup = BuildAgentSharingGroup;
+                                                    GUBPNodes[GamePlatformCookedAndCompiledNode.StaticGetFullName(HostPlatform, NonCodeProject, Plat)].AgentSharingGroup = BuildAgentSharingGroup;
+                                                    GUBPNodes[NodeName].AgentSharingGroup = BuildAgentSharingGroup;
+
+                                                    if (PlatPair.bTest)
+                                                    {
+                                                        AddNode(new FormalBuildTestNode(this, NonCodeProject, HostPlatform, Plat, PlatPair.TargetConfig));
+                                                    }
+                                                }
 													else
 													{														
 														AddNode(new FormalBuildNode(this, NonCodeProject, HostPlatform, new List<UnrealTargetPlatform>() { Plat }, new List<UnrealTargetConfiguration>() { PlatPair.TargetConfig }));
 													}
-                                                }												
                                             }
-										}
+                                    }
+                                    }
                                     }
                                     if (!bNoAutomatedTesting)
                                     {
@@ -5322,12 +5329,12 @@ public class GUBP : BuildCommand
                                     {
                                         // we don't want this delayed
                                         // this would normally wait for the testing phase, we just want to build it right away
-										RemovePseudodependencyFromNode(
-											CookNode.StaticGetFullName(HostPlatform, CodeProj, CookedPlatform),
-											WaitForTestShared.StaticGetFullName());
-										RemovePseudodependencyFromNode(
-											CookNode.StaticGetFullName(HostPlatform, CodeProj, CookedPlatform),
-											CookNode.StaticGetFullName(HostPlatform, Branch.BaseEngineProject, CookedPlatform));
+                                        RemovePseudodependencyFromNode(
+                                            CookNode.StaticGetFullName(HostPlatform, CodeProj, CookedPlatform),
+                                            WaitForTestShared.StaticGetFullName());
+                                        RemovePseudodependencyFromNode(
+                                            CookNode.StaticGetFullName(HostPlatform, CodeProj, CookedPlatform),
+                                            CookNode.StaticGetFullName(HostPlatform, Branch.BaseEngineProject, CookedPlatform));
 
                                         string BuildAgentSharingGroup = CodeProj.GameName + "_MakeFormalBuild_" + Plat.ToString() + HostPlatformNode.StaticGetHostPlatformSuffix(HostPlatform);
                                         if (Plat == UnrealTargetPlatform.IOS || Plat == UnrealTargetPlatform.Android || Plat == UnrealTargetPlatform.XboxOne || !Options.bSeparateGamePromotion) // These trash build products, so we need to use different agents
@@ -5557,11 +5564,11 @@ if (HostPlatform == UnrealTargetPlatform.Mac) continue; //temp hack till mac aut
         bool bRelatedToNode = false;
 		bool bGraphSubset = false;
         var NodesToDo = new HashSet<string>();
-		
+
         {
-			string NodeSpec = ParseParamValue("Node");
+            string NodeSpec = ParseParamValue("Node");
             if (String.IsNullOrEmpty(NodeSpec))
-            {				
+            {
                 NodeSpec = ParseParamValue("RelatedToNode");
                 if (!String.IsNullOrEmpty(NodeSpec))
                 {
@@ -6017,36 +6024,43 @@ if (HostPlatform == UnrealTargetPlatform.Mac) continue; //temp hack till mac aut
             string LastSticky = "";
             bool HitNonSticky = false;
             bool bHaveECNodes = false;
+            List<string> StepList = new List<string>();
+            StepList.Add("use strict;");
+            StepList.Add("use diagnostics;");
+            StepList.Add("use ElectricCommander();");
+            StepList.Add("my $ec = new ElectricCommander;");
+            StepList.Add("$ec->setTimeout(600);");
+            StepList.Add("my $batch = $ec->newBatch(\"serial\");");
             // sticky nodes are ones that we run on the main agent. We run then first and they must not be intermixed with parallel jobs
-			foreach (var NodeToDo in OrdereredToDo)
-			{
-				if (GUBPNodes[NodeToDo].RunInEC() && !NodeIsAlreadyComplete(NodeToDo, LocalOnly)) // if something is already finished, we don't put it into EC
-				{
-					bHaveECNodes = true;
-					if (GUBPNodes[NodeToDo].IsSticky())
-					{
-						LastSticky = NodeToDo;
-						if (HitNonSticky && !bSkipTriggers)
-						{
-							throw new AutomationException("Sticky and non-sticky jobs did not sort right.");
-						}
-					}
-					else
-					{
-						HitNonSticky = true;
-					}
-				}
-			}
+            foreach (var NodeToDo in OrdereredToDo)
+            {
+                if (GUBPNodes[NodeToDo].RunInEC() && !NodeIsAlreadyComplete(NodeToDo, LocalOnly)) // if something is already finished, we don't put it into EC
+                {
+                    bHaveECNodes = true;
+                    if (GUBPNodes[NodeToDo].IsSticky())
+                    {
+                        LastSticky = NodeToDo;
+                        if (HitNonSticky && !bSkipTriggers)
+                        {
+                            throw new AutomationException("Sticky and non-sticky jobs did not sort right.");
+                        }
+                    }
+                    else
+                    {
+                        HitNonSticky = true;
+                    }
+                }
+            }
 			
             string ParentPath = ParseParamValue("ParentPath");
-            string BaseArgs = String.Format("createJobStep --parentPath {0}", ParentPath);
+            string BaseArgs = String.Format("$batch->createJobStep({{parentPath => '{0}'", ParentPath);
 
             bool bHasNoop = false;
             if (LastSticky == "" && bHaveECNodes)
             {
                 // if we don't have any sticky nodes and we have other nodes, we run a fake noop just to release the resource 
-                string Args = String.Format("{0} --subprocedure GUBP_UAT_Node --parallel 0 --jobStepName Noop --actualParameter NodeName=Noop --actualParameter Sticky=1 --releaseMode release", BaseArgs);
-                RunECTool(Args);
+                string Args = String.Format("{0}, subprocedure => 'GUBP_UAT_Node', parallel => '0', jobStepName => 'Noop', actualParameter => [{{actualParameterName => 'NodeName', value => 'Noop'}}, {{actualParametername => 'Sticky', value =>'1' }}], releaseMode => 'release'", BaseArgs);
+                StepList.Add(Args);
                 bHasNoop = true;
             }
 
@@ -6097,22 +6111,22 @@ if (HostPlatform == UnrealTargetPlatform.Mac) continue; //temp hack till mac aut
                     }
                     string Procedure = GUBPNodes[NodeToDo].ECProcedure();
 
-                    string Args = String.Format("{0} --subprocedure {1} --parallel {2} --jobStepName {3} --actualParameter NodeName={4}",
+                    string Args = String.Format("{0}, subprocedure => '{1}', parallel => '{2}', jobStepName => '{3}', actualParameter => [{{actualParameterName => 'NodeName', value =>'{4}'}}",
                         BaseArgs, Procedure, DoParallel ? 1 : 0, NodeToDo, NodeToDo);
                     string ProcedureParams = GUBPNodes[NodeToDo].ECProcedureParams();
                     if (!String.IsNullOrEmpty(ProcedureParams))
                     {
-                        Args = Args + " " + ProcedureParams;
+                        Args = Args + ProcedureParams;
                     }
 
                     if ((Procedure == "GUBP_UAT_Trigger" || Procedure == "GUBP_Hardcoded_Trigger") && !String.IsNullOrEmpty(EMails))
                     {
-                        Args = Args + " --actualParameter EmailsForTrigger=\"" + EMails + "\"";
+                        Args = Args + ", {actualParameterName => 'EmailsForTrigger', value => \'" + EMails + "\'}";
                     }
-
+                    Args = Args + "]";
                     string PreCondition = "";
                     string RunCondition = "";
-                    var UncompletedEcDeps = new List<string>();						
+                    var UncompletedEcDeps = new List<string>();
 
                     string MyAgentGroup = GUBPNodes[NodeToDo].AgentSharingGroup;
                     bool bDoNestedJobstep = false;
@@ -6122,7 +6136,7 @@ if (HostPlatform == UnrealTargetPlatform.Mac) continue; //temp hack till mac aut
                     string NodeParentPath = ParentPath;
 					string PreconditionParentPath;
 					if (GUBPNodes[NodeToDo].GetFullName().Contains("MakeBuild") && GUBPNodes[NodeToDo].FullNamesOfPseudosependencies.Contains(SharedLabelPromotableNode.StaticGetFullName(false)) && !bGraphSubset)
-					{
+                    {
 						RemovePseudodependencyFromNode(NodeToDo, SharedLabelPromotableNode.StaticGetFullName(false));
 						PreconditionParentPath = GetPropertyFromStep("/myWorkflow/ParentJob");
 						UncompletedEcDeps = GetECDependencies(NodeToDo);
@@ -6130,20 +6144,20 @@ if (HostPlatform == UnrealTargetPlatform.Mac) continue; //temp hack till mac aut
 					else
 					{
 						PreconditionParentPath = ParentPath;
-						var EcDeps = GetECDependencies(NodeToDo);
-						foreach (var Dep in EcDeps)
-						{
-							if (GUBPNodes[Dep].RunInEC() && !NodeIsAlreadyComplete(Dep, LocalOnly) && OrdereredToDo.Contains(Dep)) // if something is already finished, we don't put it into EC
-							{
-								if (OrdereredToDo.IndexOf(Dep) > OrdereredToDo.IndexOf(NodeToDo))
-								{
-									throw new AutomationException("Topological sort error, node {0} has a dependency of {1} which sorted after it.", NodeToDo, Dep);
-								}
-								UncompletedEcDeps.Add(Dep);
-							}
-						}
-					}
-					var PreConditionUncompletedEcDeps = UncompletedEcDeps;
+                        var EcDeps = GetECDependencies(NodeToDo);
+                        foreach (var Dep in EcDeps)
+                        {
+                            if (GUBPNodes[Dep].RunInEC() && !NodeIsAlreadyComplete(Dep, LocalOnly) && OrdereredToDo.Contains(Dep)) // if something is already finished, we don't put it into EC
+                            {
+                                if (OrdereredToDo.IndexOf(Dep) > OrdereredToDo.IndexOf(NodeToDo))
+                                {
+                                    throw new AutomationException("Topological sort error, node {0} has a dependency of {1} which sorted after it.", NodeToDo, Dep);
+                                }
+                                UncompletedEcDeps.Add(Dep);
+                            }
+                        }
+                    }
+                    var PreConditionUncompletedEcDeps = UncompletedEcDeps;
                     if (MyAgentGroup != "")
                     {
                         bDoNestedJobstep = true;
@@ -6187,12 +6201,12 @@ if (HostPlatform == UnrealTargetPlatform.Mac) continue; //temp hack till mac aut
                     }
                     if (PreConditionUncompletedEcDeps.Count > 0)
                     {
-                        PreCondition = "\"$[/javascript if(";
+                        PreCondition = "\"\\$\" . \"[/javascript if(";
                         // these run "parallel", but we add preconditions to serialize them
                         int Index = 0;
                         foreach (var Dep in PreConditionUncompletedEcDeps)
-                        {
-                            PreCondition = PreCondition + "getProperty('" + GetJobStep(PreconditionParentPath, Dep) + "/status') == 'completed'";
+                        {                            
+                            PreCondition = PreCondition + "getProperty('" + GetJobStep(PreconditionParentPath, Dep) + "/status\') == \'completed\'";
                             Index++;
                             if (Index != PreConditionUncompletedEcDeps.Count)
                             {
@@ -6204,12 +6218,12 @@ if (HostPlatform == UnrealTargetPlatform.Mac) continue; //temp hack till mac aut
 
                     if (UncompletedEcDeps.Count > 0)
                     {
-                        RunCondition = "\"$[/javascript if(";
+                        RunCondition = "\"\\$\" . \"[/javascript if(";
                         int Index = 0;
                         foreach (var Dep in UncompletedEcDeps)
                         {
-                            RunCondition = RunCondition + "('$[" + GetJobStep(PreconditionParentPath, Dep) + "/outcome]' == 'success' || ";
-                            RunCondition = RunCondition + "'$[" + GetJobStep(PreconditionParentPath, Dep) + "/outcome]' == 'warning')";
+							RunCondition = RunCondition + "\\(\'\\$\" . \"[" + GetJobStep(PreconditionParentPath, Dep) + "/outcome]\' == \'success\') || ";
+							RunCondition = RunCondition + "\\(\'\\$\" . \"[" + GetJobStep(PreconditionParentPath, Dep) + "/outcome]\' == \'warning\')";
 
                             Index++;
                             if (Index != UncompletedEcDeps.Count)
@@ -6225,53 +6239,56 @@ if (HostPlatform == UnrealTargetPlatform.Mac) continue; //temp hack till mac aut
                         if (bDoFirstNestedJobstep)
                         {
                             {
-                                string NestArgs = String.Format("createJobStep --parentPath {0} --jobStepName {1} --parallel 1",
+                                string NestArgs = String.Format("$batch->createJobStep({{parentPath => '{0}', jobStepName => '{1}', parallel => '1'",
                                     ParentPath, MyAgentGroup);
                                 if (!String.IsNullOrEmpty(PreCondition))
                                 {
-                                    NestArgs = NestArgs + " --precondition " + PreCondition;
+                                    NestArgs = NestArgs + ", precondition => " + PreCondition;
                                 }
-                                RunECTool(NestArgs);
+                                NestArgs = NestArgs + "});";
+                                StepList.Add(NestArgs);
                             }
                             {
-                                string NestArgs = String.Format("createJobStep --parentPath {0}/jobSteps[{1}] --jobStepName {2}_GetPool --subprocedure GUBP{3}_AgentShare_GetPool --parallel 1 --actualParameter AgentSharingGroup={4} --actualParameter NodeName={5}",
+                                string NestArgs = String.Format("$batch->createJobStep({{parentPath => '{0}/jobSteps[{1}]', jobStepName => '{2}_GetPool', subprocedure => 'GUBP{3}_AgentShare_GetPool', parallel => '1', actualParameter => [{{actualParameterName => 'AgentSharingGroup', value => '{4}'}}, {{actualParameterName => 'NodeName', value => '{5}'}}]",
                                     ParentPath, MyAgentGroup, MyAgentGroup, GUBPNodes[NodeToDo].ECProcedureInfix(), MyAgentGroup, NodeToDo);
                                 if (!String.IsNullOrEmpty(PreCondition))
                                 {
-                                    NestArgs = NestArgs + " --precondition " + PreCondition;
+                                    NestArgs = NestArgs + ", precondition => " + PreCondition;
                                 }
-                                RunECTool(NestArgs);
+                                NestArgs = NestArgs + "});";
+                                StepList.Add(NestArgs);
                             }
                             {
-                                string NestArgs = String.Format("createJobStep --parentPath {0}/jobSteps[{1}] --jobStepName {2}_GetAgent --subprocedure GUBP{3}_AgentShare_GetAgent --parallel 1 --exclusiveMode call --resourceName {4} --actualParameter AgentSharingGroup={5} --actualParameter NodeName={6}",
+                                string NestArgs = String.Format("$batch->createJobStep({{parentPath => '{0}/jobSteps[{1}]', jobStepName => '{2}_GetAgent', subprocedure => 'GUBP{3}_AgentShare_GetAgent', parallel => '1', exclusiveMode => 'call', resourceName => '{4}', actualParameter => [{{actualParameterName => 'AgentSharingGroup', value => '{5}'}}, {{actualParameterName => 'NodeName', value=> '{6}'}}]",
                                     ParentPath, MyAgentGroup, MyAgentGroup, GUBPNodes[NodeToDo].ECProcedureInfix(),
                                     String.Format("$[/myJob/jobSteps[{0}]/ResourcePool]", MyAgentGroup),
                                     MyAgentGroup, NodeToDo);
                                 {
-                                    NestArgs = NestArgs + " --precondition ";
-                                    NestArgs = NestArgs + "\"$[/javascript if(";
+                                    NestArgs = NestArgs + ", precondition  => ";
+                                    NestArgs = NestArgs + "\"\\$\" . \"[/javascript if(";
                                     NestArgs = NestArgs + "getProperty('" + PreconditionParentPath + "/jobSteps[" + MyAgentGroup + "]/jobSteps[" + MyAgentGroup + "_GetPool]/status') == 'completed'";
                                     NestArgs = NestArgs + ") true;]\"";
                                 }
-                                RunECTool(NestArgs);
+                                NestArgs = NestArgs + "});";
+                                StepList.Add(NestArgs);
                             }
                             {
-                                PreCondition = "\"$[/javascript if(";
+                                PreCondition = "\"\\$\" . \"[/javascript if(";
                                 PreCondition = PreCondition + "getProperty('" + PreconditionParentPath + "/jobSteps[" + MyAgentGroup + "]/jobSteps[" + MyAgentGroup + "_GetAgent]/status') == 'completed'";
                                 PreCondition = PreCondition + ") true;]\"";
                             }
                         }
-                        Args = Args.Replace(String.Format("--parentPath {0} ", ParentPath), String.Format("--parentPath {0} ", NodeParentPath));
+                        Args = Args.Replace(String.Format("parentPath => '{0}'", ParentPath), String.Format("parentPath => '{0}'", NodeParentPath));
                         Args = Args.Replace("UAT_Node_Parallel_AgentShare", "UAT_Node_Parallel_AgentShare3");
                     }
 
                     if (!String.IsNullOrEmpty(PreCondition))
                     {
-                        Args = Args + " --precondition " + PreCondition;
+                        Args = Args + ", precondition => " + PreCondition;
                     }
                     if (!String.IsNullOrEmpty(RunCondition))
                     {
-                        Args = Args + " --condition " + RunCondition;
+                        Args = Args + ", condition => " + RunCondition;
                     }
 #if false
                     // this doesn't work because it includes precondition time
@@ -6282,11 +6299,10 @@ if (HostPlatform == UnrealTargetPlatform.Mac) continue; //temp hack till mac aut
 #endif
                     if (Sticky && NodeToDo == LastSticky)
                     {
-                        Args = Args + " --releaseMode release";
+                        Args = Args + ", releaseMode => 'release'";
                     }
-
-
-                    RunECTool(Args);
+                    Args = Args + "});";
+                    StepList.Add(Args);
                     if (bFakeEC && 
                         !UnfinishedTriggers.Contains(NodeToDo) &&
                         (GUBPNodes[NodeToDo].ECProcedure().StartsWith("GUBP_UAT_Node") || GUBPNodes[NodeToDo].ECProcedure().StartsWith("GUBP_Mac_UAT_Node")) // other things we really can't test
@@ -6319,19 +6335,19 @@ if (HostPlatform == UnrealTargetPlatform.Mac) continue; //temp hack till mac aut
                         int MyIndex = MyChain.IndexOf(NodeToDo);
                         if (MyIndex == MyChain.Count - 1)
                         {
-                            var RelPreCondition = "\"$[/javascript if(";
+                            var RelPreCondition = "\"\\$\" . \"[/javascript if(";
                             // this runs "parallel", but we a precondition to serialize it
                             RelPreCondition = RelPreCondition + "getProperty('" + PreconditionParentPath + "/jobSteps[" + NodeToDo + "]/status') == 'completed'";
                             RelPreCondition = RelPreCondition + ") true;]\"";
                             // we need to release the resource
-                            string RelArgs = String.Format("{0} --subprocedure GUBP_Release_AgentShare --parallel 1 --jobStepName Release_{1} --actualParameter AgentSharingGroup={2} --releaseMode release --precondition {3}",
+                            string RelArgs = String.Format("{0}, subprocedure => 'GUBP_Release_AgentShare', parallel => '1', jobStepName => 'Release_{1}', actualParameter => [{{actualParameterName => 'AgentSharingGroup', valued => '{2}'}}], releaseMode => 'release', precondition => '{3}'",
                                 BaseArgs, MyAgentGroup, MyAgentGroup, RelPreCondition);
-
-                            RunECTool(RelArgs);
+                            StepList.Add(RelArgs);
                         }
                     }
                 }
             }
+            WriteECPerl(StepList);
             ECProps.Add(String.Format("HasTests={0}", bHasTests));
             {
                 ECProps.Add("GUBP_LoadedProps=1");
