Fixed bug with launching from uat when using an already started cook on the fly server.
Should use default port.

diff --git a/Engine/Source/Programs/AutomationTool/Scripts/CopyBuildToStagingDirectory.Automation.cs b/Engine/Source/Programs/AutomationTool/Scripts/CopyBuildToStagingDirectory.Automation.cs
--- a/Engine/Source/Programs/AutomationTool/Scripts/CopyBuildToStagingDirectory.Automation.cs
+++ b/Engine/Source/Programs/AutomationTool/Scripts/CopyBuildToStagingDirectory.Automation.cs
@@ -708,29 +708,34 @@ public partial class Project : CommandUtils
 							{
 								if (IP.UnicastAddresses[Index].IsDnsEligible && IP.UnicastAddresses[Index].Address.AddressFamily == System.Net.Sockets.AddressFamily.InterNetwork)
 								{
-									if (Params.Port != null)
-									{
-										foreach (var Port in Params.Port)
-										{
-											if (!FirstParam)
-											{
-												FileHostParams += "+";
-											}
-											FirstParam = false;
-											string[] PortProtocol = Port.Split(new char[] { ':' });
-											if (PortProtocol.Length > 1)
-											{
-												FileHostParams += String.Format("{0}://{1}:{2}", PortProtocol[0], IP.UnicastAddresses[Index].Address.ToString(), PortProtocol[1]);
-											}
-											else
-											{
-												FileHostParams += IP.UnicastAddresses[Index].Address.ToString();
-												FileHostParams += ":";
-												FileHostParams += Params.Port;
-											}
-
-										}
-									}
+                                    if (Params.Port != null)
+                                    {
+                                        foreach (var Port in Params.Port)
+                                        {
+                                            if (!FirstParam)
+                                            {
+                                                FileHostParams += "+";
+                                            }
+                                            FirstParam = false;
+                                            string[] PortProtocol = Port.Split(new char[] { ':' });
+                                            if (PortProtocol.Length > 1)
+                                            {
+                                                FileHostParams += String.Format("{0}://{1}:{2}", PortProtocol[0], IP.UnicastAddresses[Index].Address.ToString(), PortProtocol[1]);
+                                            }
+                                            else
+                                            {
+                                                FileHostParams += IP.UnicastAddresses[Index].Address.ToString();
+                                                FileHostParams += ":";
+                                                FileHostParams += Params.Port;
+                                            }
+
+                                        }
+                                    }
+                                    else
+                                    {
+                                        // use default port
+                                        FileHostParams += IP.UnicastAddresses[Index].Address.ToString();
+                                    }
 
 								}
 							}
@@ -771,6 +776,11 @@ public partial class Project : CommandUtils
 											}
 										}
 									}
+                                    else
+                                    {
+                                        // use default port
+                                        FileHostParams += IP.UnicastAddresses[Index].Address.ToString();
+                                    }
 
 								}
 							}
@@ -801,6 +811,11 @@ public partial class Project : CommandUtils
 
 					}
 				}
+                else
+                {
+                    // use default port
+                    FileHostParams += IP.UnicastAddresses[Index].Address.ToString();
+                }
 				FileHostParams += " ";
 			}
 
diff --git a/Engine/Source/Programs/AutomationTool/Scripts/RunProjectCommand.Automation.cs b/Engine/Source/Programs/AutomationTool/Scripts/RunProjectCommand.Automation.cs
--- a/Engine/Source/Programs/AutomationTool/Scripts/RunProjectCommand.Automation.cs
+++ b/Engine/Source/Programs/AutomationTool/Scripts/RunProjectCommand.Automation.cs
@@ -624,6 +624,11 @@ public partial class Project : CommandUtils
 											}
 										}
 									}
+                                    else
+                                    {
+                                        // use default port
+                                        FileHostParams += IP.UnicastAddresses[Index].Address.ToString();
+                                    }
 								}
 							}
 						}
@@ -663,6 +668,11 @@ public partial class Project : CommandUtils
 											}
 										}
 									}
+                                    else
+                                    {
+                                        // use default port
+                                        FileHostParams += IP.UnicastAddresses[Index].Address.ToString();
+                                    }
 								}
 							}
 						}
@@ -694,6 +704,11 @@ public partial class Project : CommandUtils
 
 					}
 				}
+                else
+                {
+                    // use default port
+                    FileHostParams += IP.UnicastAddresses[Index].Address.ToString();
+                }
 				TempCmdLine += " ";
 
 				if (!Params.Stage)
