Fix build compliation errors.

diff --git a/Engine/Source/Programs/AutomationTool/Scripts/CopyBuildToStagingDirectory.Automation.cs b/Engine/Source/Programs/AutomationTool/Scripts/CopyBuildToStagingDirectory.Automation.cs
--- a/Engine/Source/Programs/AutomationTool/Scripts/CopyBuildToStagingDirectory.Automation.cs
+++ b/Engine/Source/Programs/AutomationTool/Scripts/CopyBuildToStagingDirectory.Automation.cs
@@ -814,7 +814,7 @@ public partial class Project : CommandUtils
                 else
                 {
                     // use default port
-                    FileHostParams += IP.UnicastAddresses[Index].Address.ToString();
+                    FileHostParams += LocalHost;
                 }
 				FileHostParams += " ";
 			}
diff --git a/Engine/Source/Programs/AutomationTool/Scripts/RunProjectCommand.Automation.cs b/Engine/Source/Programs/AutomationTool/Scripts/RunProjectCommand.Automation.cs
--- a/Engine/Source/Programs/AutomationTool/Scripts/RunProjectCommand.Automation.cs
+++ b/Engine/Source/Programs/AutomationTool/Scripts/RunProjectCommand.Automation.cs
@@ -627,7 +627,7 @@ public partial class Project : CommandUtils
                                     else
                                     {
                                         // use default port
-                                        FileHostParams += IP.UnicastAddresses[Index].Address.ToString();
+                                        TempCmdLine += IP.UnicastAddresses[Index].Address.ToString();
                                     }
 								}
 							}
@@ -671,7 +671,7 @@ public partial class Project : CommandUtils
                                     else
                                     {
                                         // use default port
-                                        FileHostParams += IP.UnicastAddresses[Index].Address.ToString();
+                                        TempCmdLine += IP.UnicastAddresses[Index].Address.ToString();
                                     }
 								}
 							}
@@ -707,7 +707,7 @@ public partial class Project : CommandUtils
                 else
                 {
                     // use default port
-                    FileHostParams += IP.UnicastAddresses[Index].Address.ToString();
+                    TempCmdLine += LocalHost;
                 }
 				TempCmdLine += " ";
 
