Fix warning building ShooterGame; probing for whether a directory is writeable or not is not a failure.

diff --git a/Engine/Source/Programs/AutomationTool/CommandUtils.cs b/Engine/Source/Programs/AutomationTool/CommandUtils.cs
--- a/Engine/Source/Programs/AutomationTool/CommandUtils.cs
+++ b/Engine/Source/Programs/AutomationTool/CommandUtils.cs
@@ -2014,33 +2014,36 @@ namespace AutomationTool
             return StorageDirectory;
         }
 
-        static bool DirectoryExistsAndIsWritable_NoExceptions(string Dir)
+        public static bool DirectoryExistsAndIsWritable_NoExceptions(string Dir)
         {
-            try
+            if (!DirectoryExists_NoExceptions(Dir))
             {
-                if (!DirectoryExists_NoExceptions(Dir))
-                {
-                    return false;
-                }
-                var TestGUID = Guid.NewGuid();
-                var Filename = CombinePaths(Dir, TestGUID.ToString() + ".Temp.txt");
-                WriteAllText_NoExceptions(Filename, "Test");
-                if (FileExists_NoExceptions(true, Filename))
-                {
-                    DeleteFile_NoExceptions(Filename, true);
-                    //Log(System.Diagnostics.TraceEventType.Information, "Resolved shared dir {0}", Dir);
-                    return true;
-                }
-                Log(System.Diagnostics.TraceEventType.Warning, "Shared dir {0} is not writable", Dir);
+				Log(System.Diagnostics.TraceEventType.Information, "Directory {0} does not exist", Dir);
+				return false;
+			}
 
-            }
-            catch (Exception Ex)
+            try
             {
-                Log(System.Diagnostics.TraceEventType.Warning, "Failed to resolve shared dir {0}", Dir);
-                Log(System.Diagnostics.TraceEventType.Warning, LogUtils.FormatException(Ex));
-            }
-            return false;
-        }
+				string Filename = CombinePaths(Dir, Guid.NewGuid().ToString() + ".Temp.txt");
+				string NativeFilename = ConvertSeparators(PathSeparator.Default, Filename);
+				using(StreamWriter Writer = new StreamWriter(NativeFilename))
+				{
+					Writer.Write("Test");
+				}
+				if(File.Exists(NativeFilename))
+				{
+		            DeleteFile_NoExceptions(Filename, true);
+		            Log(System.Diagnostics.TraceEventType.Information, "Directory {0} is writable", Dir);
+					return true;
+				}
+			}
+			catch(Exception)
+			{
+			}
+
+			Log(System.Diagnostics.TraceEventType.Information, "Directory {0} is not writable", Dir);
+			return false;
+		}
 
         static Dictionary<string, string> ResolveCache = new Dictionary<string, string>();
         public static string ResolveSharedBuildDirectory(string GameFolder)
