Fix packaging for Linux TTP #344934.

diff --git a/Engine/Source/Programs/AutomationTool/Scripts/CopyBuildToStagingDirectory.Automation.cs b/Engine/Source/Programs/AutomationTool/Scripts/CopyBuildToStagingDirectory.Automation.cs
--- a/Engine/Source/Programs/AutomationTool/Scripts/CopyBuildToStagingDirectory.Automation.cs
+++ b/Engine/Source/Programs/AutomationTool/Scripts/CopyBuildToStagingDirectory.Automation.cs
@@ -287,9 +287,9 @@ public partial class Project : CommandUtils
 				SC.StageFiles(StagedFileType.NonUFS, CombinePaths(SC.LocalRoot, "Engine/Shaders/StandaloneRenderer"));
 
 				SC.StageFiles(StagedFileType.NonUFS, CombinePaths(SC.LocalRoot, "Engine/Content/Localization/ICU"));
+				// Linux platform stages ICU in GetFilesToDeployOrStage(), accounting for the actual architecture
 				if (SC.StageTargetPlatform.PlatformType == UnrealTargetPlatform.Win64 ||
 					SC.StageTargetPlatform.PlatformType == UnrealTargetPlatform.Win32 ||
-					SC.StageTargetPlatform.PlatformType == UnrealTargetPlatform.Linux ||
 					SC.StageTargetPlatform.PlatformType == UnrealTargetPlatform.Mac)
 				{
 					SC.StageFiles(StagedFileType.NonUFS, CombinePaths(SC.LocalRoot, "Engine/Binaries/ThirdParty/ICU"));
