build script, runcondition fix

diff --git a/Engine/Source/Programs/AutomationTool/Scripts/GUBP.Automation.cs b/Engine/Source/Programs/AutomationTool/Scripts/GUBP.Automation.cs
--- a/Engine/Source/Programs/AutomationTool/Scripts/GUBP.Automation.cs
+++ b/Engine/Source/Programs/AutomationTool/Scripts/GUBP.Automation.cs
@@ -6222,8 +6222,8 @@ if (HostPlatform == UnrealTargetPlatform.Mac) continue; //temp hack till mac aut
                         int Index = 0;
                         foreach (var Dep in UncompletedEcDeps)
                         {
-							RunCondition = RunCondition + "\\(\'\\$\" . \"[" + GetJobStep(PreconditionParentPath, Dep) + "/outcome]\' == \'success\') || ";
-							RunCondition = RunCondition + "\\(\'\\$\" . \"[" + GetJobStep(PreconditionParentPath, Dep) + "/outcome]\' == \'warning\')";
+							RunCondition = RunCondition + "(getProperty('" + GetJobStep(PreconditionParentPath, Dep) + "/outcome\') == 'success') || ";
+							RunCondition = RunCondition + "(getProperty('" + GetJobStep(PreconditionParentPath, Dep) + "/outcome\') == 'warning')";
 
                             Index++;
                             if (Index != UncompletedEcDeps.Count)
@@ -6231,7 +6231,7 @@ if (HostPlatform == UnrealTargetPlatform.Mac) continue; //temp hack till mac aut
                                 RunCondition = RunCondition + " && ";
                             }
                         }
-                        RunCondition = RunCondition + ") true; else false;]\"";
+                        RunCondition = RunCondition + ")true; else false;]\"";
                     }
 
                     if (bDoNestedJobstep)
