Rename the binary when packaging content-only projects for Linux.

No more UE4Game, the binary will have the project project name and will be placed in <Project>/Binaries/Linux/.

TTP #338841

diff --git a/Engine/Source/Programs/AutomationTool/Linux/LinuxPlatform.Automation.cs b/Engine/Source/Programs/AutomationTool/Linux/LinuxPlatform.Automation.cs
--- a/Engine/Source/Programs/AutomationTool/Linux/LinuxPlatform.Automation.cs
+++ b/Engine/Source/Programs/AutomationTool/Linux/LinuxPlatform.Automation.cs
@@ -65,7 +65,7 @@ public abstract class BaseLinuxPlatform : Platform
 			if (Exe.StartsWith(CombinePaths(SC.RuntimeProjectRootDir, "Binaries", SC.PlatformDir)))
 			{
 				// remap the project root. For Rocket executables, rename the executable to the game name.
-				if (Params.Rocket && Exe == Exes[0])
+				if (!Params.IsCodeBasedProject && Exe == Exes[0])
 				{
 					SC.StageFiles(StagedFileType.NonUFS, CombinePaths(SC.ProjectRoot, "Binaries", SC.PlatformDir), Path.GetFileNameWithoutExtension(Exe), true, null, CommandUtils.CombinePaths(SC.RelativeProjectRootForStage, "Binaries", SC.PlatformDir), false, true, SC.ShortProjectName);
 				}
@@ -77,7 +77,7 @@ public abstract class BaseLinuxPlatform : Platform
 			else if (Exe.StartsWith(CombinePaths(SC.RuntimeRootDir, "Engine/Binaries", SC.PlatformDir)))
 			{
 				// Move the executable for non-code rocket projects into the game directory, using the game name, so it can figure out the UProject to look for and is consitent with code projects.
-				if (Params.Rocket && Exe == Exes[0])
+				if (!Params.IsCodeBasedProject && Exe == Exes[0])
 				{
 					SC.StageFiles(StagedFileType.NonUFS, CombinePaths(SC.LocalRoot, "Engine/Binaries", SC.PlatformDir), Path.GetFileNameWithoutExtension(Exe), true, null, CommandUtils.CombinePaths(SC.RelativeProjectRootForStage, "Binaries", SC.PlatformDir), false, true, SC.ShortProjectName);
 				}
