Minor: made comments around previous change more meaningful.

diff --git a/Engine/Source/Programs/AutomationTool/Linux/LinuxPlatform.Automation.cs b/Engine/Source/Programs/AutomationTool/Linux/LinuxPlatform.Automation.cs
--- a/Engine/Source/Programs/AutomationTool/Linux/LinuxPlatform.Automation.cs
+++ b/Engine/Source/Programs/AutomationTool/Linux/LinuxPlatform.Automation.cs
@@ -64,7 +64,7 @@ public abstract class BaseLinuxPlatform : Platform
 
 			if (Exe.StartsWith(CombinePaths(SC.RuntimeProjectRootDir, "Binaries", SC.PlatformDir)))
 			{
-				// remap the project root. For Rocket executables, rename the executable to the game name.
+				// remap the project root. For content-only projects, rename the executable to the project name.
 				if (!Params.IsCodeBasedProject && Exe == Exes[0])
 				{
 					SC.StageFiles(StagedFileType.NonUFS, CombinePaths(SC.ProjectRoot, "Binaries", SC.PlatformDir), Path.GetFileNameWithoutExtension(Exe), true, null, CommandUtils.CombinePaths(SC.RelativeProjectRootForStage, "Binaries", SC.PlatformDir), false, true, SC.ShortProjectName);
@@ -76,7 +76,7 @@ public abstract class BaseLinuxPlatform : Platform
 			}
 			else if (Exe.StartsWith(CombinePaths(SC.RuntimeRootDir, "Engine/Binaries", SC.PlatformDir)))
 			{
-				// Move the executable for non-code rocket projects into the game directory, using the game name, so it can figure out the UProject to look for and is consitent with code projects.
+				// Move the executable for content-only projects into the project directory, using the project name, so it can figure out the UProject to look for and is consistent with code projects.
 				if (!Params.IsCodeBasedProject && Exe == Exes[0])
 				{
 					SC.StageFiles(StagedFileType.NonUFS, CombinePaths(SC.LocalRoot, "Engine/Binaries", SC.PlatformDir), Path.GetFileNameWithoutExtension(Exe), true, null, CommandUtils.CombinePaths(SC.RelativeProjectRootForStage, "Binaries", SC.PlatformDir), false, true, SC.ShortProjectName);
