Don't send automated test failure emails to causers. They don't run often enough to be useful, and the failures are obscure.

diff --git a/Engine/Source/Programs/AutomationTool/Scripts/GUBP.Automation.cs b/Engine/Source/Programs/AutomationTool/Scripts/GUBP.Automation.cs
--- a/Engine/Source/Programs/AutomationTool/Scripts/GUBP.Automation.cs
+++ b/Engine/Source/Programs/AutomationTool/Scripts/GUBP.Automation.cs
@@ -231,7 +231,7 @@ public class GUBP : BuildCommand
     }
 
     public abstract class GUBPEmailHacker
-    {        
+    {
         public virtual List<string> AddEmails(GUBP bp, string Branch, string NodeName)
         {
             return new List<string>();
@@ -244,10 +244,14 @@ public class GUBP : BuildCommand
         {
             return Emails;
         }
-    }
+		public virtual bool VetoEmailingCausers(GUBP bp, string Branch, string NodeName)
+		{
+			return false; // People who have submitted since last-green will be included unless vetoed by overriding this method. 
+		}
+	}
 
     private static List<GUBPEmailHacker> EmailHackers;
-    private string HackEmails(string Emails, string Branch, string NodeName)
+    private string HackEmails(string Emails, string Causers, string Branch, string NodeName)
     {
         string OnlyEmail = ParseParamValue("OnlyEmail");
         if (!String.IsNullOrEmpty(OnlyEmail))
@@ -280,7 +284,10 @@ public class GUBP : BuildCommand
             }
         }
         List<string> Result = new List<string>(Emails.Split(' '));
-
+		if(!EmailHackers.Any(x => x.VetoEmailingCausers(this, Branch, NodeName)))
+		{
+			Result.AddRange(Causers.Split(' '));
+		}
         foreach (var EmailHacker in EmailHackers)
         {
             Result.AddRange(EmailHacker.AddEmails(this, Branch, NodeName));
@@ -3502,14 +3509,14 @@ public class GUBP : BuildCommand
         string ECPerlFile = CommandUtils.CombinePaths(CommandUtils.CmdEnv.LogFolder, "jobsteps.pl");
         WriteAllLines_NoExceptions(ECPerlFile, Args.ToArray());
     }
-    string GetEMailListForNode(GUBP bp, string NodeToDo, string Causers)
+    string GetEMailListForNode(GUBP bp, string NodeToDo, string Emails, string Causers)
     {        
         var BranchForEmail = "";
         if (P4Enabled)
         {
             BranchForEmail = P4Env.BuildRootP4;
         }
-        return HackEmails(Causers, BranchForEmail, NodeToDo);
+        return HackEmails(Emails, Causers, BranchForEmail, NodeToDo);
     }
     int GetFrequencyForNode(GUBP bp, string NodeToDo, int BaseFrequency)
     {
@@ -3641,7 +3648,7 @@ public class GUBP : BuildCommand
             string EMails = "";
             if (AddEmailProps)
             {
-                EMails = GetEMailListForNode(bp, NodeToDo, "");
+                EMails = GetEMailListForNode(bp, NodeToDo, "", "");
             }
             if (bShowTriggers)
             {
@@ -4377,15 +4384,11 @@ public class GUBP : BuildCommand
 
         {
             var AdditonalEmails = "";
+
+			string Causers = "";
             if (ParseParam("CIS") && !GUBPNodes[NodeToDo].SendSuccessEmail() && !GUBPNodes[NodeToDo].TriggerNode())
             {
-#if false // cisheros was disabled because it was ignored and was just spam
-                if (NumPeople > 50 || NumPeople == 0)
-                {
-                    FailCauserEMails = "[CISHeros]";
-                }
-#endif
-                AdditonalEmails = FailCauserEMails;
+				Causers = FailCauserEMails;
             }
             string AddEmails = ParseParamValue("AddEmails");
             if (!String.IsNullOrEmpty(AddEmails))
@@ -4393,7 +4396,7 @@ public class GUBP : BuildCommand
                 AdditonalEmails = GUBPNode.MergeSpaceStrings(AddEmails, AdditonalEmails);
             }
 
-            EMails = GetEMailListForNode(this, NodeToDo, AdditonalEmails);
+            EMails = GetEMailListForNode(this, NodeToDo, AdditonalEmails, Causers);
 
             ECProps.Add("FailEmails/" + NodeToDo + "=" + EMails);
         }
