build script, fix for merge stomp

diff --git a/Engine/Source/Programs/AutomationTool/Scripts/GUBP.Automation.cs b/Engine/Source/Programs/AutomationTool/Scripts/GUBP.Automation.cs
--- a/Engine/Source/Programs/AutomationTool/Scripts/GUBP.Automation.cs
+++ b/Engine/Source/Programs/AutomationTool/Scripts/GUBP.Automation.cs
@@ -3502,7 +3502,7 @@ public class GUBP : BuildCommand
             return RunAndLog("ectool", Args, Options: Opts);
         }
     }
-    string GetEMailListForNode(GUBP bp, string NodeToDo, string Causers)
+    string GetEMailListForNode(GUBP bp, string NodeToDo, string Emails, string Causers)
     {        
         var BranchForEmail = "";
         if (P4Enabled)
