Make different browsers as HTML5 Devices.

 The configuration is driven by a editor platform specific .ini setting. Also, always serve files via a web browser ( python SimpleHTTP Server ) whose life time is attached to the browser device to get around X-origin issues w/ chrome. cook on the fly path remains as such, because we already serve via the cook on the fly web server.

#HTML5

diff --git a/Engine/Config/HTML5/HTML5Engine.ini b/Engine/Config/HTML5/HTML5Engine.ini
--- a/Engine/Config/HTML5/HTML5Engine.ini
+++ b/Engine/Config/HTML5/HTML5Engine.ini
@@ -9,3 +9,15 @@ r.setres=800x600
 
 [BuildSettings]
 HeapSize=1024
+
+[HTML5DevicesWindows]
+Firefox (64 bit)="C:/Program Files/Nightly/firefox.exe"
+Firefox (32 bit)="C:/Program Files (x86)/Mozilla Firefox/firefox.exe"
+Chrome="C:/Program Files (x86)/Google/Chrome/Application/chrome.exe"
+
+[HTML5DevicesMac]
+Safari="/Applications/Safari.app/Contents/MacOS/Safari"
+Firefox="/Applications/Firefox.app/Contents/MacOS/firefox-bin"
+
+[HTML5DevicesLinux]
+Firefox="/usr/bin/firefox"
diff --git a/Engine/Source/Developer/HTML5/HTML5TargetPlatform/Private/HTML5TargetDevice.cpp b/Engine/Source/Developer/HTML5/HTML5TargetPlatform/Private/HTML5TargetDevice.cpp
--- a/Engine/Source/Developer/HTML5/HTML5TargetPlatform/Private/HTML5TargetDevice.cpp
+++ b/Engine/Source/Developer/HTML5/HTML5TargetPlatform/Private/HTML5TargetDevice.cpp
@@ -30,7 +30,7 @@ FTargetDeviceId FHTML5TargetDevice::GetId( ) const
 
 FString FHTML5TargetDevice::GetName( ) const
 {
-	return FPlatformProcess::ComputerName();
+	return Name;
 }
 
 
diff --git a/Engine/Source/Developer/HTML5/HTML5TargetPlatform/Private/HTML5TargetDevice.h b/Engine/Source/Developer/HTML5/HTML5TargetPlatform/Private/HTML5TargetDevice.h
--- a/Engine/Source/Developer/HTML5/HTML5TargetPlatform/Private/HTML5TargetDevice.h
+++ b/Engine/Source/Developer/HTML5/HTML5TargetPlatform/Private/HTML5TargetDevice.h
@@ -35,7 +35,8 @@ public:
 	 * @param InName - The device name.
 	 */
 	FHTML5TargetDevice( const ITargetPlatform& InTargetPlatform, const FString& InName )
-		: TargetPlatform(InTargetPlatform)
+		: TargetPlatform(InTargetPlatform), 
+		  Name(InName)
 	{ }
 
 	/**
@@ -111,4 +112,6 @@ private:
 
 	// Holds a reference to the device's target platform.
 	const ITargetPlatform& TargetPlatform;
+
+	FString Name; 
 };
diff --git a/Engine/Source/Developer/HTML5/HTML5TargetPlatform/Private/HTML5TargetPlatform.cpp b/Engine/Source/Developer/HTML5/HTML5TargetPlatform/Private/HTML5TargetPlatform.cpp
--- a/Engine/Source/Developer/HTML5/HTML5TargetPlatform/Private/HTML5TargetPlatform.cpp
+++ b/Engine/Source/Developer/HTML5/HTML5TargetPlatform/Private/HTML5TargetPlatform.cpp
@@ -15,11 +15,28 @@
 
 FHTML5TargetPlatform::FHTML5TargetPlatform( )
 {
-	LocalDevice = MakeShareable(new FHTML5TargetDevice(*this, FPlatformProcess::ComputerName()));
-
 	// load the final HTML5 engine settings for this game
 	FConfigCacheIni::LoadLocalIniFile(HTML5EngineSettings, TEXT("Engine"), true, *PlatformName());
 
+	FString DeviceSectionName;
+
+#if PLATFORM_WINDOWS
+	DeviceSectionName = "HTML5DevicesWindows";
+#elif PLATFORM_MAC
+	DeviceSectionName = "HTML5DevicesMac";
+#else 
+	DeviceSectionName = "HTML5DevicesLinux";
+#endif 
+
+	FConfigSection AvaliableDevicesSection = HTML5EngineSettings[DeviceSectionName];
+	for (auto It : AvaliableDevicesSection)
+	{
+		const FString& BrowserName = It.Key.ToString();
+		const FString& BrowserPath = It.Value;
+		if(FPlatformFileManager::Get().GetPlatformFile().FileExists(*It.Value))
+			LocalDevice.Add(MakeShareable(new FHTML5TargetDevice(*this, It.Key.ToString())));
+	}
+
 #if WITH_ENGINE
 	// load up texture settings from the config file
 	HTML5LODSettings.Initialize(HTML5EngineSettings, TEXT("SystemSettings"));
@@ -34,7 +51,7 @@ FHTML5TargetPlatform::FHTML5TargetPlatform( )
 void FHTML5TargetPlatform::GetAllDevices( TArray<ITargetDevicePtr>& OutDevices ) const
 {
 	OutDevices.Reset();
-	OutDevices.Add(LocalDevice);
+	OutDevices.Append(LocalDevice);
 }
 
 
@@ -46,17 +63,20 @@ ECompressionFlags FHTML5TargetPlatform::GetBaseCompressionMethod( ) const
 
 ITargetDevicePtr FHTML5TargetPlatform::GetDefaultDevice( ) const
 {
-	return LocalDevice;
+	if(LocalDevice.Num())
+		return LocalDevice[0];
+	else 
+		return NULL; 
 }
 
 
 ITargetDevicePtr FHTML5TargetPlatform::GetDevice( const FTargetDeviceId& DeviceId )
 {
-	if (LocalDevice.IsValid() && (DeviceId == LocalDevice->GetId()))
+	for ( auto Device : LocalDevice )
 	{
-		return LocalDevice;
+		if ( Device.IsValid() && DeviceId == Device->GetId() )
+			return Device; 
 	}
-
 	return NULL;
 }
 
diff --git a/Engine/Source/Developer/HTML5/HTML5TargetPlatform/Private/HTML5TargetPlatform.h b/Engine/Source/Developer/HTML5/HTML5TargetPlatform/Private/HTML5TargetPlatform.h
--- a/Engine/Source/Developer/HTML5/HTML5TargetPlatform/Private/HTML5TargetPlatform.h
+++ b/Engine/Source/Developer/HTML5/HTML5TargetPlatform/Private/HTML5TargetPlatform.h
@@ -95,7 +95,7 @@ private:
 	FConfigFile HTML5EngineSettings;
 
 	// Holds the local device.
-	FHTML5TargetDevicePtr LocalDevice;
+	TArray<ITargetDevicePtr> LocalDevice;
 
 #if WITH_ENGINE
 	// Holds the cached target LOD settings.
diff --git a/Engine/Source/Programs/AutomationTool/HTML5/HTML5Platform.Automation.cs b/Engine/Source/Programs/AutomationTool/HTML5/HTML5Platform.Automation.cs
--- a/Engine/Source/Programs/AutomationTool/HTML5/HTML5Platform.Automation.cs
+++ b/Engine/Source/Programs/AutomationTool/HTML5/HTML5Platform.Automation.cs
@@ -200,42 +200,30 @@ public class HTML5Platform : Platform
 
 	public override ProcessResult RunClient(ERunOptions ClientRunFlags, string ClientApp, string ClientCmdLine, ProjectParams Params)
 	{
-		// look for firefox
-		string[] PossiblePaths = new string[] 
-		{
-			Path.Combine(Environment.GetFolderPath(Environment.SpecialFolder.ProgramFiles), @"Nightly\firefox.exe"),
-			Path.Combine(Environment.GetFolderPath(Environment.SpecialFolder.ProgramFilesX86), @"Nightly\firefox.exe"),
-			Path.Combine(Environment.GetFolderPath(Environment.SpecialFolder.ProgramFiles), @"Mozilla Firefox\firefox.exe"),
-			Path.Combine(Environment.GetFolderPath(Environment.SpecialFolder.ProgramFilesX86), @"Mozilla Firefox\firefox.exe"),
-		};
-
-		// set up a directory for temp profile
-		string FirefoxProfileCommand = " -no-remote ";
-
-		// look for the best firefox to tun
-		string FirefoxPath = null;
-		foreach (string PossiblePath in PossiblePaths)
-		{
-			if (File.Exists(PossiblePath))
-			{
-				FirefoxPath = PossiblePath;
-				break;
-			}
-		}
-
-		// if we didn't find one, just use the shell to open it
-		if (FirefoxPath == null)
-		{
-			FirefoxPath = CmdEnv.CmdExe;
-			FirefoxProfileCommand = " ";
-		}
+        // look for browser
+        var ConfigCache = new UnrealBuildTool.ConfigCacheIni(UnrealTargetPlatform.HTML5, "Engine", Path.GetDirectoryName(Params.RawProjectPath), CombinePaths(CmdEnv.LocalRoot, "Engine"));
+
+        string DeviceSection; 
+
+        if ( Utils.IsRunningOnMono )
+        {
+            DeviceSection = "HTML5DevicesMac";
+        }
+        else
+        {
+            DeviceSection = "HTML5DevicesWindows";
+        }
+
+        string browserPath;
+        string DeviceName = Params.Device.Split('@')[1];
+        bool ok = ConfigCache.GetString(DeviceSection, DeviceName, out browserPath);
+        
+		if (!ok)
+			throw new System.Exception ("Incorrect browser configuration in HTML5Engine.ini "); 
 
 		// open the webpage
-		// what to do with the commandline!??
-		string HTMLPath = Path.ChangeExtension(ClientApp, "html");
-		if (!File.Exists(HTMLPath))  // its probably a content only game - then it exists in the UE4 directory and not in the game directory. 
-			HTMLPath = Path.Combine(CombinePaths(CmdEnv.LocalRoot, "Engine"), "Binaries", "HTML5", Path.GetFileName(HTMLPath));
-
+        string directory = Path.GetDirectoryName(ClientApp);
+        string url = Path.GetFileName(ClientApp) +".html"; 
 		// Are we running via cook on the fly server? 
 		// find our http url - This is awkward because RunClient doesn't have real information that NFS is running or not. 
 		bool IsCookOnTheFly = false;
@@ -251,7 +239,7 @@ public class HTML5Platform : Platform
 				{
 					if (Url.Contains("http"))
 					{
-						HTMLPath = Url + "/" + Path.GetFileName(HTMLPath);
+						url = Url + "/" + Path.GetFileName(url);
 						IsCookOnTheFly = true;
 					}
 				}
@@ -260,25 +248,59 @@ public class HTML5Platform : Platform
 				break;
 		}
 
-		if (IsCookOnTheFly)
-		{
-			HTMLPath += "?cookonthefly=true";
-		}
+        if (IsCookOnTheFly)
+        {
+            url += "?cookonthefly=true";
+        }
+        else
+        {
 
-		if (BuildHostPlatform.Current.Platform == UnrealTargetPlatform.Mac)
-		{
-			FirefoxPath = "/bin/bash";
-			FirefoxProfileCommand = " -c  ' open -a Firefox.app  \"" + HTMLPath + "\"   ' ";
-		}
-		else if (BuildHostPlatform.Current.Platform == UnrealTargetPlatform.Win64)
-		{
-			FirefoxProfileCommand = FirefoxProfileCommand + "\"" + HTMLPath + "\" ";
-		}
+            url = "http://127.0.0.1:8000/" + url;
+            // this will be killed UBT instances dies. 
+            string input = String.Format(" -m SimpleHTTPServer 8000");
+
+
+			string PythonName = "python2.exe";
+
+			if (Utils.IsRunningOnMono)
+				PythonName = "python2"; 
+
+			ProcessResult Result = ProcessManager.CreateProcess(PythonName, true, "html5server.log");
+			Result.ProcessObject.StartInfo.FileName = PythonName;
+            Result.ProcessObject.StartInfo.UseShellExecute = false; 
+            Result.ProcessObject.StartInfo.RedirectStandardOutput = true;
+            Result.ProcessObject.StartInfo.RedirectStandardInput = true;
+            Result.ProcessObject.StartInfo.WorkingDirectory = directory;
+            Result.ProcessObject.StartInfo.Arguments = input;
+            Result.ProcessObject.Start();
+
+
+            Result.ProcessObject.OutputDataReceived += delegate(object sender, System.Diagnostics.DataReceivedEventArgs e)
+            {
+                System.Console.WriteLine(e.Data);
+            };
+
+            System.Console.WriteLine("Starting Browser Process"); 
+
+            ProcessResult ClientProcess = Run(browserPath, url, null, ClientRunFlags | ERunOptions.NoWaitForExit);
+
+            ClientProcess.ProcessObject.EnableRaisingEvents = true; 
+            ClientProcess.ProcessObject.Exited += delegate(System.Object o, System.EventArgs e)
+            {
+                System.Console.WriteLine("Browser Process Ended - Killing Webserver");
+                // send kill. 
+                Result.ProcessObject.StandardInput.Close();
+                Result.ProcessObject.Kill();
+            };
+
+            return ClientProcess;
+        }
+
+        System.Console.WriteLine("Browser Path " + browserPath);
+        ProcessResult BrowserProcess = Run(browserPath, url, null, ClientRunFlags | ERunOptions.NoWaitForExit);
 
-		System.Console.WriteLine("Firefox command line: " + FirefoxProfileCommand);
-		ProcessResult ClientProcess = Run(FirefoxPath, FirefoxProfileCommand, null, ClientRunFlags | ERunOptions.NoWaitForExit);
+        return BrowserProcess; 
 
-		return ClientProcess;
 	}
 
 	public override string GetCookPlatform(bool bDedicatedServer, bool bIsClientOnly, string CookFlavor)
