Linux host platform for UAT.

- Note: target framework for UAT (mono project) is upped to 4.5 to match UBT.

diff --git a/Engine/Source/Programs/AutomationTool/AutomationTool.csproj b/Engine/Source/Programs/AutomationTool/AutomationTool.csproj
--- a/Engine/Source/Programs/AutomationTool/AutomationTool.csproj
+++ b/Engine/Source/Programs/AutomationTool/AutomationTool.csproj
@@ -106,6 +106,7 @@
     <Compile Include="Distiller.cs" />
     <Compile Include="HelpUtils.cs" />
     <Compile Include="HostPlatform.cs" />
+    <Compile Include="LinuxHostPlatform.cs" />
     <Compile Include="LocalBuildEnvironment.cs" />
     <Compile Include="LocalP4Environment.cs" />
     <Compile Include="Log.cs" />
diff --git a/Engine/Source/Programs/AutomationTool/AutomationTool_Mono.csproj b/Engine/Source/Programs/AutomationTool/AutomationTool_Mono.csproj
--- a/Engine/Source/Programs/AutomationTool/AutomationTool_Mono.csproj
+++ b/Engine/Source/Programs/AutomationTool/AutomationTool_Mono.csproj
@@ -10,7 +10,7 @@
     <AppDesignerFolder>Properties</AppDesignerFolder>
     <RootNamespace>AutomationTool</RootNamespace>
     <AssemblyName>AutomationTool</AssemblyName>
-    <TargetFrameworkVersion>v4.0</TargetFrameworkVersion>
+    <TargetFrameworkVersion>v4.5</TargetFrameworkVersion>
     <TargetFrameworkProfile>
     </TargetFrameworkProfile>
     <FileAlignment>512</FileAlignment>
@@ -109,6 +109,7 @@
     <Compile Include="Distiller.cs" />
     <Compile Include="HelpUtils.cs" />
     <Compile Include="HostPlatform.cs" />
+    <Compile Include="LinuxHostPlatform.cs" />
     <Compile Include="LocalBuildEnvironment.cs" />
     <Compile Include="LocalP4Environment.cs" />
     <Compile Include="Log.cs" />
@@ -154,4 +155,4 @@
   <Target Name="AfterBuild">
   </Target>
   -->
-</Project>
\ No newline at end of file
+</Project>
diff --git a/Engine/Source/Programs/AutomationTool/HostPlatform.cs b/Engine/Source/Programs/AutomationTool/HostPlatform.cs
--- a/Engine/Source/Programs/AutomationTool/HostPlatform.cs
+++ b/Engine/Source/Programs/AutomationTool/HostPlatform.cs
@@ -2,6 +2,7 @@
 
 using System;
 using System.Collections.Generic;
+using System.IO;
 using System.Linq;
 using System.Text;
 using System.Diagnostics;
@@ -27,13 +28,26 @@ namespace AutomationTool
 		/// </summary>
 		internal static void Initialize()
 		{
-			if (UnrealBuildTool.BuildHostPlatform.Current.Platform == UnrealBuildTool.UnrealTargetPlatform.Mac)
+			PlatformID Platform = Environment.OSVersion.Platform;
+			switch (Platform)
 			{
-				RunningPlatform = new MacHostPlatform();
-			}
-			else
-			{
-				RunningPlatform = new WindowsHostPlatform();
+				case PlatformID.Win32NT:
+					RunningPlatform = new WindowsHostPlatform();
+					break;
+
+				case PlatformID.Unix:
+					if (File.Exists ("/System/Library/CoreServices/SystemVersion.plist"))
+					{
+						RunningPlatform = new MacHostPlatform();
+					} 
+					else 
+					{
+						RunningPlatform = new LinuxHostPlatform();
+					}
+					break;
+
+				default:
+					throw new Exception ("Unhandled runtime platform " + Platform);
 			}
 		}
 
diff --git a/Engine/Source/Programs/AutomationTool/Linux/LinuxPlatform.Automation.cs b/Engine/Source/Programs/AutomationTool/Linux/LinuxPlatform.Automation.cs
--- a/Engine/Source/Programs/AutomationTool/Linux/LinuxPlatform.Automation.cs
+++ b/Engine/Source/Programs/AutomationTool/Linux/LinuxPlatform.Automation.cs
@@ -113,6 +113,11 @@ public abstract class BaseLinuxPlatform : Platform
 		}
 	}
 
+	public override string GetEditorCookPlatform()
+	{
+		return "Linux";
+	}
+
 	/// <summary>
 	/// return true if we need to change the case of filenames outside of pak files
 	/// </summary>
@@ -179,6 +184,17 @@ chmod 700 $HOME/Desktop/{1}.desktop", DesiredGLVersion, SC.ShortProjectName, SC.
 
 	public override void Package(ProjectParams Params, DeploymentContext SC, int WorkingCL)
 	{
+		// package up the program
+		PrintRunTime();
+	}
+
+	public override bool CanHostPlatform(UnrealTargetPlatform Platform)
+	{
+		if (Platform == UnrealTargetPlatform.Mac || Platform == UnrealTargetPlatform.Win32 || Platform == UnrealTargetPlatform.Win64)
+		{
+			return false;
+		}
+		return true;
 	}
 
 	/// <summary>
@@ -213,7 +229,7 @@ chmod 700 $HOME/Desktop/{1}.desktop", DesiredGLVersion, SC.ShortProjectName, SC.
 				&& ((!String.IsNullOrEmpty(ProjParams.DeviceUsername) && !String.IsNullOrEmpty(ProjParams.DevicePassword))
 					|| !ProjParams.Unattended)) // Skip key generation in unattended mode if information is missing
 			{
-				Log("Configuring linux host");
+				Log("Configuring Linux host");
 
 				// Prompt for username if not already set
 				while (String.IsNullOrEmpty(ProjParams.DeviceUsername))
@@ -248,7 +264,7 @@ chmod 700 $HOME/Desktop/{1}.desktop", DesiredGLVersion, SC.ShortProjectName, SC.
 					} while (key.Key != ConsoleKey.Enter);
 					Console.WriteLine();
 				}
-				String Command = @"C:\Windows\system32\cmd.exe";
+				String Command = @"/bin/sh";
 
 				// Convienence names
 				String keyName = String.Format("UE4:{0}@{1}", Environment.UserName, Environment.MachineName);
@@ -277,7 +293,7 @@ chmod 700 $HOME/Desktop/{1}.desktop", DesiredGLVersion, SC.ShortProjectName, SC.
 			// Fail if a key couldn't be found or generated
 			if (String.IsNullOrEmpty(linuxKey))
 			{
-				throw new AutomationException("can't deploy/run using a linux device without a valid SSH key");
+				throw new AutomationException("can't deploy/run using a Linux device without a valid SSH key");
 			}
 
 			// Set username from key
@@ -304,15 +320,15 @@ chmod 700 $HOME/Desktop/{1}.desktop", DesiredGLVersion, SC.ShortProjectName, SC.
 		return new List<string> { };
 	}
 
+	public override bool IsSupported { get { return true; } }
+
 }
 
-public class CentOSx64Platform : BaseLinuxPlatform
+
+public class GenericLinuxPlatform : BaseLinuxPlatform
 {
-	public CentOSx64Platform()
+	public GenericLinuxPlatform()
 		: base(UnrealTargetPlatform.Linux)
 	{
 	}
-
-	public override bool IsSupported { get { return true; } }
 }
-
diff --git a/Engine/Source/Programs/AutomationTool/LinuxHostPlatform.cs b/Engine/Source/Programs/AutomationTool/LinuxHostPlatform.cs
new file mode 100644
--- /dev/null
+++ b/Engine/Source/Programs/AutomationTool/LinuxHostPlatform.cs
@@ -0,0 +1,128 @@
+﻿// Copyright 1998-2014 Epic Games, Inc. All Rights Reserved.
+
+using System;
+using System.Collections.Generic;
+using System.Linq;
+using System.Text;
+using System.Diagnostics;
+using System.IO;
+using UnrealBuildTool;
+
+namespace AutomationTool
+{
+	class LinuxHostPlatform : HostPlatform
+	{
+		public override void SetFrameworkVars()
+		{
+			// @todo
+		}
+
+		public override string GetMsBuildExe()
+		{
+			return "xbuild";
+		}
+
+		public override string GetMsDevExe()
+		{
+			return "";
+		}
+
+		public override string RelativeBinariesFolder
+		{
+			get { return @"Engine/Binaries/Linux/"; }
+		}
+
+		public override string GetUE4ExePath(string UE4Exe)
+		{
+			return CommandUtils.CombinePaths(CommandUtils.CmdEnv.LocalRoot, RelativeBinariesFolder, UE4Exe);
+		}
+
+		public override string LocalBuildsLogFolder
+		{
+			get { return Path.Combine(Environment.GetFolderPath(Environment.SpecialFolder.Personal), "Library/Logs/Unreal Engine/LocalBuildLogs"); }
+		}
+
+		public override string P4Exe
+		{
+			get { return "/usr/bin/p4"; }
+		}
+
+		public override Process CreateProcess(string AppName)
+		{
+			var NewProcess = new Process();
+			if (AppName == "mono")
+			{
+				// Enable case-insensitive mode for Mono
+				NewProcess.StartInfo.EnvironmentVariables.Add("MONO_IOMAP", "case");
+			}
+			return NewProcess;
+		}
+
+		public override void SetupOptionsForRun(ref string AppName, ref CommandUtils.ERunOptions Options, ref string CommandLine)
+		{
+			if (AppName == "sh" || AppName == "xbuild" || AppName == "codesign")
+			{
+				Options &= ~CommandUtils.ERunOptions.AppMustExist;
+			}
+			if (AppName == "xbuild")
+			{
+				AppName = "xbuild";
+				CommandLine = (String.IsNullOrEmpty(CommandLine) ? "" : CommandLine) + " /verbosity:quiet /nologo";
+				// For some reason AutomationScripts.Automation.csproj has ToolsVersion set
+				// to 11.0, which is no available on linux, so force ToolsVersion to 4.0
+				CommandLine += " /tv:4.0";
+				// Some projects have TargetFrameworkProfile=Client which causes warnings on Linux
+				// so force it to empty.
+				CommandLine += " /p:TargetFrameworkProfile=";
+			}
+			if (AppName.EndsWith(".exe") || ((AppName.Contains("/Binaries/Win64/") || AppName.Contains("/Binaries/Linux/")) && string.IsNullOrEmpty(Path.GetExtension(AppName))))
+			{
+				if (AppName.Contains("/Binaries/Win64/") || AppName.Contains("/Binaries/Linux/"))
+				{
+					AppName = AppName.Replace("/Binaries/Win64/", "/Binaries/Linux/");
+					AppName = AppName.Replace("-cmd.exe", "");
+					AppName = AppName.Replace("-Cmd.exe", "");
+					AppName = AppName.Replace(".exe", "");
+				}
+				else
+				{
+					// It's a C# app, so run it with Mono
+					AppName = "mono";
+					CommandLine = "\"" + AppName + "\" " + (String.IsNullOrEmpty(CommandLine) ? "" : CommandLine);
+					Options &= ~CommandUtils.ERunOptions.AppMustExist;
+				}
+			}
+		}
+
+		public override void SetConsoleCtrlHandler(ProcessManager.CtrlHandlerDelegate Handler)
+		{
+			// @todo: add mono support
+		}
+
+		public override string UBTProjectName
+		{
+			get { return "UnrealBuildTool_Mono"; }
+		}
+
+		public override UnrealTargetPlatform HostEditorPlatform
+		{
+			get { return UnrealTargetPlatform.Linux; }
+		}
+
+		public override string PdbExtension
+		{
+			get { return ".exe.mdb"; }
+		}
+		static string[] SystemServices = new string[]
+		{
+			// TODO: Add any system process names here
+		};
+		public override string[] DontKillProcessList
+		{
+			get
+			{
+				return SystemServices;
+			}
+		}
+	}
+}
diff --git a/Engine/Source/Programs/AutomationTool/Win/Win.Automation.csproj b/Engine/Source/Programs/AutomationTool/Win/Win.Automation.csproj
--- a/Engine/Source/Programs/AutomationTool/Win/Win.Automation.csproj
+++ b/Engine/Source/Programs/AutomationTool/Win/Win.Automation.csproj
@@ -16,7 +16,7 @@
     <DebugSymbols>true</DebugSymbols>
     <DebugType>full</DebugType>
     <Optimize>false</Optimize>
-    <OutputPath>../../../../Binaries/DotNet/AutomationScripts/</OutputPath>
+    <OutputPath>../../../../Binaries/DotNET/AutomationScripts/</OutputPath>
     <DefineConstants>DEBUG;TRACE</DefineConstants>
     <ErrorReport>prompt</ErrorReport>
     <WarningLevel>4</WarningLevel>
@@ -26,7 +26,7 @@
   <PropertyGroup Condition=" '$(Configuration)|$(Platform)' == 'Development|AnyCPU' ">
     <DebugType>pdbonly</DebugType>
     <Optimize>true</Optimize>
-    <OutputPath>../../../../Binaries/DotNet/AutomationScripts/</OutputPath>
+    <OutputPath>../../../../Binaries/DotNET/AutomationScripts/</OutputPath>
     <DefineConstants>TRACE</DefineConstants>
     <ErrorReport>prompt</ErrorReport>
     <WarningLevel>4</WarningLevel>
@@ -65,4 +65,4 @@
     </ProjectReference>
   </ItemGroup>
   <Import Project="$(MSBuildToolsPath)\Microsoft.CSharp.targets" />
-</Project>
\ No newline at end of file
+</Project>
