UAT: Use current host platform instead of guessing it (wrongly).

diff --git a/Engine/Source/Programs/AutomationTool/Scripts/BuildCookRun.Automation.cs b/Engine/Source/Programs/AutomationTool/Scripts/BuildCookRun.Automation.cs
--- a/Engine/Source/Programs/AutomationTool/Scripts/BuildCookRun.Automation.cs
+++ b/Engine/Source/Programs/AutomationTool/Scripts/BuildCookRun.Automation.cs
@@ -350,7 +350,7 @@ public class BuildCookRun : BuildCommand
 					}
 					if (!FileExists_NoExceptions(ProjectName))
 					{
-						var Branch = new BranchInfo(new List<UnrealTargetPlatform> { UnrealBuildTool.Utils.IsRunningOnMono ? UnrealTargetPlatform.Mac : UnrealTargetPlatform.Win64 });
+						var Branch = new BranchInfo(new List<UnrealTargetPlatform> { UnrealBuildTool.BuildHostPlatform.Current.Platform });
 						var GameProj = Branch.FindGame(OriginalProjectName);
 						if (GameProj != null)
 						{
