Fixed "Build Failure" error after successfully deploying to Android

CommandUtils.Run was disposing of processes started without NoWaitForExit, so any caller checking ProcessResult.HasExited would crash raising InvalidOperationException as the ProcessResult.Proc was junk.

Made the ProcessResult destructor call Proc.Dispose().

diff --git a/Engine/Source/Programs/AutomationTool/ProcessUtils.cs b/Engine/Source/Programs/AutomationTool/ProcessUtils.cs
--- a/Engine/Source/Programs/AutomationTool/ProcessUtils.cs
+++ b/Engine/Source/Programs/AutomationTool/ProcessUtils.cs
@@ -231,6 +231,14 @@ namespace AutomationTool
 			Proc.EnableRaisingEvents = false;
 		}
 
+        ~ProcessResult()
+        {
+            if(Proc != null)
+            {
+                Proc.Dispose();
+            }
+        }
+
 		/// <summary>
 		/// Removes a process from the list of tracked processes.
 		/// </summary>
@@ -702,7 +710,6 @@ namespace AutomationTool
                 }
 				Result.ExitCode = Proc.ExitCode;
 				Result.OnProcessExited();
-				Proc.Dispose();
 			}
 			else
 			{
