Improvement to Process cleanup change in CL 2285994

diff --git a/Engine/Source/Programs/AutomationTool/ProcessUtils.cs b/Engine/Source/Programs/AutomationTool/ProcessUtils.cs
--- a/Engine/Source/Programs/AutomationTool/ProcessUtils.cs
+++ b/Engine/Source/Programs/AutomationTool/ProcessUtils.cs
@@ -256,6 +256,15 @@ namespace AutomationTool
 				Listener.TraceEvent(EventCache, Source, Verbosity, (int)Verbosity, Message);
 			}
 		}
+       
+		/// <summary>
+		/// Manually dispose of Proc and set it to null.
+		/// </summary>
+        public void DisposeProcess()
+        {
+            Proc.Dispose();
+            Proc = null;
+        }
 
 		/// <summary>
 		/// Process.OutputDataReceived event handler.
@@ -710,6 +719,7 @@ namespace AutomationTool
                 }
 				Result.ExitCode = Proc.ExitCode;
 				Result.OnProcessExited();
+                Result.DisposeProcess();
 			}
 			else
 			{
