UAT: Fix for local cook on the fly not specifying any ports for -filehostip

diff --git a/Engine/Source/Programs/AutomationTool/Scripts/RunProjectCommand.Automation.cs b/Engine/Source/Programs/AutomationTool/Scripts/RunProjectCommand.Automation.cs
--- a/Engine/Source/Programs/AutomationTool/Scripts/RunProjectCommand.Automation.cs
+++ b/Engine/Source/Programs/AutomationTool/Scripts/RunProjectCommand.Automation.cs
@@ -681,7 +681,7 @@ public partial class Project : CommandUtils
 
 				const string LocalHost = "127.0.0.1";
 
-				if (Params.Port != null)
+				if (!IsNullOrEmpty(Params.Port))
 				{
 					foreach (var Port in Params.Port)
 					{
