UAT: Fix for IsBuildMachine not being set when running with invalid commandline

diff --git a/Engine/Source/Programs/AutomationTool/Automation.cs b/Engine/Source/Programs/AutomationTool/Automation.cs
--- a/Engine/Source/Programs/AutomationTool/Automation.cs
+++ b/Engine/Source/Programs/AutomationTool/Automation.cs
@@ -271,13 +271,17 @@ AutomationTool.exe [-verbose] [-compileonly] [-p4] Command0 [-Arg0 -Arg1 -Arg2 
 		/// <param name="CommandLine">Command line</param>
 		public static void Process(string[] CommandLine)
 		{
+			// Initial check for local or build machine runs BEFORE we parse the command line (We need this value set
+			// in case something throws the exception while parsing the command line)
+			IsBuildMachine = !String.IsNullOrEmpty(Environment.GetEnvironmentVariable("uebp_LOCAL_ROOT"));
+
 			// Scan the command line for commands to execute.
 			var CommandsToExecute = new List<CommandInfo>();
 			var AdditionalScriptsFolders = new List<string>();
 			ParseCommandLine(CommandLine, CommandsToExecute, AdditionalScriptsFolders);
 
-			// Are we running locally or on a build machine
-			IsBuildMachine = GlobalCommandLine.ForceLocal ? false : !String.IsNullOrEmpty(Environment.GetEnvironmentVariable("uebp_LOCAL_ROOT"));
+			// Check for build machine override (force local)
+			IsBuildMachine = GlobalCommandLine.ForceLocal ? false : IsBuildMachine;
 			Log.TraceInformation("IsBuildMachine={0}", IsBuildMachine);
 
 			// should we kill processes on exit
