UAT: changes/improvements for Linux as GUBP host platform.

- Adds _OnLinux GUBP nodes.
- Fixed assumptions that "anything non-Mac is Windows".
- Other minor fixes.

It is now possible to build RootEditor_OnLinux (and dependent nodes, including Tools) on a Linux agent.

diff --git a/Engine/Source/Programs/AutomationTool/Scripts/GUBP.Automation.cs b/Engine/Source/Programs/AutomationTool/Scripts/GUBP.Automation.cs
--- a/Engine/Source/Programs/AutomationTool/Scripts/GUBP.Automation.cs
+++ b/Engine/Source/Programs/AutomationTool/Scripts/GUBP.Automation.cs
@@ -605,6 +605,10 @@ public class GUBP : BuildCommand
                 }
                 return "_OnMac";
             }
+			else if (InHostPlatform == UnrealTargetPlatform.Linux)
+			{
+				return "_OnLinux";
+			}
             return "";
         }
         public virtual UnrealTargetPlatform GetAgentPlatform()
@@ -731,7 +735,7 @@ public class GUBP : BuildCommand
         public override UE4Build.BuildAgenda GetAgenda(GUBP bp)
         {
             var Agenda = new UE4Build.BuildAgenda();
-			if (HostPlatform != UnrealTargetPlatform.Mac && !GUBP.bForceIncrementalCompile)
+			if (HostPlatform == UnrealTargetPlatform.Win64 && !GUBP.bForceIncrementalCompile)
             {
                 Agenda.DotNetProjects.AddRange(
                     new string[] 
@@ -748,7 +752,7 @@ public class GUBP : BuildCommand
         }
         public override void PostBuild(GUBP bp, UE4Build UE4Build)
         {
-            if (HostPlatform != UnrealTargetPlatform.Mac)
+			if (HostPlatform == UnrealTargetPlatform.Win64)
             {
             	UE4Build.AddUATFilesToBuildProducts();
             	UE4Build.AddUBTFilesToBuildProducts();
@@ -899,7 +903,7 @@ public class GUBP : BuildCommand
         {
             var Agenda = new UE4Build.BuildAgenda();
 
-            if (HostPlatform != UnrealTargetPlatform.Mac)
+			if (HostPlatform == UnrealTargetPlatform.Win64)
             {
 				if (!GUBP.bForceIncrementalCompile)
 				{
@@ -1036,7 +1040,7 @@ public class GUBP : BuildCommand
             bool bAnyAdded = false;
             var Agenda = new UE4Build.BuildAgenda();
 
-            if (HostPlatform != UnrealTargetPlatform.Mac)
+			if (HostPlatform == UnrealTargetPlatform.Win64)
             {
                 bAnyAdded = true;
                 Agenda.DotNetProjects.AddRange(
@@ -4518,7 +4522,7 @@ public class GUBP : BuildCommand
             BranchForOptions = P4Env.BuildRootP4;
         }
         BranchOptions = GetBranchOptions(BranchForOptions);
-        bool WithMac = !BranchOptions.PlatformsToRemove.Contains(UnrealTargetPlatform.Mac);        
+        bool WithMac = !BranchOptions.PlatformsToRemove.Contains(UnrealTargetPlatform.Mac);
         if (ParseParam("NoMac"))
         {
             WithMac = false;
@@ -4528,7 +4532,17 @@ public class GUBP : BuildCommand
             HostPlatforms.Add(UnrealTargetPlatform.Mac);
         }
 
-        bBuildRocket = ParseParam("BuildRocket");
+		bool WithLinux = !BranchOptions.PlatformsToRemove.Contains(UnrealTargetPlatform.Linux);
+		if (ParseParam("NoLinux"))
+		{
+			WithLinux = false;
+		}
+		if (WithLinux)
+		{
+			HostPlatforms.Add(UnrealTargetPlatform.Linux);
+		}
+
+		bBuildRocket = ParseParam("BuildRocket");
         bForceIncrementalCompile = ParseParam("ForceIncrementalCompile");
         bool bNoAutomatedTesting = ParseParam("NoAutomatedTesting");        
         StoreName = ParseParamValue("Store");
@@ -4823,7 +4837,7 @@ public class GUBP : BuildCommand
             Log("Active Platform: {0}", Plat.ToString());
         }
 
-        if (HostPlatforms.Count == 2)
+        if (HostPlatforms.Count >= 2)
         {
             // make sure each project is set up with the right assumptions on monolithics that prefer a platform.
             foreach (var CodeProj in Branch.CodeProjects)
@@ -4979,7 +4993,7 @@ public class GUBP : BuildCommand
                 }
 
 
-                if (HostPlatform != UnrealTargetPlatform.Mac) //temp hack till mac automated testing works
+				if (HostPlatform == UnrealTargetPlatform.Win64) //temp hack till automated testing works on other platforms than Win64
                 {
 
                     var EditorTests = Branch.BaseEngineProject.Properties.Targets[TargetRules.TargetType.Editor].Rules.GUBP_GetEditorTests_EditorTypeOnly(HostPlatform);
@@ -5235,7 +5249,7 @@ public class GUBP : BuildCommand
                 }
 
                 AddNode(new EditorGameNode(this, HostPlatform, CodeProj));
-                if (!bNoAutomatedTesting && HostPlatform != UnrealTargetPlatform.Mac) //temp hack till mac automated testing works
+				if (!bNoAutomatedTesting && HostPlatform == UnrealTargetPlatform.Win64) //temp hack till automated testing works on other platforms than Win64
                 {
                     var EditorTests = CodeProj.Properties.Targets[TargetRules.TargetType.Editor].Rules.GUBP_GetEditorTests_EditorTypeOnly(HostPlatform);
                     var EditorTestNodes = new List<string>();
diff --git a/Engine/Source/Programs/UnrealBuildTool/System/RulesCompiler.cs b/Engine/Source/Programs/UnrealBuildTool/System/RulesCompiler.cs
--- a/Engine/Source/Programs/UnrealBuildTool/System/RulesCompiler.cs
+++ b/Engine/Source/Programs/UnrealBuildTool/System/RulesCompiler.cs
@@ -663,11 +663,15 @@ namespace UnrealBuildTool
             // hack to set up the templates without adding anything to their .targets.cs files
             if (!String.IsNullOrEmpty(TargetName) && TargetName.StartsWith("TP_"))
             {
-                Result.Add(UnrealTargetPlatform.IOS);
-                if (HostPlatform != UnrealTargetPlatform.Mac)
+				if (HostPlatform == UnrealTargetPlatform.Win64)
+				{
+					Result.Add(UnrealTargetPlatform.IOS);
+					Result.Add(UnrealTargetPlatform.Android);
+				}
+                else if (HostPlatform == UnrealTargetPlatform.Mac)
                 {
-                    Result.Add(UnrealTargetPlatform.Android);
-                }
+					Result.Add(UnrealTargetPlatform.IOS);
+				}
             }
             return Result;
         }
