GUBP: do not add Linux nodes unless running on Linux (temp?).

- Prevents errors when trying to cross-compile Linux editor (whose cross-compilation is broken atm).

diff --git a/Engine/Source/Programs/AutomationTool/Scripts/GUBP.Automation.cs b/Engine/Source/Programs/AutomationTool/Scripts/GUBP.Automation.cs
--- a/Engine/Source/Programs/AutomationTool/Scripts/GUBP.Automation.cs
+++ b/Engine/Source/Programs/AutomationTool/Scripts/GUBP.Automation.cs
@@ -4448,8 +4448,9 @@ public class GUBP : BuildCommand
 
     [Help("Runs one, several or all of the GUBP nodes")]
     [Help(typeof(UE4Build))]
-    [Help("NoMac", "Toggle to exclude the Mac host platform, default is Win64+Mac")]
-    [Help("NoPC", "Toggle to exclude the PC host platform, default is Win64+Mac")]
+    [Help("NoMac", "Toggle to exclude the Mac host platform, default is Win64+Mac+Linux")]
+	[Help("NoLinux", "Toggle to exclude the Linux (PC, 64-bit) host platform, default is Win64+Mac+Linux")]
+	[Help("NoPC", "Toggle to exclude the PC host platform, default is Win64+Mac+Linux")]
     [Help("CleanLocal", "delete the local temp storage before we start")]
     [Help("Store=", "Sets the name of the temp storage block, normally, this is built for you.")]
     [Help("StoreSuffix=", "Tacked onto a store name constructed from CL, branch, etc")]
@@ -4533,7 +4534,8 @@ public class GUBP : BuildCommand
         }
 
 		bool WithLinux = !BranchOptions.PlatformsToRemove.Contains(UnrealTargetPlatform.Linux);
-		if (ParseParam("NoLinux"))
+		// @TODO: exclude temporarily unless running on a Linux machine to prevent spurious GUBP failures
+		if (UnrealBuildTool.BuildHostPlatform.Current.Platform != UnrealTargetPlatform.Linux || ParseParam("NoLinux"))
 		{
 			WithLinux = false;
 		}
