build script, priming for linux ec procedures

diff --git a/Engine/Source/Programs/AutomationTool/Scripts/GUBP.Automation.cs b/Engine/Source/Programs/AutomationTool/Scripts/GUBP.Automation.cs
--- a/Engine/Source/Programs/AutomationTool/Scripts/GUBP.Automation.cs
+++ b/Engine/Source/Programs/AutomationTool/Scripts/GUBP.Automation.cs
@@ -625,6 +625,14 @@ public class GUBP : BuildCommand
                 }
                 return "_Mac";
             }
+			if(GetAgentPlatform() == UnrealTargetPlatform.Linux)
+			{
+				if(IsSticky())
+				{
+					throw new AutomationException("Node {0} is sticky, but Linux hosted. Sticky nodes must be PC hosted.", GetFullName());
+				}
+				return "_Linux";
+			}
             return "";
         }
         public virtual string GetHostPlatformSuffix()
