- Added apk selection based on connected device

diff --git a/Engine/Source/Programs/AutomationTool/Android/AndroidPlatform.Automation.cs b/Engine/Source/Programs/AutomationTool/Android/AndroidPlatform.Automation.cs
--- a/Engine/Source/Programs/AutomationTool/Android/AndroidPlatform.Automation.cs
+++ b/Engine/Source/Programs/AutomationTool/Android/AndroidPlatform.Automation.cs
@@ -305,8 +305,7 @@ public class AndroidPlatform : Platform
 
 	public override void Deploy(ProjectParams Params, DeploymentContext SC)
 	{
-		// @todo android: Get this from the connected device we are running on if separate Apks!
-		string DeviceArchitecture = "";
+		string DeviceArchitecture = GetBestDeviceArchitecture(Params);
 
 		string AdbCommand = GetAdbCommand(Params);
 		string ApkName = GetFinalApkName(Params, SC.StageExecutables[0], true, DeviceArchitecture);
@@ -556,10 +555,70 @@ public class AndroidPlatform : Platform
 		return Path.Combine(Subdirs[0], "aapt.exe");
 	}
 
+	private string GetBestDeviceArchitecture(ProjectParams Params)
+	{
+		bool bMakeSeparateApks = UnrealBuildTool.Android.UEDeployAndroid.ShouldMakeSeparateApks();
+		// if we are joining all .so's into a single .apk, there's no need to find the best one - there is no other one
+		if (!bMakeSeparateApks)
+		{
+			return "";
+		}
+
+		string[] AppArchitectures = AndroidToolChain.GetAllArchitectures();
+
+		// ask the device
+		ProcessResult Result = Run(CmdEnv.CmdExe, GetAdbCommand(Params) + " shell getprop ro.product.cpu.abi", null, ERunOptions.AppMustExist);
+
+		// the output is just the architecture
+		string DeviceArch = UnrealBuildTool.Android.UEDeployAndroid.GetUE4Arch(Result.Output.Trim());
+
+		// if the architecture wasn't built, look for a backup
+		if (Array.IndexOf(AppArchitectures, DeviceArch) == -1)
+		{
+			// go from 64 to 32-bit
+			if (DeviceArch == "-arm64")
+			{
+				DeviceArch = "-armv7";
+			}
+			// go from 64 to 32-bit
+			else if (DeviceArch == "-x86_64")
+			{
+				if (Array.IndexOf(AppArchitectures, "-x86") == -1)
+				{
+					DeviceArch = "-x86";
+				}
+				// if it didn't have 32-bit x86, look for 64-bit arm for emulation
+				// @todo android 64-bit: x86_64 most likely can't emulate arm64 at this ponit
+// 				else if (Array.IndexOf(AppArchitectures, "-arm64") == -1)
+// 				{
+// 					DeviceArch = "-arm64";
+// 				}
+				// finally try for 32-bit arm emulation (Houdini)
+				else
+				{
+					DeviceArch = "-armv7";
+				}
+			}
+			// use armv7 (with Houdini emulation)
+			else if (DeviceArch == "-x86")
+			{
+				DeviceArch = "-armv7";
+			}
+		}
+
+		// if after the fallbacks, we still don't have it, we can't continue
+		if (Array.IndexOf(AppArchitectures, DeviceArch) == -1)
+		{
+			throw new AutomationException("Unable to run because you don't have an apk that is usable on {0}", Params.Device);
+		}
+
+		return DeviceArch;
+	}
+
+
 	public override ProcessResult RunClient(ERunOptions ClientRunFlags, string ClientApp, string ClientCmdLine, ProjectParams Params)
 	{
-		// @todo android: Get this from the connected device we are running on if separate Apks!
-		string DeviceArchitecture = "";
+		string DeviceArchitecture = GetBestDeviceArchitecture(Params);
 
 		string ApkName = ClientApp + DeviceArchitecture + ".apk";
 		if (!File.Exists(ApkName))
@@ -567,7 +626,7 @@ public class AndroidPlatform : Platform
 			ApkName = GetFinalApkName(Params, Path.GetFileNameWithoutExtension(ClientApp), true, DeviceArchitecture);
 		}
 
-		Console.WriteLine("Apk='{0}', CLientApp='{1}', ExeName='{2}'", ApkName, ClientApp, Params.ProjectGameExeFilename);
+		Console.WriteLine("Apk='{0}', ClientApp='{1}', ExeName='{2}'", ApkName, ClientApp, Params.ProjectGameExeFilename);
 
 		// run aapt to get the name of the intent
 		string PackageName = GetPackageInfo(ApkName, false);
@@ -627,7 +686,7 @@ public class AndroidPlatform : Platform
 			}
 
 			// this is just to get the ue4 log to go to the output
-			Run(CmdEnv.CmdExe, AdbCommand + "logcat -d -s UE4");
+			Run(CmdEnv.CmdExe, AdbCommand + "logcat -d -s UE4 -s Debug");
 
 			// get the log we actually want to save
 			ProcessResult LogFileProcess = Run(CmdEnv.CmdExe, AdbCommand + "logcat -d", null, ERunOptions.AppMustExist);
