Reduce frequency of documentation builds, and remove filter for code docs.

diff --git a/Engine/Source/Programs/AutomationTool/Scripts/BuildDocumentation.Automation.cs b/Engine/Source/Programs/AutomationTool/Scripts/BuildDocumentation.Automation.cs
--- a/Engine/Source/Programs/AutomationTool/Scripts/BuildDocumentation.Automation.cs
+++ b/Engine/Source/Programs/AutomationTool/Scripts/BuildDocumentation.Automation.cs
@@ -45,7 +45,7 @@ class DocumentationNode : GUBP.GUBPNode
 
 	public override int CISFrequencyQuantumShift(GUBP bp)
 	{
-		return base.CISFrequencyQuantumShift(bp);
+		return base.CISFrequencyQuantumShift(bp) + 4;
 	}
 
 	protected void ExecuteApiDocTool(string Arguments, string LogName)
@@ -87,7 +87,7 @@ class CodeDocumentationNode : DocumentationNode
 
 	public override void DoBuild(GUBP bp)
 	{
-		ExecuteApiDocTool("-rebuildcode -filter=Messaging/...", "APIDocTool-Code");
+		ExecuteApiDocTool("-rebuildcode", "APIDocTool-Code");
 		SubmitOutputs("Code documentation", "Engine/Documentation/Builds/CodeAPI-*", "Engine/Documentation/CHM/API.chm");
 		base.DoBuild(bp);
 	}
