- Fixed the hacky workaround for Mono Xbox UAT automation script, by always passing MONO as a #define to UAT scripts (since we can't make a _Mono.csproj version of each automation.csproj)

diff --git a/Engine/Source/Programs/AutomationTool/LinuxHostPlatform.cs b/Engine/Source/Programs/AutomationTool/LinuxHostPlatform.cs
--- a/Engine/Source/Programs/AutomationTool/LinuxHostPlatform.cs
+++ b/Engine/Source/Programs/AutomationTool/LinuxHostPlatform.cs
@@ -1,4 +1,4 @@
-﻿// Copyright 1998-2014 Epic Games, Inc. All Rights Reserved.
+// Copyright 1998-2014 Epic Games, Inc. All Rights Reserved.
 
 using System;
 using System.Collections.Generic;
@@ -71,6 +71,8 @@ namespace AutomationTool
 				// For some reason AutomationScripts.Automation.csproj has ToolsVersion set
 				// to 11.0, which is no available on linux, so force ToolsVersion to 4.0
 				CommandLine += " /tv:4.0";
+				// Pass #define MONO to all the automation scripts (see XboxOne)
+				CommandLine += " /p:DefineConstants=MONO";
 				// Some projects have TargetFrameworkProfile=Client which causes warnings on Linux
 				// so force it to empty.
 				CommandLine += " /p:TargetFrameworkProfile=";
diff --git a/Engine/Source/Programs/AutomationTool/MacHostPlatform.cs b/Engine/Source/Programs/AutomationTool/MacHostPlatform.cs
--- a/Engine/Source/Programs/AutomationTool/MacHostPlatform.cs
+++ b/Engine/Source/Programs/AutomationTool/MacHostPlatform.cs
@@ -1,4 +1,4 @@
-﻿// Copyright 1998-2014 Epic Games, Inc. All Rights Reserved.
+// Copyright 1998-2014 Epic Games, Inc. All Rights Reserved.
 
 using System;
 using System.Collections.Generic;
@@ -93,7 +93,7 @@ namespace AutomationTool
 			if (AppName == "xbuild")
 			{
 				AppName = "sh";
-				CommandLine = "-c 'xbuild " + (String.IsNullOrEmpty(CommandLine) ? "" : CommandLine) + " /property:TargetFrameworkVersion=v4.5 /tv:4.0 /verbosity:quiet /nologo |grep -i error; if [ $? -ne 1 ]; then exit 1; else exit 0; fi'";
+				CommandLine = "-c 'xbuild " + (String.IsNullOrEmpty(CommandLine) ? "" : CommandLine) + " /property:TargetFrameworkVersion=v4.5 /tv:4.0 /p:DefineConstants=MONO /verbosity:quiet /nologo |grep -i error; if [ $? -ne 1 ]; then exit 1; else exit 0; fi'";
 			}
 			if (AppName.EndsWith(".exe") || ((AppName.Contains("/Binaries/Win64/") || AppName.Contains("/Binaries/Mac/")) && string.IsNullOrEmpty(Path.GetExtension(AppName))))
 			{
