build script, shorten monolithic nodenames

diff --git a/Engine/Source/Programs/AutomationTool/Scripts/GUBP.Automation.cs b/Engine/Source/Programs/AutomationTool/Scripts/GUBP.Automation.cs
--- a/Engine/Source/Programs/AutomationTool/Scripts/GUBP.Automation.cs
+++ b/Engine/Source/Programs/AutomationTool/Scripts/GUBP.Automation.cs
@@ -1342,7 +1342,7 @@ public class GUBP : BuildCommand
 
         public static string StaticGetFullName(UnrealTargetPlatform InHostPlatform, BranchInfo.BranchUProject InGameProj, UnrealTargetPlatform InTargetPlatform)
         {
-            return InGameProj.GameName + "_" + InTargetPlatform + "_Monolithics" + StaticGetHostPlatformSuffix(InHostPlatform);
+            return InGameProj.GameName + "_" + InTargetPlatform + "_Mono" + StaticGetHostPlatformSuffix(InHostPlatform);
         }
         public override string GetFullName()
         {
@@ -2857,7 +2857,6 @@ public class GUBP : BuildCommand
                 // this is a codeless project, use the base project
                 WorkingGameProject = bp.Branch.BaseEngineProject;
             }
-
             if (!WorkingGameProject.Properties.Targets.ContainsKey(GameOrClient))
             {
                 throw new AutomationException("Can't make a game build for {0} because it doesn't have a {1} target.", WorkingGameProject.GameName, GameOrClient.ToString());
@@ -5155,37 +5154,37 @@ public class GUBP : BuildCommand
                                                 {
 													if (Options.bSeparateGamePromotion)
 													{
-                                                    var NodeName = AddNode(new FormalBuildNode(this, NonCodeProject, HostPlatform, new List<UnrealTargetPlatform>() { Plat }, new List<UnrealTargetConfiguration>() {PlatPair.TargetConfig}));
-                                                    // we don't want this delayed
-                                                    // this would normally wait for the testing phase, we just want to build it right away
-                                                    RemovePseudodependencyFromNode(
-                                                        CookNode.StaticGetFullName(HostPlatform, NonCodeProject, CookedPlatform),
-                                                        WaitForTestShared.StaticGetFullName());
-                                                    RemovePseudodependencyFromNode(
-                                                        CookNode.StaticGetFullName(HostPlatform, NonCodeProject, CookedPlatform),
-                                                        CookNode.StaticGetFullName(HostPlatform, Branch.BaseEngineProject, CookedPlatform));
-
-                                                    string BuildAgentSharingGroup = NonCodeProject.GameName + "_MakeFormalBuild_" + Plat.ToString() + HostPlatformNode.StaticGetHostPlatformSuffix(HostPlatform);
-                                                    if (Plat == UnrealTargetPlatform.IOS || Plat == UnrealTargetPlatform.Android) // These trash build products, so we need to use different agents
-                                                    {
-                                                        BuildAgentSharingGroup = "";
-                                                    }
-                                                    GUBPNodes[CookNode.StaticGetFullName(HostPlatform, NonCodeProject, CookedPlatform)].AgentSharingGroup = BuildAgentSharingGroup;
-                                                    GUBPNodes[GamePlatformCookedAndCompiledNode.StaticGetFullName(HostPlatform, NonCodeProject, Plat)].AgentSharingGroup = BuildAgentSharingGroup;
-                                                    GUBPNodes[NodeName].AgentSharingGroup = BuildAgentSharingGroup;
-
-                                                    if (PlatPair.bTest)
-                                                    {
-                                                        AddNode(new FormalBuildTestNode(this, NonCodeProject, HostPlatform, Plat, PlatPair.TargetConfig));
-                                                    }
-                                                }
+														var NodeName = AddNode(new FormalBuildNode(this, NonCodeProject, HostPlatform, new List<UnrealTargetPlatform>() { Plat }, new List<UnrealTargetConfiguration>() {PlatPair.TargetConfig}));
+														// we don't want this delayed
+														// this would normally wait for the testing phase, we just want to build it right away
+														RemovePseudodependencyFromNode(
+															CookNode.StaticGetFullName(HostPlatform, NonCodeProject, CookedPlatform),
+															WaitForTestShared.StaticGetFullName());
+														RemovePseudodependencyFromNode(
+															CookNode.StaticGetFullName(HostPlatform, NonCodeProject, CookedPlatform),
+															CookNode.StaticGetFullName(HostPlatform, Branch.BaseEngineProject, CookedPlatform));
+
+	                                                    string BuildAgentSharingGroup = NonCodeProject.GameName + "_MakeFormalBuild_" + Plat.ToString() + HostPlatformNode.StaticGetHostPlatformSuffix(HostPlatform);
+		                                                if (Plat == UnrealTargetPlatform.IOS || Plat == UnrealTargetPlatform.Android) // These trash build products, so we need to use different agents
+			                                            {
+				                                            BuildAgentSharingGroup = "";
+					                                    }
+						                                GUBPNodes[CookNode.StaticGetFullName(HostPlatform, NonCodeProject, CookedPlatform)].AgentSharingGroup = BuildAgentSharingGroup;
+							                            GUBPNodes[GamePlatformCookedAndCompiledNode.StaticGetFullName(HostPlatform, NonCodeProject, Plat)].AgentSharingGroup = BuildAgentSharingGroup;
+								                        GUBPNodes[NodeName].AgentSharingGroup = BuildAgentSharingGroup;
+													
+														if (PlatPair.bTest)
+														{
+															AddNode(new FormalBuildTestNode(this, NonCodeProject, HostPlatform, Plat, PlatPair.TargetConfig));
+														}
+													}
 													else
 													{														
 														AddNode(new FormalBuildNode(this, NonCodeProject, HostPlatform, new List<UnrealTargetPlatform>() { Plat }, new List<UnrealTargetConfiguration>() { PlatPair.TargetConfig }));
 													}
-                                            }
-                                    }
-                                    }
+												}
+											}
+										}
                                     }
                                     if (!bNoAutomatedTesting)
                                     {
