Increase frequency for doc builds.

diff --git a/Engine/Source/Programs/AutomationTool/Scripts/BuildDocumentation.Automation.cs b/Engine/Source/Programs/AutomationTool/Scripts/BuildDocumentation.Automation.cs
--- a/Engine/Source/Programs/AutomationTool/Scripts/BuildDocumentation.Automation.cs
+++ b/Engine/Source/Programs/AutomationTool/Scripts/BuildDocumentation.Automation.cs
@@ -45,7 +45,7 @@ class DocumentationNode : GUBP.GUBPNode
 
 	public override int CISFrequencyQuantumShift(GUBP bp)
 	{
-		return base.CISFrequencyQuantumShift(bp) + 4;
+		return base.CISFrequencyQuantumShift(bp) + 3;
 	}
 
 	protected void ExecuteApiDocTool(string Arguments, string LogName)
