build script, fix cooks for formal builds

diff --git a/Engine/Source/Programs/AutomationTool/Scripts/GUBP.Automation.cs b/Engine/Source/Programs/AutomationTool/Scripts/GUBP.Automation.cs
--- a/Engine/Source/Programs/AutomationTool/Scripts/GUBP.Automation.cs
+++ b/Engine/Source/Programs/AutomationTool/Scripts/GUBP.Automation.cs
@@ -2658,10 +2658,10 @@ public class GUBP : BuildCommand
         {
             return base.Priority() + 20.0f;
         }
-        public override int CISFrequencyQuantumShift(GUBP bp)
-        {
-            return base.CISFrequencyQuantumShift(bp) + 3;
-        }
+		public override int CISFrequencyQuantumShift(GUBP bp)
+		{
+			return base.CISFrequencyQuantumShift(bp) + 3;
+		}
         public static string GetArchiveDirectory(BranchInfo.BranchUProject InGameProj, UnrealTargetPlatform InHostPlatform, List<UnrealTargetPlatform> InClientTargetPlatforms = null, List<UnrealTargetConfiguration> InClientConfigs = null, List<UnrealTargetPlatform> InServerTargetPlatforms = null, List<UnrealTargetConfiguration> InServerConfigs = null, bool InClientNotGame = false)
         {
             string BaseDir = ResolveSharedBuildDirectory(InGameProj.GameName);
@@ -5152,59 +5152,54 @@ public class GUBP : BuildCommand
                                             {
                                                 if (PlatPair.TargetPlatform == Plat)
                                                 {
+													var NodeName = AddNode(new FormalBuildNode(this, NonCodeProject, HostPlatform, new List<UnrealTargetPlatform>() { Plat }, new List<UnrealTargetConfiguration>() {PlatPair.TargetConfig}));
+													// we don't want this delayed
+													// this would normally wait for the testing phase, we just want to build it right away
+													RemovePseudodependencyFromNode(
+														CookNode.StaticGetFullName(HostPlatform, NonCodeProject, CookedPlatform),
+														WaitForTestShared.StaticGetFullName());
+													RemovePseudodependencyFromNode(
+														CookNode.StaticGetFullName(HostPlatform, NonCodeProject, CookedPlatform),
+														CookNode.StaticGetFullName(HostPlatform, Branch.BaseEngineProject, CookedPlatform));
+													string BuildAgentSharingGroup = "";
 													if (Options.bSeparateGamePromotion)
 													{
-														var NodeName = AddNode(new FormalBuildNode(this, NonCodeProject, HostPlatform, new List<UnrealTargetPlatform>() { Plat }, new List<UnrealTargetConfiguration>() {PlatPair.TargetConfig}));
-														// we don't want this delayed
-														// this would normally wait for the testing phase, we just want to build it right away
-														RemovePseudodependencyFromNode(
-															CookNode.StaticGetFullName(HostPlatform, NonCodeProject, CookedPlatform),
-															WaitForTestShared.StaticGetFullName());
-														RemovePseudodependencyFromNode(
-															CookNode.StaticGetFullName(HostPlatform, NonCodeProject, CookedPlatform),
-															CookNode.StaticGetFullName(HostPlatform, Branch.BaseEngineProject, CookedPlatform));
-
-	                                                    string BuildAgentSharingGroup = NonCodeProject.GameName + "_MakeFormalBuild_" + Plat.ToString() + HostPlatformNode.StaticGetHostPlatformSuffix(HostPlatform);
-		                                                if (Plat == UnrealTargetPlatform.IOS || Plat == UnrealTargetPlatform.Android) // These trash build products, so we need to use different agents
-			                                            {
-				                                            BuildAgentSharingGroup = "";
-					                                    }
-						                                GUBPNodes[CookNode.StaticGetFullName(HostPlatform, NonCodeProject, CookedPlatform)].AgentSharingGroup = BuildAgentSharingGroup;
-							                            GUBPNodes[GamePlatformCookedAndCompiledNode.StaticGetFullName(HostPlatform, NonCodeProject, Plat)].AgentSharingGroup = BuildAgentSharingGroup;
-								                        GUBPNodes[NodeName].AgentSharingGroup = BuildAgentSharingGroup;
-													
-														if (PlatPair.bTest)
+														BuildAgentSharingGroup = NonCodeProject.GameName + "_MakeFormalBuild_" + Plat.ToString() + HostPlatformNode.StaticGetHostPlatformSuffix(HostPlatform);
+														if (Plat == UnrealTargetPlatform.IOS || Plat == UnrealTargetPlatform.Android) // These trash build products, so we need to use different agents
 														{
-															AddNode(new FormalBuildTestNode(this, NonCodeProject, HostPlatform, Plat, PlatPair.TargetConfig));
+															BuildAgentSharingGroup = "";
 														}
+														GUBPNodes[CookNode.StaticGetFullName(HostPlatform, NonCodeProject, CookedPlatform)].AgentSharingGroup = BuildAgentSharingGroup;
+														GUBPNodes[GamePlatformCookedAndCompiledNode.StaticGetFullName(HostPlatform, NonCodeProject, Plat)].AgentSharingGroup = BuildAgentSharingGroup;
+														GUBPNodes[NodeName].AgentSharingGroup = BuildAgentSharingGroup;
 													}
-													else
-													{														
-														AddNode(new FormalBuildNode(this, NonCodeProject, HostPlatform, new List<UnrealTargetPlatform>() { Plat }, new List<UnrealTargetConfiguration>() { PlatPair.TargetConfig }));
-													}
+													
+													if (PlatPair.bTest)
+													{
+														AddNode(new FormalBuildTestNode(this, NonCodeProject, HostPlatform, Plat, PlatPair.TargetConfig));
+													}																										
 												}
 											}
 										}
-                                    }
+									}
                                     if (!bNoAutomatedTesting)
-                                    {
-                                        
-                                    if (HostPlatform == UnrealTargetPlatform.Mac) continue; //temp hack till mac automated testing works
-                                    var ThisMonoGameTestNodes = new List<string>();
-                                    foreach (var Test in GameTests)
-                                    {
-                                        var TestName = Test.Key + "_" + Plat.ToString();
-                                        ThisMonoGameTestNodes.Add(AddNode(new UATTestNode(this, HostPlatform, NonCodeProject, TestName, Test.Value, CookedAgentSharingGroup, false, RequiredPlatforms)));
-                                    }
-                                    if (ThisMonoGameTestNodes.Count > 0)
                                     {                                        
-                                        GameTestNodes.Add(AddNode(new GameAggregateNode(this, HostPlatform, NonCodeProject, "CookedTests_" + Plat.ToString() + "_" + Kind.ToString() + HostPlatformNode.StaticGetHostPlatformSuffix(HostPlatform), ThisMonoGameTestNodes, 0.0f)));
-                                    }
-                                }
-                            }
-                        }
-                    }
-                }
+										if (HostPlatform == UnrealTargetPlatform.Mac) continue; //temp hack till mac automated testing works
+										var ThisMonoGameTestNodes = new List<string>();
+										foreach (var Test in GameTests)
+										{
+											var TestName = Test.Key + "_" + Plat.ToString();
+											ThisMonoGameTestNodes.Add(AddNode(new UATTestNode(this, HostPlatform, NonCodeProject, TestName, Test.Value, CookedAgentSharingGroup, false, RequiredPlatforms)));
+										}
+										if (ThisMonoGameTestNodes.Count > 0)
+										{                                        
+											GameTestNodes.Add(AddNode(new GameAggregateNode(this, HostPlatform, NonCodeProject, "CookedTests_" + Plat.ToString() + "_" + Kind.ToString() + HostPlatformNode.StaticGetHostPlatformSuffix(HostPlatform), ThisMonoGameTestNodes, 0.0f)));
+										}
+									}
+								}
+							}
+						}
+					}
                 }
 #if false
                 //for now, non-code projects don't do client or server.
@@ -5316,118 +5311,120 @@ public class GUBP : BuildCommand
                                 }
 								if (!AdditionalPlatforms.Contains(Plat))
 								{
-                                string CookedPlatform = Platform.Platforms[Plat].GetCookPlatform(Kind == TargetRules.TargetType.Server, Kind == TargetRules.TargetType.Client, "");
-                                if (!GUBPNodes.ContainsKey(CookNode.StaticGetFullName(HostPlatform, CodeProj, CookedPlatform)))
-                                {
-                                    AddNode(new CookNode(this, HostPlatform, CodeProj, Plat, CookedPlatform));
-                                }
-                                if (!GUBPNodes.ContainsKey(GamePlatformCookedAndCompiledNode.StaticGetFullName(HostPlatform, CodeProj, Plat)))
-                                {
-                                    AddNode(new GamePlatformCookedAndCompiledNode(this, HostPlatform, CodeProj, Plat, true));
-                                }
-                                var FormalBuildConfigs = Target.Rules.GUBP_GetConfigsForFormalBuilds_MonolithicOnly(HostPlatform);
-
-                                foreach (var Config in FormalBuildConfigs)
-                                {
-                                    string FormalNodeName = null;                                    
-                                    if (Kind == TargetRules.TargetType.Client)
-                                    {
-                                        if (Plat == Config.TargetPlatform)
-                                        {
-                                            FormalNodeName = AddNode(new FormalBuildNode(this, CodeProj, HostPlatform, InClientTargetPlatforms: new List<UnrealTargetPlatform>() { Config.TargetPlatform }, InClientConfigs: new List<UnrealTargetConfiguration>() { Config.TargetConfig }, InClientNotGame: true));
-                                        }
-                                    }
-                                    else if (Kind == TargetRules.TargetType.Server)
-                                    {
-                                        if (Plat == Config.TargetPlatform)
-                                        {
-                                            FormalNodeName = AddNode(new FormalBuildNode(this, CodeProj, HostPlatform, InServerTargetPlatforms: new List<UnrealTargetPlatform>() { Config.TargetPlatform }, InServerConfigs: new List<UnrealTargetConfiguration>() { Config.TargetConfig }));
-                                        }
-                                    }
-                                    else if (Kind == TargetRules.TargetType.Game)
-                                    {
-                                        if (Plat == Config.TargetPlatform)
-                                        {
-                                            FormalNodeName = AddNode(new FormalBuildNode(this, CodeProj, HostPlatform, InClientTargetPlatforms: new List<UnrealTargetPlatform>() { Config.TargetPlatform }, InClientConfigs: new List<UnrealTargetConfiguration>() { Config.TargetConfig }));
-                                        }
-                                    }
-                                    if (FormalNodeName != null)
-                                    {
-                                        // we don't want this delayed
-                                        // this would normally wait for the testing phase, we just want to build it right away
-                                        RemovePseudodependencyFromNode(
-                                            CookNode.StaticGetFullName(HostPlatform, CodeProj, CookedPlatform),
-                                            WaitForTestShared.StaticGetFullName());
-                                        RemovePseudodependencyFromNode(
-                                            CookNode.StaticGetFullName(HostPlatform, CodeProj, CookedPlatform),
-                                            CookNode.StaticGetFullName(HostPlatform, Branch.BaseEngineProject, CookedPlatform));
-
-                                        string BuildAgentSharingGroup = CodeProj.GameName + "_MakeFormalBuild_" + Plat.ToString() + HostPlatformNode.StaticGetHostPlatformSuffix(HostPlatform);
-                                        if (Plat == UnrealTargetPlatform.IOS || Plat == UnrealTargetPlatform.Android || Plat == UnrealTargetPlatform.XboxOne || !Options.bSeparateGamePromotion) // These trash build products, so we need to use different agents
-                                        {
-                                            BuildAgentSharingGroup = "";
-                                        }
-                                        GUBPNodes[CookNode.StaticGetFullName(HostPlatform, CodeProj, CookedPlatform)].AgentSharingGroup = BuildAgentSharingGroup;
-                                        GUBPNodes[GamePlatformCookedAndCompiledNode.StaticGetFullName(HostPlatform, CodeProj, Plat)].AgentSharingGroup = BuildAgentSharingGroup;
-                                        GUBPNodes[FormalNodeName].AgentSharingGroup = BuildAgentSharingGroup;
-
-                                        if (Config.bTest)
-                                        {
-                                            AddNode(new FormalBuildTestNode(this, CodeProj, HostPlatform, Plat, Config.TargetConfig));
-                                        }
-                                    }
-                                }
-                                if (!bNoAutomatedTesting)
-                                {
-if (HostPlatform == UnrealTargetPlatform.Mac) continue; //temp hack till mac automated testing works
-                                var GameTests = Target.Rules.GUBP_GetGameTests_MonolithicOnly(HostPlatform, GetAltHostPlatform(HostPlatform), Plat);
-                                var RequiredPlatforms = new List<UnrealTargetPlatform> { Plat };
-                                var ThisMonoGameTestNodes = new List<string>();
-
-                                foreach (var Test in GameTests)
-                                {
-                                    var TestNodeName = Test.Key + "_" + Plat.ToString();
-                                    ThisMonoGameTestNodes.Add(AddNode(new UATTestNode(this, HostPlatform, CodeProj, TestNodeName, Test.Value, CookedAgentSharingGroup, false, RequiredPlatforms)));
-                                }
-                                if (ThisMonoGameTestNodes.Count > 0)
-                                {
-                                    GameTestNodes.Add(AddNode(new GameAggregateNode(this, HostPlatform, CodeProj, "CookedTests_" + Plat.ToString() + "_" + Kind.ToString() + HostPlatformNode.StaticGetHostPlatformSuffix(HostPlatform), ThisMonoGameTestNodes, 0.0f)));
-                                }
+									string CookedPlatform = Platform.Platforms[Plat].GetCookPlatform(Kind == TargetRules.TargetType.Server, Kind == TargetRules.TargetType.Client, "");
+									if (!GUBPNodes.ContainsKey(CookNode.StaticGetFullName(HostPlatform, CodeProj, CookedPlatform)))
+									{
+										AddNode(new CookNode(this, HostPlatform, CodeProj, Plat, CookedPlatform));
+									}
+									if (!GUBPNodes.ContainsKey(GamePlatformCookedAndCompiledNode.StaticGetFullName(HostPlatform, CodeProj, Plat)))
+									{
+										AddNode(new GamePlatformCookedAndCompiledNode(this, HostPlatform, CodeProj, Plat, true));
+									}
+									var FormalBuildConfigs = Target.Rules.GUBP_GetConfigsForFormalBuilds_MonolithicOnly(HostPlatform);
+
+									foreach (var Config in FormalBuildConfigs)
+									{
+										string FormalNodeName = null;                                    
+										if (Kind == TargetRules.TargetType.Client)
+										{
+											if (Plat == Config.TargetPlatform)
+											{
+												FormalNodeName = AddNode(new FormalBuildNode(this, CodeProj, HostPlatform, InClientTargetPlatforms: new List<UnrealTargetPlatform>() { Config.TargetPlatform }, InClientConfigs: new List<UnrealTargetConfiguration>() { Config.TargetConfig }, InClientNotGame: true));
+											}
+										}
+										else if (Kind == TargetRules.TargetType.Server)
+										{
+											if (Plat == Config.TargetPlatform)
+											{
+												FormalNodeName = AddNode(new FormalBuildNode(this, CodeProj, HostPlatform, InServerTargetPlatforms: new List<UnrealTargetPlatform>() { Config.TargetPlatform }, InServerConfigs: new List<UnrealTargetConfiguration>() { Config.TargetConfig }));
+											}
+										}
+										else if (Kind == TargetRules.TargetType.Game)
+										{
+											if (Plat == Config.TargetPlatform)
+											{
+												FormalNodeName = AddNode(new FormalBuildNode(this, CodeProj, HostPlatform, InClientTargetPlatforms: new List<UnrealTargetPlatform>() { Config.TargetPlatform }, InClientConfigs: new List<UnrealTargetConfiguration>() { Config.TargetConfig }));
+											}
+										}
+										if (FormalNodeName != null)
+										{
+											// we don't want this delayed
+											// this would normally wait for the testing phase, we just want to build it right away
+											RemovePseudodependencyFromNode(
+												CookNode.StaticGetFullName(HostPlatform, CodeProj, CookedPlatform),
+												WaitForTestShared.StaticGetFullName());
+											RemovePseudodependencyFromNode(
+												CookNode.StaticGetFullName(HostPlatform, CodeProj, CookedPlatform),
+												CookNode.StaticGetFullName(HostPlatform, Branch.BaseEngineProject, CookedPlatform));
+											string BuildAgentSharingGroup = "";
+											if (Options.bSeparateGamePromotion)
+											{
+												BuildAgentSharingGroup = CodeProj.GameName + "_MakeFormalBuild_" + Plat.ToString() + HostPlatformNode.StaticGetHostPlatformSuffix(HostPlatform);
+												if (Plat == UnrealTargetPlatform.IOS || Plat == UnrealTargetPlatform.Android || Plat == UnrealTargetPlatform.XboxOne) // These trash build products, so we need to use different agents
+												{
+													BuildAgentSharingGroup = "";
+												}
+												GUBPNodes[CookNode.StaticGetFullName(HostPlatform, CodeProj, CookedPlatform)].AgentSharingGroup = BuildAgentSharingGroup;
+												GUBPNodes[GamePlatformCookedAndCompiledNode.StaticGetFullName(HostPlatform, CodeProj, Plat)].AgentSharingGroup = BuildAgentSharingGroup;
+												GUBPNodes[FormalNodeName].AgentSharingGroup = BuildAgentSharingGroup;
+											}
 
-                                }
-                            }
-                        }
-                    }
-                }
+											if (Config.bTest)
+											{
+												AddNode(new FormalBuildTestNode(this, CodeProj, HostPlatform, Plat, Config.TargetConfig));
+											}																				
+										}
+									}
+									if (!bNoAutomatedTesting)
+									{
+										if (HostPlatform == UnrealTargetPlatform.Mac) continue; //temp hack till mac automated testing works
+										var GameTests = Target.Rules.GUBP_GetGameTests_MonolithicOnly(HostPlatform, GetAltHostPlatform(HostPlatform), Plat);
+										var RequiredPlatforms = new List<UnrealTargetPlatform> { Plat };
+										var ThisMonoGameTestNodes = new List<string>();
+
+										foreach (var Test in GameTests)
+										{
+											var TestNodeName = Test.Key + "_" + Plat.ToString();
+											ThisMonoGameTestNodes.Add(AddNode(new UATTestNode(this, HostPlatform, CodeProj, TestNodeName, Test.Value, CookedAgentSharingGroup, false, RequiredPlatforms)));
+										}
+										if (ThisMonoGameTestNodes.Count > 0)
+										{
+											GameTestNodes.Add(AddNode(new GameAggregateNode(this, HostPlatform, CodeProj, "CookedTests_" + Plat.ToString() + "_" + Kind.ToString() + HostPlatformNode.StaticGetHostPlatformSuffix(HostPlatform), ThisMonoGameTestNodes, 0.0f)));
+										}
+									}
+								}
+							}
+						}
+					}
 				}
-                if (!bNoAutomatedTesting)
-                {
-                foreach (var ServerPlatform in ServerPlatforms)
-                {
-                    foreach (var GamePlatform in GamePlatforms)
-                    {
-if (HostPlatform == UnrealTargetPlatform.Mac) continue; //temp hack till mac automated testing works
-                        var Target = CodeProj.Properties.Targets[TargetRules.TargetType.Game];
-                        var ClientServerTests = Target.Rules.GUBP_GetClientServerTests_MonolithicOnly(HostPlatform, GetAltHostPlatform(HostPlatform), ServerPlatform, GamePlatform);
-                        var RequiredPlatforms = new List<UnrealTargetPlatform> { ServerPlatform };
-                        if (ServerPlatform != GamePlatform)
-                        {
-                            RequiredPlatforms.Add(GamePlatform);
-                        }
-                        foreach (var Test in ClientServerTests)
-                        {
-                            var TestNodeName = Test.Key + "_" + GamePlatform.ToString() + "_" + ServerPlatform.ToString();
-                            GameTestNodes.Add(AddNode(new UATTestNode(this, HostPlatform, CodeProj, TestNodeName, Test.Value, CookedAgentSharingGroup, false, RequiredPlatforms)));
-                        }
-                    }
-                }
-                if (GameTestNodes.Count > 0)
-                {
-                    AddNode(new GameAggregateNode(this, HostPlatform, CodeProj, "AllCookedTests", GameTestNodes));
-                }
-            }
-            }
-        }
+				if (!bNoAutomatedTesting)
+				{
+					foreach (var ServerPlatform in ServerPlatforms)
+					{
+						foreach (var GamePlatform in GamePlatforms)
+						{
+							if (HostPlatform == UnrealTargetPlatform.Mac) continue; //temp hack till mac automated testing works
+							var Target = CodeProj.Properties.Targets[TargetRules.TargetType.Game];
+							var ClientServerTests = Target.Rules.GUBP_GetClientServerTests_MonolithicOnly(HostPlatform, GetAltHostPlatform(HostPlatform), ServerPlatform, GamePlatform);
+							var RequiredPlatforms = new List<UnrealTargetPlatform> { ServerPlatform };
+							if (ServerPlatform != GamePlatform)
+							{
+								RequiredPlatforms.Add(GamePlatform);
+							}
+							foreach (var Test in ClientServerTests)
+							{
+								var TestNodeName = Test.Key + "_" + GamePlatform.ToString() + "_" + ServerPlatform.ToString();
+								GameTestNodes.Add(AddNode(new UATTestNode(this, HostPlatform, CodeProj, TestNodeName, Test.Value, CookedAgentSharingGroup, false, RequiredPlatforms)));
+							}
+						}
+					}
+					if (GameTestNodes.Count > 0)
+					{
+						AddNode(new GameAggregateNode(this, HostPlatform, CodeProj, "AllCookedTests", GameTestNodes));
+					}
+				}
+			}
+		}
 
         int NumSharedAllHosts = 0;
         foreach (var CodeProj in Branch.CodeProjects)
