OneSky automation added to UE4 localization automation.

diff --git a/Engine/Source/Programs/AutomationTool/Scripts/Localisation.Automation.cs b/Engine/Source/Programs/AutomationTool/Scripts/Localisation.Automation.cs
--- a/Engine/Source/Programs/AutomationTool/Scripts/Localisation.Automation.cs
+++ b/Engine/Source/Programs/AutomationTool/Scripts/Localisation.Automation.cs
@@ -1,13 +1,29 @@
 ﻿// Copyright 1998-2014 Epic Games, Inc. All Rights Reserved.
 using System;
 using System.Collections.Generic;
+using System.Globalization;
+using System.Linq;
 using System.Text;
 using System.IO;
 using AutomationTool;
 using UnrealBuildTool;
+using OneSky;
+using EpicGames.OneSkyLocalization.Config;
 
 class Localise : BuildCommand
 {
+    private void ExportProjectToDirectory(OneSkyService oneSkyService, ProjectGroup projectGroup, string ProjectName)
+    {
+        var project = GetProject(oneSkyService, projectGroup.Name, ProjectName);
+        var file = project.UploadedFiles.FirstOrDefault(f => f.Filename == (ProjectName + ".po"));
+
+        //Export
+        if (file != null)
+        {
+            ExportFileToDirectory(file, new DirectoryInfo(CmdEnv.LocalRoot + "/Engine/Content/Localization/" + ProjectName), projectGroup.EnabledCultures);
+        }
+    }
+
 	public override void ExecuteBuild()
 	{
 		var EditorExe = CombinePaths(CmdEnv.LocalRoot, @"Engine/Binaries/Win64/UE4Editor-Cmd.exe");
@@ -21,6 +37,16 @@ class Localise : BuildCommand
 			Log("Localize from label {0}", P4Env.LabelToSync);
 		}
 
+        OneSkyConfigData OneSkyConfig = OneSkyConfigHelper.Find("OneSkyConfig_EpicGames");
+        var oneSkyService = new OneSkyService(OneSkyConfig.ApiKey, OneSkyConfig.ApiSecret);
+        var projectGroup = GetProjectGroup(oneSkyService, "Unreal Engine");
+
+        // Export all text from OneSky
+        ExportProjectToDirectory(oneSkyService, projectGroup, "Engine");
+        ExportProjectToDirectory(oneSkyService, projectGroup, "Editor");
+        ExportProjectToDirectory(oneSkyService, projectGroup, "PropertyNames");
+        ExportProjectToDirectory(oneSkyService, projectGroup, "ToolTips");
+
 		// Setup editor arguments for SCC.
 		string EditorArguments = String.Empty;
 		if (P4Enabled)
@@ -52,7 +78,7 @@ class Localise : BuildCommand
 		{
 			Log("Localization for {0} {1}", EditorArguments, CommandletArguments);
 
-			Log("Running UE4Editor to generate Localisation data");
+            Log("Running UE4Editor to generate localization data");
 
 			string Arguments = String.Format("-run=GatherText {0} {1}", EditorArguments, CommandletArguments);
 			var RunResult = Run(EditorExe, Arguments);
@@ -63,6 +89,12 @@ class Localise : BuildCommand
 			}
 		}
 
+        // Upload all text to OneSky
+        UploadDirectoryToProject(GetProject(oneSkyService, "Unreal Engine", "Engine"), new DirectoryInfo(CmdEnv.LocalRoot + "/Engine/Content/Localization/Engine"), "*.po");
+        UploadDirectoryToProject(GetProject(oneSkyService, "Unreal Engine", "Editor"), new DirectoryInfo(CmdEnv.LocalRoot + "/Engine/Content/Localization/Editor"), "*.po");
+        UploadDirectoryToProject(GetProject(oneSkyService, "Unreal Engine", "PropertyNames"), new DirectoryInfo(CmdEnv.LocalRoot + "/Engine/Content/Localization/PropertyNames"), "*.po");
+        UploadDirectoryToProject(GetProject(oneSkyService, "Unreal Engine", "ToolTips"), new DirectoryInfo(CmdEnv.LocalRoot + "/Engine/Content/Localization/ToolTips"), "*.po");
+
 		// Localisation statistics estimator.
 		if (P4Enabled)
 		{
@@ -86,5 +118,106 @@ class Localise : BuildCommand
 			}
 		}
 	}
+
+    private static ProjectGroup GetProjectGroup(OneSkyService oneSkyService, string GroupName)
+    {
+        var launcherGroup = oneSkyService.ProjectGroups.FirstOrDefault(g => g.Name == GroupName);
+
+        if (launcherGroup == null)
+        {
+            launcherGroup = new ProjectGroup(GroupName, new CultureInfo("en"));
+            oneSkyService.ProjectGroups.Add(launcherGroup);
+        }
+
+        return launcherGroup;
+    }
+
+    private static OneSky.Project GetProject(OneSkyService oneSkyService, string GroupName, string ProjectName, string ProjectDescription = "")
+    {
+        var launcherGroup = GetProjectGroup(oneSkyService, GroupName);
+
+        OneSky.Project appProject = launcherGroup.Projects.FirstOrDefault(p => p.Name == ProjectName);
+
+        if (appProject == null)
+        {
+            ProjectType projectType = oneSkyService.ProjectTypes.First(pt => pt.Code == "website");
+
+            appProject = new OneSky.Project(ProjectName, ProjectDescription, projectType);
+            launcherGroup.Projects.Add(appProject);
+        }
+
+        return appProject;
+    }
+
+    private static void ExportFileToDirectory(UploadedFile file, DirectoryInfo destination, IEnumerable<CultureInfo> cultures)
+    {
+        foreach (var culture in cultures)
+        {
+            var cultureDirectory = new DirectoryInfo(Path.Combine(destination.FullName, OneSky.LocaleCodeHelper.ConvertToLocaleCode(culture.Name)));
+            if (!cultureDirectory.Exists)
+            {
+                cultureDirectory.Create();
+            }
+
+            using (var memoryStream = new MemoryStream())
+            {
+                var exportFile = new FileInfo(Path.Combine(cultureDirectory.FullName, file.Filename));
+
+                var exportTranslationState = file.ExportTranslation(culture, memoryStream).Result;
+                if (exportTranslationState == UploadedFile.ExportTranslationState.Success)
+                {
+                    memoryStream.Position = 0;
+                    using (Stream fileStream = File.OpenWrite(exportFile.FullName))
+                    {
+                        memoryStream.CopyTo(fileStream);
+                        Console.WriteLine("[SUCCESS] Exporting: " + exportFile.FullName + " Locale: " + OneSky.LocaleCodeHelper.ConvertToLocaleCode(culture.Name));
+                    }
+                }
+                else if (exportTranslationState == UploadedFile.ExportTranslationState.NoContent)
+                {
+                    Console.WriteLine("[WARNING] Exporting: " + exportFile.FullName + " Locale: " + OneSky.LocaleCodeHelper.ConvertToLocaleCode(culture.Name) + " has no translations!");
+                }
+                else
+                {
+                    Console.WriteLine("[FAILED] Exporting: " + exportFile.FullName + " Locale: " + OneSky.LocaleCodeHelper.ConvertToLocaleCode(culture.Name));
+                }
+            }
+        }
+    }
+
+    static void UploadDirectoryToProject(OneSky.Project project, DirectoryInfo directory, string fileExtension)
+    {
+        foreach (var file in Directory.GetFiles(directory.FullName, fileExtension, SearchOption.AllDirectories))
+        {
+            DirectoryInfo parentDirectory = Directory.GetParent(file);
+            string localeName = parentDirectory.Name;
+            string currentFile = file;
+
+            using (var fileStream = File.OpenRead(currentFile))
+            {
+                // Read the BOM
+                var bom = new byte[3];
+                fileStream.Read(bom, 0, 3);
+
+                //We want to ignore the utf8 BOM
+                if (bom[0] != 0xef || bom[1] != 0xbb || bom[2] != 0xbf)
+                {
+                    fileStream.Position = 0;
+                }
+
+                Console.WriteLine("Uploading: " + currentFile + " Locale: " + localeName);
+                var uploadedFile = project.Upload(Path.GetFileName(currentFile), fileStream, new CultureInfo(OneSky.LocaleCodeHelper.ConvertFromLocaleCode(localeName))).Result;
+
+                if (uploadedFile == null)
+                {
+                    Console.WriteLine("[FAILED] Uploading: " + currentFile + " Locale: " + localeName);
+                }
+                else
+                {
+                    Console.WriteLine("[SUCCESS] Uploading: " + currentFile + " Locale: " + localeName);
+                }
+            }
+        }
+    }
 }
 
