UAT: Support for filtering out unsupported script modules on different platforms

diff --git a/Engine/Source/Programs/AutomationTool/HostPlatform.cs b/Engine/Source/Programs/AutomationTool/HostPlatform.cs
--- a/Engine/Source/Programs/AutomationTool/HostPlatform.cs
+++ b/Engine/Source/Programs/AutomationTool/HostPlatform.cs
@@ -111,6 +111,13 @@ namespace AutomationTool
 		/// <param name="Handler"></param>
 		abstract public void SetConsoleCtrlHandler(ProcessManager.CtrlHandlerDelegate Handler);
 
+		/// <summary>
+		/// Platform specific override to skip loading/compiling unsupported modules
+		/// </summary>
+		/// <param name="ModuleName">Module name</param>
+		/// <returns>True if module should be compiled or loaded</returns>
+		abstract public bool IsScriptModuleSupported(string ModuleName);
+
 		/// <summary>
 		/// Gets UBT project name for the current platform.
 		/// </summary>
diff --git a/Engine/Source/Programs/AutomationTool/LinuxHostPlatform.cs b/Engine/Source/Programs/AutomationTool/LinuxHostPlatform.cs
--- a/Engine/Source/Programs/AutomationTool/LinuxHostPlatform.cs
+++ b/Engine/Source/Programs/AutomationTool/LinuxHostPlatform.cs
@@ -101,6 +101,12 @@ namespace AutomationTool
 			// @todo: add mono support
 		}
 
+		public override bool IsScriptModuleSupported(string ModuleName)
+		{
+			// @todo: add more unsupported modules here
+			return String.Compare(ModuleName, "XBoxOne", true) != 0;
+		}
+
 		public override string UBTProjectName
 		{
 			get { return "UnrealBuildTool_Mono"; }
diff --git a/Engine/Source/Programs/AutomationTool/MacHostPlatform.cs b/Engine/Source/Programs/AutomationTool/MacHostPlatform.cs
--- a/Engine/Source/Programs/AutomationTool/MacHostPlatform.cs
+++ b/Engine/Source/Programs/AutomationTool/MacHostPlatform.cs
@@ -124,6 +124,12 @@ namespace AutomationTool
 			// @todo: add mono support
 		}
 
+		public override bool IsScriptModuleSupported(string ModuleName)
+		{
+			// @todo: add more unsupported modules here
+			return String.Compare(ModuleName, "XBoxOne", true) != 0;
+		}
+
 		public override string UBTProjectName
 		{
 			get { return "UnrealBuildTool_Mono"; }
diff --git a/Engine/Source/Programs/AutomationTool/ScriptCompiler.cs b/Engine/Source/Programs/AutomationTool/ScriptCompiler.cs
--- a/Engine/Source/Programs/AutomationTool/ScriptCompiler.cs
+++ b/Engine/Source/Programs/AutomationTool/ScriptCompiler.cs
@@ -140,8 +140,20 @@ namespace AutomationTool
 			Log.TraceVerbose("Discovering game folders.");
 			UnrealBuildTool.RulesCompiler.SetAssemblyNameAndGameFolders("UnrealAutomationToolRules", AllGameFolders);
 
-			var Modules = UnrealBuildTool.RulesCompiler.FindAllRulesSourceFiles(UnrealBuildTool.RulesCompiler.RulesFileType.AutomationModule, RemappedAdditionalScriptFolders);
-			CompileModules(Modules);
+			var DiscoveredModules = UnrealBuildTool.RulesCompiler.FindAllRulesSourceFiles(UnrealBuildTool.RulesCompiler.RulesFileType.AutomationModule, RemappedAdditionalScriptFolders);
+			var ModulesToCompile = new List<string>(DiscoveredModules.Count);
+			foreach (var ModuleFilename in DiscoveredModules)
+			{
+				if (HostPlatform.Current.IsScriptModuleSupported(CommandUtils.GetFilenameWithoutAnyExtensions(ModuleFilename)))
+				{
+					ModulesToCompile.Add(ModuleFilename);
+				}
+				else
+				{
+					CommandUtils.LogVerbose("Script module {0} filtered by the Host Platform and will not be compiled.", ModuleFilename);
+				}
+			}
+			CompileModules(ModulesToCompile);
 
 			Environment.CurrentDirectory = OldCWD;
 		}
@@ -243,6 +255,12 @@ namespace AutomationTool
 				CommandUtils.Log("Found {0} script DLL(s).", ScriptDLLFiles.Length);
 				foreach (var ScriptsDLLFilename in ScriptDLLFiles)
 				{
+
+					if (!HostPlatform.Current.IsScriptModuleSupported(CommandUtils.GetFilenameWithoutAnyExtensions(ScriptsDLLFilename)))
+					{
+						CommandUtils.LogVerbose("Script module {0} filtered by the Host Platform and will not be loaded.", ScriptsDLLFilename);
+						continue;
+					}
 					// Load the assembly into our app domain
 					CommandUtils.Log("Loading script DLL: {0}", ScriptsDLLFilename);
 					try
diff --git a/Engine/Source/Programs/AutomationTool/WindowsHostPlatform.cs b/Engine/Source/Programs/AutomationTool/WindowsHostPlatform.cs
--- a/Engine/Source/Programs/AutomationTool/WindowsHostPlatform.cs
+++ b/Engine/Source/Programs/AutomationTool/WindowsHostPlatform.cs
@@ -191,6 +191,11 @@ namespace AutomationTool
 			ProcessManager.SetConsoleCtrlHandler(Handler, true);
 		}
 
+		public override bool IsScriptModuleSupported(string ModuleName)
+		{
+			return true;
+		}
+
 		public override string UBTProjectName
 		{
 			get { return "UnrealBuildTool"; }
