build script, add noautomatedtesting branch hacker

diff --git a/Engine/Source/Programs/AutomationTool/Scripts/GUBP.Automation.cs b/Engine/Source/Programs/AutomationTool/Scripts/GUBP.Automation.cs
--- a/Engine/Source/Programs/AutomationTool/Scripts/GUBP.Automation.cs
+++ b/Engine/Source/Programs/AutomationTool/Scripts/GUBP.Automation.cs
@@ -192,7 +192,8 @@ public class GUBP : BuildCommand
         public class BranchOptions
         {            
             public List<UnrealTargetPlatform> PlatformsToRemove = new List<UnrealTargetPlatform>();
-            public List<string> ExcludeNodes = new List<string>();            
+            public List<string> ExcludeNodes = new List<string>();
+			public bool bNoAutomatedTesting = false;
         }
         public virtual void ModifyOptions(GUBP bp, ref BranchOptions Options, string Branch)
         {
@@ -4553,7 +4554,7 @@ public class GUBP : BuildCommand
 
 		bBuildRocket = ParseParam("BuildRocket");
         bForceIncrementalCompile = ParseParam("ForceIncrementalCompile");
-        bool bNoAutomatedTesting = ParseParam("NoAutomatedTesting");        
+        bool bNoAutomatedTesting = ParseParam("NoAutomatedTesting") || BranchOptions.bNoAutomatedTesting;		
         StoreName = ParseParamValue("Store");
         string StoreSuffix = ParseParamValue("StoreSuffix", "");
         if (bBuildRocket)
