CIS Fix: for now let all script modules compile on all platforms in UAT

diff --git a/Engine/Source/Programs/AutomationTool/LinuxHostPlatform.cs b/Engine/Source/Programs/AutomationTool/LinuxHostPlatform.cs
--- a/Engine/Source/Programs/AutomationTool/LinuxHostPlatform.cs
+++ b/Engine/Source/Programs/AutomationTool/LinuxHostPlatform.cs
@@ -104,7 +104,7 @@ namespace AutomationTool
 		public override bool IsScriptModuleSupported(string ModuleName)
 		{
 			// @todo: add more unsupported modules here
-			return String.Compare(ModuleName, "XBoxOne", true) != 0;
+			return true;
 		}
 
 		public override string UBTProjectName
diff --git a/Engine/Source/Programs/AutomationTool/MacHostPlatform.cs b/Engine/Source/Programs/AutomationTool/MacHostPlatform.cs
--- a/Engine/Source/Programs/AutomationTool/MacHostPlatform.cs
+++ b/Engine/Source/Programs/AutomationTool/MacHostPlatform.cs
@@ -127,7 +127,7 @@ namespace AutomationTool
 		public override bool IsScriptModuleSupported(string ModuleName)
 		{
 			// @todo: add more unsupported modules here
-			return String.Compare(ModuleName, "XBoxOne", true) != 0;
+			return true;
 		}
 
 		public override string UBTProjectName
