Win32 and Win64 configurations staging/archiving without overwriting each other

diff --git a/Engine/Source/Programs/AutomationTool/Scripts/CopyBuildToStagingDirectory.Automation.cs b/Engine/Source/Programs/AutomationTool/Scripts/CopyBuildToStagingDirectory.Automation.cs
--- a/Engine/Source/Programs/AutomationTool/Scripts/CopyBuildToStagingDirectory.Automation.cs
+++ b/Engine/Source/Programs/AutomationTool/Scripts/CopyBuildToStagingDirectory.Automation.cs
@@ -857,7 +857,13 @@ public partial class Project : CommandUtils
 		if (InDedicatedServer && (Params.Cook || Params.CookOnTheFly))
 		{
 			PlatformsToStage = Params.ServerTargetPlatforms;
-		}
+        }
+
+        bool prefixArchiveDir = false;
+        if (PlatformsToStage.Contains(UnrealTargetPlatform.Win32) && PlatformsToStage.Contains(UnrealTargetPlatform.Win64))
+        {
+            prefixArchiveDir = true;
+        }
 
 		List<DeploymentContext> DeploymentContexts = new List<DeploymentContext>();
 		foreach (var StagePlatform in PlatformsToStage)
@@ -884,12 +890,26 @@ public partial class Project : CommandUtils
 					}
 					ExecutablesToStage.Add(Exe);
 				}
-			}
+            }
+
+            string StageDirectory = (Params.Stage || !String.IsNullOrEmpty(Params.StageDirectoryParam)) ? Params.BaseStageDirectory : "";
+            string ArchiveDirectory = (Params.Archive || !String.IsNullOrEmpty(Params.ArchiveDirectoryParam)) ? Params.BaseArchiveDirectory : "";
+            if (prefixArchiveDir && (StagePlatform == UnrealTargetPlatform.Win32 || StagePlatform == UnrealTargetPlatform.Win64))
+            {
+                if (Params.Stage)
+                {
+                    StageDirectory = CombinePaths(Params.BaseStageDirectory, StagePlatform.ToString());
+                }
+                if (Params.Archive)
+                {
+                    ArchiveDirectory = CombinePaths(Params.BaseArchiveDirectory, StagePlatform.ToString());
+                }
+            }
 
 			//@todo should pull StageExecutables from somewhere else if not cooked
-			var SC = new DeploymentContext(Params.RawProjectPath, CmdEnv.LocalRoot,
-				(Params.Stage || !String.IsNullOrEmpty(Params.StageDirectoryParam)) ? Params.BaseStageDirectory : "",
-				(Params.Archive || !String.IsNullOrEmpty(Params.ArchiveDirectoryParam)) ? Params.BaseArchiveDirectory : "",
+            var SC = new DeploymentContext(Params.RawProjectPath, CmdEnv.LocalRoot,
+                StageDirectory,
+                ArchiveDirectory,
 				Params.CookFlavor,
 				Params.GetTargetPlatformInstance(CookedDataPlatform),
 				Params.GetTargetPlatformInstance(StagePlatform),
