[MMVP-74] Android packaging / deployment error reporting

diff --git a/Engine/Source/Editor/UnrealEd/Private/Analytics/EditorAnalytics.cpp b/Engine/Source/Editor/UnrealEd/Private/Analytics/EditorAnalytics.cpp
--- a/Engine/Source/Editor/UnrealEd/Private/Analytics/EditorAnalytics.cpp
+++ b/Engine/Source/Editor/UnrealEd/Private/Analytics/EditorAnalytics.cpp
@@ -130,12 +130,12 @@ FString FEditorAnalytics::TranslateErrorCode(int32 ErrorCode)
 		return TEXT("Failed to Code Sign");
 	case EAnalyticsErrorCodes::DeviceBackupFailed:
 		return TEXT("Failed to backup device");
-	case EAnalyticsErrorCodes::IPAUninstallFailed:
-		return TEXT("Failed to Uninstall IPA");
-	case EAnalyticsErrorCodes::IPAInstallFailed:
-		return TEXT("Failed to Install IPA");
-	case EAnalyticsErrorCodes::IPANotFound:
-		return TEXT("IPA Not Found for Install");
+	case EAnalyticsErrorCodes::AppUninstallFailed:
+		return TEXT("Failed to Uninstall app");
+	case EAnalyticsErrorCodes::AppInstallFailed:
+		return TEXT("Failed to Install app");
+	case EAnalyticsErrorCodes::AppNotFound:
+		return TEXT("App package file not found for Install");
 	case EAnalyticsErrorCodes::StubNotSignedCorrectly:
 		return TEXT("Stub not signed correctly.");
 	case EAnalyticsErrorCodes::IPAMissingInfoPList:
@@ -148,6 +148,18 @@ FString FEditorAnalytics::TranslateErrorCode(int32 ErrorCode)
 		return TEXT("Could not create directory");
 	case EAnalyticsErrorCodes::CopyFile:
 		return TEXT("Could not copy file");
+	case EAnalyticsErrorCodes::OnlyOneObbFileSupported:
+		return TEXT("Android packaging supports only exactly one obb/pak file");
+	case EAnalyticsErrorCodes::FailureGettingPackageInfo:
+		return TEXT("Failed to get package info from APK file");
+	case EAnalyticsErrorCodes::OnlyOneTargetConfigurationSupported:
+		return TEXT("Android is only able to package a single target configuration");
+	case EAnalyticsErrorCodes::ObbNotFound:
+		return TEXT("OBB/PAK file not found");
+	case EAnalyticsErrorCodes::AndroidBuildToolsPathNotFound:
+		return TEXT("Android build-tools directory not found");
+	case EAnalyticsErrorCodes::NoApkSuitableForArchitecture:
+		return TEXT("No APK suitable for architecture found");
 	}
 	return TEXT("Unknown Error");
 }
diff --git a/Engine/Source/Editor/UnrealEd/Public/EditorAnalytics.h b/Engine/Source/Editor/UnrealEd/Public/EditorAnalytics.h
--- a/Engine/Source/Editor/UnrealEd/Public/EditorAnalytics.h
+++ b/Engine/Source/Editor/UnrealEd/Public/EditorAnalytics.h
@@ -38,15 +38,21 @@ namespace EAnalyticsErrorCodes
 		FailedToCreateIPA = 31,
 		FailedToCodeSign = 32,
 		DeviceBackupFailed = 33,
-		IPAUninstallFailed = 34,
-		IPAInstallFailed = 35,
-		IPANotFound = 36,
+		AppUninstallFailed = 34,
+		AppInstallFailed = 35,
+		AppNotFound = 36,
 		StubNotSignedCorrectly = 37,
 		IPAMissingInfoPList = 38,
 		DeleteFile = 39,
 		DeleteDirectory = 40,
 		CreateDirectory = 41,
 		CopyFile = 42,
+		OnlyOneObbFileSupported = 50,
+		FailureGettingPackageInfo = 51,
+		OnlyOneTargetConfigurationSupported = 52,
+		ObbNotFound = 53,
+		AndroidBuildToolsPathNotFound = 54,
+		NoApkSuitableForArchitecture = 55,
 		LauncherFailed = 100,
 		UATLaunchFailure = 101,
 	};
diff --git a/Engine/Source/Programs/AutomationTool/Android/AndroidPlatform.Automation.cs b/Engine/Source/Programs/AutomationTool/Android/AndroidPlatform.Automation.cs
--- a/Engine/Source/Programs/AutomationTool/Android/AndroidPlatform.Automation.cs
+++ b/Engine/Source/Programs/AutomationTool/Android/AndroidPlatform.Automation.cs
@@ -60,12 +60,14 @@ public class AndroidPlatform : Platform
 		string PackageName = GetPackageInfo(ApkName, false);
 		if (PackageName == null)
 		{
+            ErrorReporter.Error("Failed to get package name from " + ApkName, (int)ErrorCodes.Error_FailureGettingPackageInfo);
 			throw new AutomationException("Failed to get package name from " + ApkName);
 		}
 
 		string PackageVersion = GetPackageInfo(ApkName, true);
 		if (PackageVersion == null || PackageVersion.Length == 0)
 		{
+            ErrorReporter.Error("Failed to get package version from " + ApkName, (int)ErrorCodes.Error_FailureGettingPackageInfo);
 			throw new AutomationException("Failed to get package version from " + ApkName);
 		}
 
@@ -123,7 +125,9 @@ public class AndroidPlatform : Platform
 			// for now, we only support 1 pak/obb file
 			if (PakFiles.Length > 1)
 			{
-				throw new AutomationException("Can't package for Android with 0 or more than 1 pak file (found {0} pak files in {1})", PakFiles.Length, SC.StageDirectory);
+                string ErrorString = String.Format("Can't package for Android with 0 or more than 1 pak file (found {0} pak files in {1})", PakFiles.Length, SC.StageDirectory);
+                ErrorReporter.Error(ErrorString, (int)ErrorCodes.Error_OnlyOneObbFileSupported);
+                throw new AutomationException(ErrorString);
 			}
 
 			string LocalObbName = GetFinalObbName(ApkName);
@@ -178,7 +182,9 @@ public class AndroidPlatform : Platform
 	{
 		if (SC.StageTargetConfigurations.Count != 1)
 		{
-			throw new AutomationException("Android is currently only able to package one target configuration at a time, but StageTargetConfigurations contained {0} configurations", SC.StageTargetConfigurations.Count);
+            string ErrorString = String.Format("Android is currently only able to package one target configuration at a time, but StageTargetConfigurations contained {0} configurations", SC.StageTargetConfigurations.Count);
+            ErrorReporter.Error(ErrorString, (int)ErrorCodes.Error_OnlyOneTargetConfigurationSupported);
+			throw new AutomationException(ErrorString);
 		}
 
 		string[] Architectures = UnrealBuildTool.AndroidToolChain.GetAllArchitectures();
@@ -193,12 +199,16 @@ public class AndroidPlatform : Platform
 			// verify the files exist
 			if (!FileExists(ApkName))
 			{
-				throw new AutomationException("ARCHIVE FAILED - {0} was not found", ApkName);
+                string ErrorString = String.Format("ARCHIVE FAILED - {0} was not found", ApkName);
+                ErrorReporter.Error(ErrorString, (int)ErrorCodes.Error_AppNotFound);
+                throw new AutomationException(ErrorString);
 			}
 			if (!Params.OBBinAPK && !FileExists(ObbName))
 			{
-				throw new AutomationException("ARCHIVE FAILED - {0} was not found", ObbName);
-			}
+                string ErrorString = String.Format("ARCHIVE FAILED - {0} was not found", ObbName);
+                ErrorReporter.Error(ErrorString, (int)ErrorCodes.Error_ObbNotFound);
+                throw new AutomationException(ErrorString);
+            }
 
 			SC.ArchiveFiles(Path.GetDirectoryName(ApkName), Path.GetFileName(ApkName));
 			if (!Params.OBBinAPK)
@@ -319,16 +329,44 @@ public class AndroidPlatform : Platform
 			Deploy.PrepForUATPackageOrDeploy(Params.ShortProjectName, SC.ProjectRoot, SOName, SC.LocalRoot + "/Engine", Params.Distribution, CookFlavor);
 		}
 
+        // check the APK exists
+        if (!File.Exists(ApkName))
+        {
+            ErrorReporter.Error(String.Format("Could not find apk '{0}'", ApkName), (int)ErrorCodes.Error_AppNotFound);
+            throw new AutomationException("Could not find apk '{0}'", ApkName);
+        }
+
 		// now we can use the apk to get more info
 		string DeviceObbName = GetDeviceObbName(ApkName);
 		string PackageName = GetPackageInfo(ApkName, false);
 
-		// install the apk
-		string UninstallCommandline = AdbCommand + "uninstall " + PackageName;
-		RunAndLog(CmdEnv, CmdEnv.CmdExe, UninstallCommandline);
+        // try uninstalling an old app with the same identifier.
+        string UninstallCommandline = AdbCommand + "uninstall " + PackageName;
+        int SuccessCode = 0;
+        RunAndLog(CmdEnv, CmdEnv.CmdExe, UninstallCommandline, out SuccessCode);
 
+		// install the apk
+        SuccessCode = 0;
 		string InstallCommandline = AdbCommand + "install \"" + ApkName + "\"";
-		RunAndLog(CmdEnv, CmdEnv.CmdExe, InstallCommandline);
+        string InstallOutput = RunAndLog(CmdEnv, CmdEnv.CmdExe, InstallCommandline, out SuccessCode);
+        int FailureIndex = InstallOutput.IndexOf("Failure"); 
+
+        // adb install doesn't always return an error code on failure, and instead prints "Failure", followed by an error code.
+        if (SuccessCode != 0 || FailureIndex != -1)
+        {
+            string ErrorMessage = String.Format("Installation of apk '{0}' failed", ApkName);
+            if (FailureIndex != -1)
+            {
+                string FailureString = InstallOutput.Substring(FailureIndex + 7).Trim();
+                if (FailureString != "")
+                {
+                    ErrorMessage += ": " + FailureString;
+                }
+            }
+
+            ErrorReporter.Error(ErrorMessage, (int)ErrorCodes.Error_AppInstallFailed);
+            throw new AutomationException(ErrorMessage);
+        }
 
 		// update the ue4commandline.txt
 		// update and deploy ue4commandline.txt
@@ -549,6 +587,7 @@ public class AndroidPlatform : Platform
 		string[] Subdirs = Directory.GetDirectories(Environment.ExpandEnvironmentVariables("%ANDROID_HOME%/build-tools/"));
 		if (Subdirs.Length == 0)
 		{
+            ErrorReporter.Error("Failed to find %ANDROID_HOME%/build-tools subdirectory", (int)ErrorCodes.Error_AndroidBuildToolsPathNotFound);
 			throw new AutomationException("Failed to find %ANDROID_HOME%/build-tools subdirectory");
 		}
 		// we expect there to be one, so use the first one
@@ -609,7 +648,9 @@ public class AndroidPlatform : Platform
 		// if after the fallbacks, we still don't have it, we can't continue
 		if (Array.IndexOf(AppArchitectures, DeviceArch) == -1)
 		{
-			throw new AutomationException("Unable to run because you don't have an apk that is usable on {0}", Params.Device);
+            string ErrorString = String.Format("Unable to run because you don't have an apk that is usable on {0}", Params.Device);
+            ErrorReporter.Error(ErrorString, (int)ErrorCodes.Error_NoApkSuitableForArchitecture);
+            throw new AutomationException(ErrorString);
 		}
 
 		return DeviceArch;
@@ -632,7 +673,8 @@ public class AndroidPlatform : Platform
 		string PackageName = GetPackageInfo(ApkName, false);
 		if (PackageName == null)
 		{
-			throw new AutomationException("Failed to get package name from " + ClientApp);
+            ErrorReporter.Error("Failed to get package name from " + ClientApp, (int)ErrorCodes.Error_FailureGettingPackageInfo);
+            throw new AutomationException("Failed to get package name from " + ClientApp);
 		}
 
 		string AdbCommand = GetAdbCommand(Params);
diff --git a/Engine/Source/Programs/AutomationTool/ProcessUtils.cs b/Engine/Source/Programs/AutomationTool/ProcessUtils.cs
--- a/Engine/Source/Programs/AutomationTool/ProcessUtils.cs
+++ b/Engine/Source/Programs/AutomationTool/ProcessUtils.cs
@@ -818,6 +818,19 @@ namespace AutomationTool
             return "";
         }
 
+        /// <summary>
+        /// Runs external program and writes the output to a logfile.
+        /// </summary>
+        /// <param name="Env">Environment to use.</param>
+        /// <param name="App">Executable to run</param>
+        /// <param name="CommandLine">Commandline to pass on to the executable</param>
+        /// <param name="LogName">Name of the logfile ( if null, executable name is used )</param>
+        /// <returns>Whether the program executed successfully or not.</returns>
+        public static string RunAndLog(CommandEnvironment Env, string App, string CommandLine, out int SuccessCode, string LogName = null)
+        {
+            return RunAndLog(App, CommandLine, out SuccessCode, GetRunAndLogLogName(Env, App, LogName));
+        }
+
 		/// <summary>
 		/// Runs UAT recursively
 		/// </summary>
diff --git a/Engine/Source/Programs/AutomationTool/Program.cs b/Engine/Source/Programs/AutomationTool/Program.cs
--- a/Engine/Source/Programs/AutomationTool/Program.cs
+++ b/Engine/Source/Programs/AutomationTool/Program.cs
@@ -39,16 +39,22 @@ namespace AutomationTool
 		Error_FailedToCreateIPA = 31,
 		Error_FailedToCodeSign = 32,
 		Error_DeviceBackupFailed = 33,
-		Error_IPAUninstallFailed = 34,
-		Error_IPAInstallFailed = 35,
-		Error_IPANotFound = 36,
+		Error_AppUninstallFailed = 34,
+		Error_AppInstallFailed = 35,
+		Error_AppNotFound = 36,
 		Error_StubNotSignedCorrectly = 37,
 		Error_IPAMissingInfoPList = 38,
 		Error_DeleteFile = 39,
 		Error_DeleteDirectory = 40,
 		Error_CreateDirectory = 41,
 		Error_CopyFile = 42,
-		Error_LauncherFailed = 100,
+        Error_OnlyOneObbFileSupported = 50,
+        Error_FailureGettingPackageInfo = 51,
+        Error_OnlyOneTargetConfigurationSupported = 52,
+        Error_ObbNotFound = 53,
+        Error_AndroidBuildToolsPathNotFound = 54,
+        Error_NoApkSuitableForArchitecture = 55,
+        Error_LauncherFailed = 100,
 		Error_UATLaunchFailure = 101,
 	};
 
diff --git a/Engine/Source/Programs/IOS/iPhonePackager/DeployTime.cs b/Engine/Source/Programs/IOS/iPhonePackager/DeployTime.cs
--- a/Engine/Source/Programs/IOS/iPhonePackager/DeployTime.cs
+++ b/Engine/Source/Programs/IOS/iPhonePackager/DeployTime.cs
@@ -129,7 +129,7 @@ namespace iPhonePackager
 						if (!DeploymentHelper.Get().UninstallIPAOnDevice(ApplicationIdentifier))
 						{
 							Program.Error("Failed to uninstall IPA on device");
-							Program.ReturnCode = (int)ErrorCodes.Error_IPAUninstallFailed;
+							Program.ReturnCode = (int)ErrorCodes.Error_AppUninstallFailed;
 						}
 					}
 					break;
@@ -147,13 +147,13 @@ namespace iPhonePackager
 							if (!DeploymentHelper.Get().InstallIPAOnDevice(IPAPath))
 							{
 								Program.Error("Failed to install IPA on device");
-								Program.ReturnCode = (int)ErrorCodes.Error_IPAInstallFailed;
+								Program.ReturnCode = (int)ErrorCodes.Error_AppInstallFailed;
 							}
 						}
 						else
 						{
 							Program.Error(String.Format("Failed to find IPA file: '{0}'", IPAPath));
-							Program.ReturnCode = (int)ErrorCodes.Error_IPANotFound;
+							Program.ReturnCode = (int)ErrorCodes.Error_AppNotFound;
 						}
 					}
 					break;
diff --git a/Engine/Source/Programs/IOS/iPhonePackager/iPhonePackager.cs b/Engine/Source/Programs/IOS/iPhonePackager/iPhonePackager.cs
--- a/Engine/Source/Programs/IOS/iPhonePackager/iPhonePackager.cs
+++ b/Engine/Source/Programs/IOS/iPhonePackager/iPhonePackager.cs
@@ -46,15 +46,21 @@ namespace iPhonePackager
 		Error_FailedToCreateIPA = 31,
 		Error_FailedToCodeSign = 32,
 		Error_DeviceBackupFailed = 33,
-		Error_IPAUninstallFailed = 34,
-		Error_IPAInstallFailed = 35,
-		Error_IPANotFound = 36,
+		Error_AppUninstallFailed = 34,
+		Error_AppInstallFailed = 35,
+		Error_AppNotFound = 36,
 		Error_StubNotSignedCorrectly = 37,
 		Error_IPAMissingInfoPList = 38,
 		Error_DeleteFile = 39,
 		Error_DeleteDirectory = 40,
 		Error_CreateDirectory = 41,
 		Error_CopyFile = 42,
+        Error_OnlyOneObbFileSupported = 50,
+        Error_FailureGettingPackageInfo = 51,
+        Error_OnlyOneTargetConfigurationSupported = 52,
+        Error_ObbNotFound = 53,
+        Error_AndroidBuildToolsPathNotFound = 54,
+        Error_NoApkSuitableForArchitecture = 55,
 		Error_LauncherFailed = 100,
 		Error_UATLaunchFailure = 101,
 	};
