Fix up HTML5 Automation to correctly setup command line arguments. Make safari work correctly - it doesn't honor command line arguments like every other browser.
 #ttp 342330

diff --git a/Engine/Source/Developer/DesktopPlatform/Private/PlatformInfo.cpp b/Engine/Source/Developer/DesktopPlatform/Private/PlatformInfo.cpp
--- a/Engine/Source/Developer/DesktopPlatform/Private/PlatformInfo.cpp
+++ b/Engine/Source/Developer/DesktopPlatform/Private/PlatformInfo.cpp
@@ -102,7 +102,7 @@ static const FPlatformInfo PlatformInfoArray[] = {
 	BuildPlatformInfo(TEXT("Android_ETC2"),				TEXT("Android_ETC2"),		LOCTEXT("Android_ETC2", "Android (ETC2)"),					EPlatformType::Game,		EPlatformFlags::CookFlavor,		FPlatformIconPaths(TEXT("Launcher/AndroidTarget/Platform_Android_ETC2_24x"), TEXT("Launcher/AndroidTarget/Platform_Android_128x")),					TEXT("-targetplatform=Android -cookflavor=ETC2"),		TEXT("Android"),	EPlatformSDKStatus::Unknown,	TEXT("Shared/Tutorials/SettingUpAndroidTutorial"),			IsAvailableOnWindows,						true),
 	BuildPlatformInfo(TEXT("Android_PVRTC"),			TEXT("Android_PVRTC"),		LOCTEXT("Android_PVRTC", "Android (PVRTC)"),				EPlatformType::Game,		EPlatformFlags::CookFlavor,		FPlatformIconPaths(TEXT("Launcher/AndroidTarget/Platform_Android_PVRTC_24x"), TEXT("Launcher/AndroidTarget/Platform_Android_128x")),				TEXT("-targetplatform=Android -cookflavor=PVRTC"),		TEXT("Android"),	EPlatformSDKStatus::Unknown,	TEXT("Shared/Tutorials/SettingUpAndroidTutorial"),			IsAvailableOnWindows,						true),
 
-	BuildPlatformInfo(TEXT("HTML5"),					TEXT("HTML5"),				LOCTEXT("HTML5", "HTML5"),									EPlatformType::Game,		EPlatformFlags::None,			FPlatformIconPaths(TEXT("Launcher/HTML5Target/Platform_HTML5_24x"), TEXT("Launcher/HTML5Target/Platform_HTML5_128x")),								TEXT(""),												TEXT("HTML5"),		EPlatformSDKStatus::Unknown,	TEXT("Platforms/HTML5/GettingStarted"),						IsAvailableOnWindows ,	                    true),
+	BuildPlatformInfo(TEXT("HTML5"),					TEXT("HTML5"),				LOCTEXT("HTML5", "HTML5"),									EPlatformType::Game,		EPlatformFlags::None,			FPlatformIconPaths(TEXT("Launcher/HTML5Target/Platform_HTML5_24x"), TEXT("Launcher/HTML5Target/Platform_HTML5_128x")),								TEXT(""),												TEXT("HTML5"),		EPlatformSDKStatus::Unknown,	TEXT("Platforms/HTML5/GettingStarted"),						IsAvailableOnWindows || IsAvailableOnMac,   IsAvailableOnWindows),
 
 	BuildPlatformInfo(TEXT("PS4"),						TEXT("PS4"),				LOCTEXT("PS4", "PlayStation 4"),							EPlatformType::Game,		EPlatformFlags::None,			FPlatformIconPaths(TEXT("Launcher/PS4Target/Platform_PS4_24x"), TEXT("Launcher/PS4Target/Platform_PS4_128x")),										TEXT(""),												TEXT("PS4"),		EPlatformSDKStatus::Unknown,	TEXT("Platforms/PS4/GettingStarted"),						IsAvailableOnWindows,						false),
 
diff --git a/Engine/Source/Programs/AutomationTool/HTML5/HTML5Platform.Automation.cs b/Engine/Source/Programs/AutomationTool/HTML5/HTML5Platform.Automation.cs
--- a/Engine/Source/Programs/AutomationTool/HTML5/HTML5Platform.Automation.cs
+++ b/Engine/Source/Programs/AutomationTool/HTML5/HTML5Platform.Automation.cs
@@ -6,6 +6,8 @@ using System.Text;
 using System.IO;
 using AutomationTool;
 using UnrealBuildTool;
+using System.Diagnostics;
+using System.Threading;
 
 public class HTML5Platform : Platform
 {
@@ -79,7 +81,7 @@ public class HTML5Platform : Platform
 			}
 		}
 
-		HeapSize = (ulong)ConfigHeapSize * 1024L * 1024L; // convert to bytes. 
+		HeapSize = (ulong)ConfigHeapSize * 1024L * 1024L; // convert to bytes.
 		Log("Setting Heap size to {0} Mb ", ConfigHeapSize);
 
 
@@ -125,12 +127,13 @@ public class HTML5Platform : Platform
 
 				if (LineStr.Contains("%UE4CMDLINE%"))
 				{
+					InArguments = InArguments.Replace ("\"", "");
 					string[] Arguments = InArguments.Split(' ');
 					string ArgumentString = IsContentOnly ? "'" + InGameName + "/" + InGameName + ".uproject '," : "";
 					for (int i = 0; i < Arguments.Length - 1; ++i)
 					{
 						ArgumentString += "'";
-						ArgumentString += Arguments[i];
+			   		    ArgumentString += Arguments[i];
 						ArgumentString += "'";
 						ArgumentString += ",' ',";
 					}
@@ -203,7 +206,7 @@ public class HTML5Platform : Platform
         // look for browser
         var ConfigCache = new UnrealBuildTool.ConfigCacheIni(UnrealTargetPlatform.HTML5, "Engine", Path.GetDirectoryName(Params.RawProjectPath), CombinePaths(CmdEnv.LocalRoot, "Engine"));
 
-        string DeviceSection; 
+        string DeviceSection;
 
         if ( Utils.IsRunningOnMono )
         {
@@ -217,15 +220,15 @@ public class HTML5Platform : Platform
         string browserPath;
         string DeviceName = Params.Device.Split('@')[1];
         bool ok = ConfigCache.GetString(DeviceSection, DeviceName, out browserPath);
-        
+
 		if (!ok)
-			throw new System.Exception ("Incorrect browser configuration in HTML5Engine.ini "); 
+			throw new System.Exception ("Incorrect browser configuration in HTML5Engine.ini ");
 
 		// open the webpage
         string directory = Path.GetDirectoryName(ClientApp);
-        string url = Path.GetFileName(ClientApp) +".html"; 
-		// Are we running via cook on the fly server? 
-		// find our http url - This is awkward because RunClient doesn't have real information that NFS is running or not. 
+        string url = Path.GetFileName(ClientApp) +".html";
+		// Are we running via cook on the fly server?
+		// find our http url - This is awkward because RunClient doesn't have real information that NFS is running or not.
 		bool IsCookOnTheFly = false;
 		string[] Arguments = ClientCmdLine.Split(' ');
 		foreach (var Argument in Arguments)
@@ -233,7 +236,7 @@ public class HTML5Platform : Platform
 			string[] KeyVal = Argument.Split('=');
 			if (KeyVal[0].ToLower().Contains("filehostip"))
 			{
-				// split urls. 
+				// split urls.
 				string[] Urls = KeyVal[1].Split('+');
 				foreach (var Url in Urls)
 				{
@@ -256,18 +259,18 @@ public class HTML5Platform : Platform
         {
 
             url = "http://127.0.0.1:8000/" + url;
-            // this will be killed UBT instances dies. 
+            // this will be killed UBT instances dies.
             string input = String.Format(" -m SimpleHTTPServer 8000");
 
 
 			string PythonName = "python2.exe";
 
 			if (Utils.IsRunningOnMono)
-				PythonName = "python2"; 
+				PythonName = "python2";
 
 			ProcessResult Result = ProcessManager.CreateProcess(PythonName, true, "html5server.log");
 			Result.ProcessObject.StartInfo.FileName = PythonName;
-            Result.ProcessObject.StartInfo.UseShellExecute = false; 
+            Result.ProcessObject.StartInfo.UseShellExecute = false;
             Result.ProcessObject.StartInfo.RedirectStandardOutput = true;
             Result.ProcessObject.StartInfo.RedirectStandardInput = true;
             Result.ProcessObject.StartInfo.WorkingDirectory = directory;
@@ -280,26 +283,42 @@ public class HTML5Platform : Platform
                 System.Console.WriteLine(e.Data);
             };
 
-            System.Console.WriteLine("Starting Browser Process"); 
+            System.Console.WriteLine("Starting Browser Process");
+
 
-            ProcessResult ClientProcess = Run(browserPath, url, null, ClientRunFlags | ERunOptions.NoWaitForExit);
+			// safari specific hack.
+			string argument = url; 
+			if (browserPath.Contains ("Safari") && Utils.IsRunningOnMono)
+				argument = ""; 
 
-            ClientProcess.ProcessObject.EnableRaisingEvents = true; 
+            ProcessResult ClientProcess = Run(browserPath, argument, null, ClientRunFlags | ERunOptions.NoWaitForExit);
+
+            ClientProcess.ProcessObject.EnableRaisingEvents = true;
             ClientProcess.ProcessObject.Exited += delegate(System.Object o, System.EventArgs e)
             {
                 System.Console.WriteLine("Browser Process Ended - Killing Webserver");
-                // send kill. 
+                // send kill.
                 Result.ProcessObject.StandardInput.Close();
                 Result.ProcessObject.Kill();
             };
 
+			// safari needs a hack. 
+			// http://superuser.com/questions/689315/run-safari-from-terminal-with-given-url-address-without-open-command
+
+			if (browserPath.Contains ("Safari") && Utils.IsRunningOnMono) 
+			{
+				// ClientProcess.ProcessObject.WaitForInputIdle ();
+				Thread.Sleep (2000);
+				Process.Start("/usr/bin/osascript"," -e 'tell application \"Safari\" to open location \"" + url + "\"'" );
+			}
+
             return ClientProcess;
         }
 
         System.Console.WriteLine("Browser Path " + browserPath);
         ProcessResult BrowserProcess = Run(browserPath, url, null, ClientRunFlags | ERunOptions.NoWaitForExit);
-
-        return BrowserProcess; 
+	
+        return BrowserProcess;
 
 	}
 
