Temporary fix to allow chunk reuse beyond default timescale in patch generation processes.

diff --git a/Engine/Source/Programs/AutomationTool/MCPPublic.cs b/Engine/Source/Programs/AutomationTool/MCPPublic.cs
--- a/Engine/Source/Programs/AutomationTool/MCPPublic.cs
+++ b/Engine/Source/Programs/AutomationTool/MCPPublic.cs
@@ -250,7 +250,9 @@ namespace EpicGames.MCP.Automation
 
             public PatchGenerationOptions()
             {
-                DataAgeThreshold = DEFAULT_DATA_AGE_THRESHOLD;
+				// @todo Richard.Fawcett temporary work-around for build system issues.
+				// Hack to add 10,000 days should not remain beyond 2014-09-30!
+				DataAgeThreshold = DEFAULT_DATA_AGE_THRESHOLD + 10000;
             }
 
             /// <summary>
