Fix CopyBuildToStagingDirectory automation script to supply new DataAgeThreshold parameter to BPT.exe

diff --git a/Engine/Source/Programs/AutomationTool/Scripts/CopyBuildToStagingDirectory.Automation.cs b/Engine/Source/Programs/AutomationTool/Scripts/CopyBuildToStagingDirectory.Automation.cs
--- a/Engine/Source/Programs/AutomationTool/Scripts/CopyBuildToStagingDirectory.Automation.cs
+++ b/Engine/Source/Programs/AutomationTool/Scripts/CopyBuildToStagingDirectory.Automation.cs
@@ -537,6 +537,7 @@ public partial class Project : CommandUtils
 				string VersionString = P4Branch + "-CL-" + P4Change;
 
 				string CmdLine = String.Format("-BuildRoot=\"{0}\" -CloudDir=\"{1}\" -AppID={2} -AppName=\"{3}\" -BuildVersion=\"{4}\" -AppLaunch=\"{5}\"", BuildRoot, CloudDir, AppID, AppName, VersionString, AppLaunch);
+				CmdLine += " -DataAgeThreshold=12";
 				CmdLine += " -AppArgs=\"\"";
 				CmdLine += " -custom=\"bIsPatch=false\"";
 				CmdLine += String.Format(" -customint=\"ChunkID={0}\"", ChunkID);
