Add support for manifest custom arguments when calling BuildPatchTool.exe through automation framework.

diff --git a/Engine/Source/Programs/AutomationTool/MCPPublic.cs b/Engine/Source/Programs/AutomationTool/MCPPublic.cs
--- a/Engine/Source/Programs/AutomationTool/MCPPublic.cs
+++ b/Engine/Source/Programs/AutomationTool/MCPPublic.cs
@@ -291,6 +291,18 @@ namespace EpicGames.MCP.Automation
             /// two days less than the compactify data age threshold.
             /// </summary>
             public int DataAgeThreshold;
+			/// <summary>
+			/// Contains a list of custom string arguments to be embedded in the generated manifest file.
+			/// </summary>
+			public List<KeyValuePair<string, string>> CustomStringArgs;
+			/// <summary>
+			/// Contains a list of custom integer arguments to be embedded in the generated manifest file.
+			/// </summary>
+			public List<KeyValuePair<string, int>> CustomIntArgs;
+			/// <summary>
+			/// Contains a list of custom float arguments to be embedded in the generated manifest file.
+			/// </summary>
+			public List<KeyValuePair<string, float>> CustomFloatArgs;
 		}
 
         static BuildPatchToolBase Handler = null;
