Ensure that BuildPatchTool.exe is build in UAT when necessary.

diff --git a/Engine/Source/Programs/AutomationTool/Scripts/CopyBuildToStagingDirectory.Automation.cs b/Engine/Source/Programs/AutomationTool/Scripts/CopyBuildToStagingDirectory.Automation.cs
--- a/Engine/Source/Programs/AutomationTool/Scripts/CopyBuildToStagingDirectory.Automation.cs
+++ b/Engine/Source/Programs/AutomationTool/Scripts/CopyBuildToStagingDirectory.Automation.cs
@@ -507,6 +507,13 @@ public partial class Project : CommandUtils
 			{
 				var BPTExe = CombinePaths(CmdEnv.LocalRoot, "Engine/Binaries/Win64/BuildPatchTool.exe");
 
+				// Ensure that BuildPatchTool.exe exists
+				if (!InternalUtils.SafeFileExists(BPTExe))
+				{
+					CommandUtils.Log("Building BuildPatchTool.exe");
+					UE4BuildUtils.BuildBuildPatchTool(null, UnrealTargetPlatform.Win64);
+				}
+
 				string P4Change = "UnknownCL";
 				string P4Branch = "UnknownBranch";
 				if (CommandUtils.P4Enabled)
diff --git a/Engine/Source/Programs/AutomationTool/Scripts/UE4BuildUtils.cs b/Engine/Source/Programs/AutomationTool/Scripts/UE4BuildUtils.cs
--- a/Engine/Source/Programs/AutomationTool/Scripts/UE4BuildUtils.cs
+++ b/Engine/Source/Programs/AutomationTool/Scripts/UE4BuildUtils.cs
@@ -19,6 +19,11 @@ public class UE4BuildUtils : CommandUtils
 	{
 		Log("Building BuildPatchTool");
 
+		if (Command == null)
+		{
+			Command = new UE4BuildUtilDummyBuildCommand();
+		}
+
 		var UE4Build = new UE4Build(Command);
 
 		var Agenda = new UE4Build.BuildAgenda();
@@ -33,4 +38,12 @@ public class UE4BuildUtils : CommandUtils
 		UE4Build.Build(Agenda, InDeleteBuildProducts: true, InUpdateVersionFiles: true);
 		UE4Build.CheckBuildProducts(UE4Build.BuildProductFiles);
 	}
+
+	class UE4BuildUtilDummyBuildCommand : BuildCommand
+	{
+		public override void ExecuteBuild()
+		{
+			// noop
+		}
+	}
 }
