Switch UAT over to read from v2 Build Info Service

diff --git a/Engine/Source/Programs/AutomationTool/MCPPublic.cs b/Engine/Source/Programs/AutomationTool/MCPPublic.cs
--- a/Engine/Source/Programs/AutomationTool/MCPPublic.cs
+++ b/Engine/Source/Programs/AutomationTool/MCPPublic.cs
@@ -626,7 +626,7 @@ namespace EpicGames.MCP.Config
 		}
 		protected virtual bool UseV2BuildInfoService
 		{
-			get { return false; }
+			get { return true; }
 		}
     }
 
