Added EditorTutorials localization target to UE4 localization automation script and OneSky processing.

diff --git a/Engine/Source/Programs/AutomationTool/Scripts/Localisation.Automation.cs b/Engine/Source/Programs/AutomationTool/Scripts/Localisation.Automation.cs
--- a/Engine/Source/Programs/AutomationTool/Scripts/Localisation.Automation.cs
+++ b/Engine/Source/Programs/AutomationTool/Scripts/Localisation.Automation.cs
@@ -44,6 +44,7 @@ class Localise : BuildCommand
         // Export all text from OneSky
         ExportProjectToDirectory(oneSkyService, projectGroup, "Engine");
         ExportProjectToDirectory(oneSkyService, projectGroup, "Editor");
+        ExportProjectToDirectory(oneSkyService, projectGroup, "EditorTutorials");
         ExportProjectToDirectory(oneSkyService, projectGroup, "PropertyNames");
         ExportProjectToDirectory(oneSkyService, projectGroup, "ToolTips");
 
@@ -68,6 +69,7 @@ class Localise : BuildCommand
 			{
 				String.Format("-config={0}", @"./Config/Localization/Engine.ini") + (string.IsNullOrEmpty(CommandletSCCArguments) ? "" : " " + CommandletSCCArguments),
 				String.Format("-config={0}", @"./Config/Localization/Editor.ini") + (string.IsNullOrEmpty(CommandletSCCArguments) ? "" : " " + CommandletSCCArguments),
+				String.Format("-config={0}", @"./Config/Localization/EditorTutorials.ini") + (string.IsNullOrEmpty(CommandletSCCArguments) ? "" : " " + CommandletSCCArguments),
 				String.Format("-config={0}", @"./Config/Localization/PropertyNames.ini") + (string.IsNullOrEmpty(CommandletSCCArguments) ? "" : " " + CommandletSCCArguments),
 				String.Format("-config={0}", @"./Config/Localization/ToolTips.ini") + (string.IsNullOrEmpty(CommandletSCCArguments) ? "" : " " + CommandletSCCArguments),
 				String.Format("-config={0}", @"./Config/Localization/WordCount.ini"),
@@ -92,6 +94,7 @@ class Localise : BuildCommand
         // Upload all text to OneSky
         UploadDirectoryToProject(GetProject(oneSkyService, "Unreal Engine", "Engine"), new DirectoryInfo(CmdEnv.LocalRoot + "/Engine/Content/Localization/Engine"), "*.po");
         UploadDirectoryToProject(GetProject(oneSkyService, "Unreal Engine", "Editor"), new DirectoryInfo(CmdEnv.LocalRoot + "/Engine/Content/Localization/Editor"), "*.po");
+        UploadDirectoryToProject(GetProject(oneSkyService, "Unreal Engine", "EditorTutorials"), new DirectoryInfo(CmdEnv.LocalRoot + "/Engine/Content/Localization/EditorTutorials"), "*.po");
         UploadDirectoryToProject(GetProject(oneSkyService, "Unreal Engine", "PropertyNames"), new DirectoryInfo(CmdEnv.LocalRoot + "/Engine/Content/Localization/PropertyNames"), "*.po");
         UploadDirectoryToProject(GetProject(oneSkyService, "Unreal Engine", "ToolTips"), new DirectoryInfo(CmdEnv.LocalRoot + "/Engine/Content/Localization/ToolTips"), "*.po");
 
