Reverted change to force build of BuildPatchTool.exe during UAT, which was causing Mac builds to fail.

diff --git a/Engine/Source/Programs/AutomationTool/Scripts/CopyBuildToStagingDirectory.Automation.cs b/Engine/Source/Programs/AutomationTool/Scripts/CopyBuildToStagingDirectory.Automation.cs
--- a/Engine/Source/Programs/AutomationTool/Scripts/CopyBuildToStagingDirectory.Automation.cs
+++ b/Engine/Source/Programs/AutomationTool/Scripts/CopyBuildToStagingDirectory.Automation.cs
@@ -508,11 +508,11 @@ public partial class Project : CommandUtils
 				var BPTExe = CombinePaths(CmdEnv.LocalRoot, "Engine/Binaries/Win64/BuildPatchTool.exe");
 
 				// Ensure that BuildPatchTool.exe exists
-				if (!InternalUtils.SafeFileExists(BPTExe))
-				{
-					CommandUtils.Log("Building BuildPatchTool.exe");
-					UE4BuildUtils.BuildBuildPatchTool(null, UnrealTargetPlatform.Win64);
-				}
+				//if (!InternalUtils.SafeFileExists(BPTExe))
+				//{
+				//	CommandUtils.Log("Building BuildPatchTool.exe");
+				//	UE4BuildUtils.BuildBuildPatchTool(null, UnrealTargetPlatform.Win64);
+				//}
 
 				string P4Change = "UnknownCL";
 				string P4Branch = "UnknownBranch";
