build script, support for code project tools

diff --git a/Engine/Source/Programs/AutomationTool/Scripts/GUBP.Automation.cs b/Engine/Source/Programs/AutomationTool/Scripts/GUBP.Automation.cs
--- a/Engine/Source/Programs/AutomationTool/Scripts/GUBP.Automation.cs
+++ b/Engine/Source/Programs/AutomationTool/Scripts/GUBP.Automation.cs
@@ -4878,7 +4878,7 @@ public class GUBP : BuildCommand
         foreach (var HostPlatform in HostPlatforms)
         {
             AddNode(new ToolsForCompileNode(HostPlatform));
-            AddNode(new RootEditorNode(HostPlatform));
+            AddNode(new RootEditorNode(HostPlatform));			
             if (bBuildRocket)
             {
                 AddNode(new RootEditorHeadersNode(HostPlatform));
@@ -4902,6 +4902,26 @@ public class GUBP : BuildCommand
                     }
                 }
             }
+			foreach(var CodeProj in Branch.CodeProjects)
+			{
+				foreach(var Program in CodeProj.Properties.Programs)
+				{
+					bool bInternalNodeOnly;
+					bool SeparateNode;
+
+					if(Program.Rules.GUBP_AlwaysBuildWithTools(HostPlatform, GUBP.bBuildRocket, out bInternalNodeOnly, out SeparateNode) && Program.Rules.SupportsPlatform(HostPlatform) && SeparateNode)
+					{
+						if(bInternalNodeOnly)
+						{
+							AddNode(new SingleInternalToolsNode(HostPlatform, Program));
+						}
+						else
+						{
+							AddNode(new SingleToolsNode(HostPlatform, Program));
+						}
+					}
+				}
+			}
 
             AddNode(new EditorAndToolsNode(this, HostPlatform));
 
@@ -5482,7 +5502,7 @@ public class GUBP : BuildCommand
             if (HasNode(GamePlatformMonolithicsNode.StaticGetFullName(UnrealTargetPlatform.Mac, Branch.BaseEngineProject, UnrealTargetPlatform.IOS)) && HasNode(ToolsNode.StaticGetFullName(UnrealTargetPlatform.Win64)))
             {
                 //AddNode(new IOSOnPCTestNode(this)); - Disable IOSOnPCTest until a1011 crash is fixed
-            }
+            }			
             if (!bPreflightBuild)
             {
                 AddNode(new CleanSharedTempStorageNode(this));
