UAT: Removing unused variables causing compilation warnings with mono.

diff --git a/Engine/Source/Programs/AutomationTool/Scripts/GUBP.Automation.cs b/Engine/Source/Programs/AutomationTool/Scripts/GUBP.Automation.cs
--- a/Engine/Source/Programs/AutomationTool/Scripts/GUBP.Automation.cs
+++ b/Engine/Source/Programs/AutomationTool/Scripts/GUBP.Automation.cs
@@ -6109,12 +6109,6 @@ public class GUBP : BuildCommand
             }
 
             var FakeECArgs = new List<string>();
-            var BranchForEmail = "";
-            if (P4Enabled)
-            {
-                BranchForEmail = P4Env.BuildRootP4;
-            }
-
             var AgentGroupChains = new Dictionary<string, List<string>>();
             foreach (var NodeToDo in OrdereredToDo)
             {
diff --git a/Engine/Source/Programs/AutomationTool/Scripts/PackageCommand.Automation.cs b/Engine/Source/Programs/AutomationTool/Scripts/PackageCommand.Automation.cs
--- a/Engine/Source/Programs/AutomationTool/Scripts/PackageCommand.Automation.cs
+++ b/Engine/Source/Programs/AutomationTool/Scripts/PackageCommand.Automation.cs
@@ -14,12 +14,6 @@ public partial class Project : CommandUtils
 	public static void Package(ProjectParams Params, int WorkingCL=-1)
 	{
 		Params.ValidateAndLog();
-/*		if (!Params.Package)
-		{
-			return;
-		}*/
-
-		bool NeedsPackage = Params.Package;
 		List<DeploymentContext> DeployContextList = new List<DeploymentContext>();
 		if (!Params.NoClient)
 		{
