Temporarily disabled code signing for Mac editor and games due to problems with signing Steamworks dylib in OS X 10.9.5/Xcode 6

diff --git a/Engine/Source/Programs/AutomationTool/CommandUtils.cs b/Engine/Source/Programs/AutomationTool/CommandUtils.cs
--- a/Engine/Source/Programs/AutomationTool/CommandUtils.cs
+++ b/Engine/Source/Programs/AutomationTool/CommandUtils.cs
@@ -2335,6 +2335,13 @@ namespace AutomationTool
 				throw new AutomationException("Can't sign '{0}', file or folder does not exist.", InPath);
 			}
 
+			// @todo: OS X 10.9.5/Xcode 6 can't sign libsteam_api.dylib, so we temporarily disable signing of the editor and games,
+			// which code sign individual files (contrary to the Launcher, which signs the whole app bundle)
+			if (CommandUtils.FileExists(InPath))
+			{
+				return;
+			}
+
 			// Executable extensions
 			List<string> Extensions = new List<string>();
 			Extensions.Add(".dylib");
