Fix for HTML5 packaging picking up the wrong version of python leading to path failures

diff --git a/Engine/Source/Programs/AutomationTool/HTML5/HTML5Platform.Automation.cs b/Engine/Source/Programs/AutomationTool/HTML5/HTML5Platform.Automation.cs
--- a/Engine/Source/Programs/AutomationTool/HTML5/HTML5Platform.Automation.cs
+++ b/Engine/Source/Programs/AutomationTool/HTML5/HTML5Platform.Automation.cs
@@ -1,6 +1,7 @@
 ﻿// Copyright 1998-2014 Epic Games, Inc. All Rights Reserved.
 using System;
 using System.Collections.Generic;
+using System.Text.RegularExpressions;
 using System.Linq;
 using System.Text;
 using System.IO;
@@ -11,6 +12,33 @@ using System.Threading;
 
 public class HTML5Platform : Platform
 {
+    static string EmscriptenSettingsPath = CombinePaths(Environment.GetFolderPath(Environment.SpecialFolder.UserProfile), "/.emscripten");
+
+    Dictionary<string, string> ReadEmscriptenSettings()
+    {
+        // Check HTML5ToolChain.cs
+        if (!System.IO.File.Exists(EmscriptenSettingsPath))
+        {
+            return new Dictionary<string, string>();
+        }
+
+        Dictionary<string, string> Settings = new Dictionary<string, string>();
+        System.IO.StreamReader SettingFile = new System.IO.StreamReader(EmscriptenSettingsPath);
+        string EMLine = SettingFile.ReadToEnd();
+        string Pattern = @"(\w+)\s*=\s*['\[]((?:[^'\\]|\\.^)*)['\]]";
+        Regex Rgx = new Regex(Pattern, RegexOptions.IgnoreCase);
+        MatchCollection Matches = Rgx.Matches(EMLine);
+        foreach (Match Matched in Matches)
+        {
+            if (Matched.Groups.Count == 3 && Matched.Groups[2].ToString() != "")
+            {
+                Settings[Matched.Groups[1].ToString()] = Matched.Groups[2].ToString();
+            }
+        }
+
+        return Settings;
+    }
+
 	public HTML5Platform()
 		: base(UnrealTargetPlatform.HTML5)
 	{
@@ -19,6 +47,7 @@ public class HTML5Platform : Platform
 	public override void Package(ProjectParams Params, DeploymentContext SC, int WorkingCL)
 	{
 		Log("Package {0}", Params.RawProjectPath);
+        var EmscriptenSettings = ReadEmscriptenSettings();
 
 		string PackagePath = Path.Combine(Path.GetDirectoryName(Params.RawProjectPath), "Binaries", "HTML5");
 		if (!Directory.Exists(PackagePath))
@@ -32,18 +61,47 @@ public class HTML5Platform : Platform
 		using (new PushedDirectory(Path.Combine(Params.BaseStageDirectory, "HTML5")))
 		{
 			string BaseSDKPath = Environment.GetEnvironmentVariable("EMSCRIPTEN");
+            string PythonPath = null;
+
+            // Check the .emscripten file for a possible python path
+            if (EmscriptenSettings.ContainsKey("PYTHON"))
+            {
+                PythonPath = EmscriptenSettings["PYTHON"];
+            }
+            // The AutoSDK defines this env var as part of its install. See setup.bat/unsetup.bat
+            // If it's missing then just assume that python lives on the path
+            if (PythonPath == null)
+            {
+                PythonPath = Environment.GetEnvironmentVariable("PYTHON");
+            }
 			// make the file_packager command line
 			if (Utils.IsRunningOnMono)
 			{
 				string PackagerPath = BaseSDKPath + "/tools/file_packager.py";
-				string CmdLine = string.Format("-c \" python {0} '{1}' --preload . --js-output='{1}.js' \" ", PackagerPath, FinalDataLocation);
-				RunAndLog(CmdEnv, "/bin/bash", CmdLine);
+                if (PythonPath == null)
+                {
+                    string CmdLine = string.Format("-c \" python {0} '{1}' --preload . --js-output='{1}.js' \" ", PackagerPath, FinalDataLocation);
+                    RunAndLog(CmdEnv, "/bin/bash", CmdLine);
+                }
+                else
+                {
+                    string CmdLine = string.Format("{0} '{1}' --preload . --js-output='{1}.js' ", PackagerPath, FinalDataLocation);
+                    RunAndLog(CmdEnv, PythonPath, CmdLine);
+                }
 			}
 			else
 			{
 				string PackagerPath = "\"" + BaseSDKPath + "\\tools\\file_packager.py\"";
-				string CmdLine = string.Format("/c python {0} \"{1}\" --preload . --js-output=\"{1}.js\"", PackagerPath, FinalDataLocation);
-				RunAndLog(CmdEnv, CommandUtils.CombinePaths(Environment.SystemDirectory, "cmd.exe"), CmdLine);
+                if (PythonPath == null)
+                {
+                    string CmdLine = string.Format("/c python {0} \"{1}\" --preload . --js-output=\"{1}.js\"", PackagerPath, FinalDataLocation);
+                    RunAndLog(CmdEnv, CommandUtils.CombinePaths(Environment.SystemDirectory, "cmd.exe"), CmdLine);
+                }
+                else 
+                {
+                    string CmdLine = string.Format("{0} \"{1}\" --preload . --js-output=\"{1}.js\"", PackagerPath, FinalDataLocation);
+                    RunAndLog(CmdEnv, PythonPath, CmdLine);
+                }
 			}
 		}
 
diff --git a/Engine/Source/Programs/UnrealBuildTool/HTML5/HTML5ToolChain.cs b/Engine/Source/Programs/UnrealBuildTool/HTML5/HTML5ToolChain.cs
--- a/Engine/Source/Programs/UnrealBuildTool/HTML5/HTML5ToolChain.cs
+++ b/Engine/Source/Programs/UnrealBuildTool/HTML5/HTML5ToolChain.cs
@@ -5,7 +5,7 @@ using System.Collections.Generic;
 using System.Text.RegularExpressions;
 using System.Diagnostics;
 using System.IO;
-using System.Linq; 
+using System.Linq;
 
 namespace UnrealBuildTool
 {
@@ -14,6 +14,32 @@ namespace UnrealBuildTool
 		// cache the location of SDK tools
 		static string EMCCPath;
 		static string PythonPath;
+        static string EmscriptenSettingsPath = Path.Combine(Environment.GetFolderPath(Environment.SpecialFolder.UserProfile), "/.emscripten");
+
+        Dictionary<string, string> ReadEmscriptenSettings()
+        {
+            // Check HTML5Platform.Automation.cs
+            if (!System.IO.File.Exists(EmscriptenSettingsPath))
+            {
+                return new Dictionary<string, string>();
+            }
+
+            Dictionary<string, string> Settings = new Dictionary<string, string>();
+            System.IO.StreamReader SettingFile = new System.IO.StreamReader(EmscriptenSettingsPath);
+            string EMLine = SettingFile.ReadToEnd();
+            string Pattern = @"(\w+)\s*=\s*['\[]((?:[^'\\]|\\.^)*)['\]]";
+            Regex Rgx = new Regex(Pattern, RegexOptions.IgnoreCase);
+            MatchCollection Matches = Rgx.Matches(EMLine);
+            foreach (Match Matched in Matches)
+            {
+                if (Matched.Groups.Count == 3 && Matched.Groups[2].ToString() != "")
+                {
+                    Settings[Matched.Groups[1].ToString()] = Matched.Groups[2].ToString();
+                }
+            }
+
+            return Settings;
+        }
 
 		public override void RegisterToolChain()
 		{
@@ -25,9 +51,17 @@ namespace UnrealBuildTool
 				BaseSDKPath = BaseSDKPath.Replace("\"", "");
 				if (!String.IsNullOrEmpty(BaseSDKPath))
                 {
+                    var EmscriptenSettings = ReadEmscriptenSettings();
 					EMCCPath = Path.Combine(BaseSDKPath, "emcc");
 					// also figure out where python lives (if no envvar, assume it's in the path)
-					PythonPath = Environment.GetEnvironmentVariable("PYTHON");
+                    if (EmscriptenSettings.ContainsKey("PYTHON"))
+                    {
+                        PythonPath = EmscriptenSettings["PYTHON"];
+                    }
+                    else
+                    {
+                        PythonPath = Environment.GetEnvironmentVariable("PYTHON");
+                    }
 
                     string PythonExeName = Utils.IsRunningOnMono ? "python" : "python.exe";
 
