build script, add option for nonunity test compile for tools

diff --git a/Engine/Source/Programs/AutomationTool/Scripts/GUBP.Automation.cs b/Engine/Source/Programs/AutomationTool/Scripts/GUBP.Automation.cs
--- a/Engine/Source/Programs/AutomationTool/Scripts/GUBP.Automation.cs
+++ b/Engine/Source/Programs/AutomationTool/Scripts/GUBP.Automation.cs
@@ -2931,6 +2931,57 @@ public class GUBP : BuildCommand
         }
     }
 
+	public class NonUnityToolNode : TestNode
+	{
+		SingleTargetProperties ProgramTarget;
+		public NonUnityToolNode(UnrealTargetPlatform InHostPlatform, SingleTargetProperties InProgramTarget)
+			: base(InHostPlatform)
+		{
+			ProgramTarget = InProgramTarget;
+			AddPseudodependency(SingleInternalToolsNode.StaticGetFullName(HostPlatform, ProgramTarget));
+			AddPseudodependency(RootEditorNode.StaticGetFullName(HostPlatform));
+		}
+		public static string StaticGetFullName(UnrealTargetPlatform InHostPlatform, SingleTargetProperties InGameProj)
+		{
+			return InGameProj.TargetName + "_NonUnityTestCompile" + StaticGetHostPlatformSuffix(InHostPlatform);
+		}
+		public override string GetFullName()
+		{
+			return StaticGetFullName(HostPlatform, ProgramTarget);
+		}
+		public override int CISFrequencyQuantumShift(GUBP bp)
+		{
+			int Result = base.CISFrequencyQuantumShift(bp) + 2;
+			if (HostPlatform == UnrealTargetPlatform.Mac)
+			{
+				Result += 1;
+			}		
+			return Result;
+		}
+		public override int AgentMemoryRequirement(GUBP bp)
+		{
+			int Result = base.AgentMemoryRequirement(bp);
+			if (HostPlatform == UnrealTargetPlatform.Mac)
+			{
+				Result = 32;
+			}
+			return Result;
+		}
+		public override void DoTest(GUBP bp)
+		{
+			var Build = new UE4Build(bp);
+			var Agenda = new UE4Build.BuildAgenda();
+
+			Agenda.AddTargets(new string[] { "UnrealHeaderTool" }, HostPlatform, UnrealTargetConfiguration.Development);
+			Agenda.AddTargets(new string[] { ProgramTarget.TargetName }, HostPlatform, UnrealTargetConfiguration.Development, InAddArgs: "-skipnonhostplatforms");			
+
+			Build.Build(Agenda, InDeleteBuildProducts: true, InUpdateVersionFiles: false, InForceNonUnity: true);
+
+			UE4Build.CheckBuildProducts(Build.BuildProductFiles);
+			SaveRecordOfSuccessAndAddToBuildProducts();
+		}
+	}
+
     public class NonUnityTestNode : TestNode
     {
         public NonUnityTestNode(UnrealTargetPlatform InHostPlatform)
@@ -4901,25 +4952,33 @@ public class GUBP : BuildCommand
                         AddNode(new SingleToolsNode(HostPlatform, ProgramTarget));
                     }
                 }
+				if (ProgramTarget.Rules.GUBP_IncludeNonUnityToolTest())
+				{
+					AddNode(new NonUnityToolNode(HostPlatform, ProgramTarget));
+				}
             }
 			foreach(var CodeProj in Branch.CodeProjects)
 			{
-				foreach(var Program in CodeProj.Properties.Programs)
+				foreach(var ProgramTarget in CodeProj.Properties.Programs)
 				{
 					bool bInternalNodeOnly;
 					bool SeparateNode;
 
-					if(Program.Rules.GUBP_AlwaysBuildWithTools(HostPlatform, GUBP.bBuildRocket, out bInternalNodeOnly, out SeparateNode) && Program.Rules.SupportsPlatform(HostPlatform) && SeparateNode)
+					if(ProgramTarget.Rules.GUBP_AlwaysBuildWithTools(HostPlatform, GUBP.bBuildRocket, out bInternalNodeOnly, out SeparateNode) && ProgramTarget.Rules.SupportsPlatform(HostPlatform) && SeparateNode)
 					{
 						if(bInternalNodeOnly)
 						{
-							AddNode(new SingleInternalToolsNode(HostPlatform, Program));
+							AddNode(new SingleInternalToolsNode(HostPlatform, ProgramTarget));
 						}
 						else
 						{
-							AddNode(new SingleToolsNode(HostPlatform, Program));
+							AddNode(new SingleToolsNode(HostPlatform, ProgramTarget));
 						}
 					}
+					if(ProgramTarget.Rules.GUBP_IncludeNonUnityToolTest())
+					{
+						AddNode(new NonUnityToolNode(HostPlatform, ProgramTarget));
+					}
 				}
 			}
 
diff --git a/Engine/Source/Programs/UnrealBuildTool/System/RulesCompiler.cs b/Engine/Source/Programs/UnrealBuildTool/System/RulesCompiler.cs
--- a/Engine/Source/Programs/UnrealBuildTool/System/RulesCompiler.cs
+++ b/Engine/Source/Programs/UnrealBuildTool/System/RulesCompiler.cs
@@ -648,7 +648,14 @@ namespace UnrealBuildTool
         {
             return new List<UnrealTargetConfiguration> { UnrealTargetConfiguration.Development };
         }
-
+		/// <summary>
+		/// Return true if target should include a NonUnity test
+		/// </summary>
+		/// <returns>true if this target should include a NonUnity test
+		public virtual bool GUBP_IncludeNonUnityToolTest()
+		{
+			return false;
+		}
         /// <summary>
         /// Return true if this target should use a platform specific pass
         /// </summary>
