HTML5 RequiresPackageToDeploy

diff --git a/Engine/Source/Programs/AutomationTool/HTML5/HTML5Platform.Automation.cs b/Engine/Source/Programs/AutomationTool/HTML5/HTML5Platform.Automation.cs
--- a/Engine/Source/Programs/AutomationTool/HTML5/HTML5Platform.Automation.cs
+++ b/Engine/Source/Programs/AutomationTool/HTML5/HTML5Platform.Automation.cs
@@ -157,6 +157,13 @@ public class HTML5Platform : Platform
 		PrintRunTime();
 	}
 
+
+    public override bool RequiresPackageToDeploy
+    {
+        get { return true; }
+    }
+
+
 	protected void GenerateFileFromTemplate(string InTemplateFile, string InOutputFile, string InGameName, string InGameConfiguration, string InArguments, bool IsContentOnly, ulong HeapSize)
 	{
 		StringBuilder outputContents = new StringBuilder();
