Updated HTML5 Automation to reflect device name change.
Also updated Run command to use python exe pulled from emscripten config file if it exists.

diff --git a/Engine/Source/Programs/AutomationTool/HTML5/HTML5Platform.Automation.cs b/Engine/Source/Programs/AutomationTool/HTML5/HTML5Platform.Automation.cs
--- a/Engine/Source/Programs/AutomationTool/HTML5/HTML5Platform.Automation.cs
+++ b/Engine/Source/Programs/AutomationTool/HTML5/HTML5Platform.Automation.cs
@@ -284,6 +284,7 @@ public class HTML5Platform : Platform
 
         string browserPath;
         string DeviceName = Params.Device.Split('@')[1];
+        DeviceName = DeviceName.Substring(0, DeviceName.LastIndexOf(" on "));
         bool ok = ConfigCache.GetString(DeviceSection, DeviceName, out browserPath);
 
 		if (!ok)
@@ -322,16 +323,28 @@ public class HTML5Platform : Platform
         }
         else
         {
-
+            var EmscriptenSettings = ReadEmscriptenSettings();
             url = "http://127.0.0.1:8000/" + url;
             // this will be killed UBT instances dies.
             string input = String.Format(" -m SimpleHTTPServer 8000");
 
 
-			string PythonName = "python2.exe";
-
-			if (Utils.IsRunningOnMono)
-				PythonName = "python2";
+			string PythonName = null;
+            // Check the .emscripten file for a possible python path
+            if (EmscriptenSettings.ContainsKey("PYTHON"))
+            {
+                PythonName = EmscriptenSettings["PYTHON"];
+            }
+            // The AutoSDK defines this env var as part of its install. See setup.bat/unsetup.bat
+            // If it's missing then just assume that python lives on the path
+            if (PythonName == null)
+            {
+                PythonName = Environment.GetEnvironmentVariable("PYTHON");
+            }
+            if (PythonName == null)
+            {
+                PythonName = Utils.IsRunningOnMono ? "python" : "python.exe";
+            }
 
 			ProcessResult Result = ProcessManager.CreateProcess(PythonName, true, "html5server.log");
 			Result.ProcessObject.StartInfo.FileName = PythonName;
