CIS fix for Mac.

diff --git a/Engine/Source/Programs/AutomationTool/Win/WinPlatformResources.Automation.cs b/Engine/Source/Programs/AutomationTool/Win/WinPlatformResources.Automation.cs
--- a/Engine/Source/Programs/AutomationTool/Win/WinPlatformResources.Automation.cs
+++ b/Engine/Source/Programs/AutomationTool/Win/WinPlatformResources.Automation.cs
@@ -56,8 +56,8 @@ class GroupIconResource
 		{
 			BinaryReader Reader = new BinaryReader(Stream);
 
-			ushort Reserved = Reader.ReadUInt16();
-			ushort Type = Reader.ReadUInt16();
+			Reader.ReadUInt16(); // Reserved
+			Reader.ReadUInt16(); // Type
 			ushort Count = Reader.ReadUInt16();
 
 			IconResource[] Icons = new IconResource[Count];
