Modular staging

diff --git a/Engine/Source/Programs/AutomationTool/Linux/LinuxPlatform.Automation.cs b/Engine/Source/Programs/AutomationTool/Linux/LinuxPlatform.Automation.cs
--- a/Engine/Source/Programs/AutomationTool/Linux/LinuxPlatform.Automation.cs
+++ b/Engine/Source/Programs/AutomationTool/Linux/LinuxPlatform.Automation.cs
@@ -31,7 +31,7 @@ public abstract class BaseLinuxPlatform : Platform
 		}
 
 		{
-			SC.StageFiles(StagedFileType.NonUFS, CombinePaths(SC.LocalRoot, "Engine/Binaries/ThirdParty/ICU/icu4c-53_1/", SC.PlatformDir, "x86_64-unknown-linux-gnu"), Params.bDebugBuildsActuallyUseDebugCRT ? "*d.so" : "*.so", false, new[] { Params.bDebugBuildsActuallyUseDebugCRT ? "*.so" : "*d.so" }, CombinePaths("Engine/Binaries", SC.PlatformDir));
+			SC.StageFiles(StagedFileType.NonUFS, CombinePaths(SC.LocalRoot, "Engine/Binaries/ThirdParty/ICU/icu4c-53_1/", SC.PlatformDir, "x86_64-unknown-linux-gnu"), Params.bDebugBuildsActuallyUseDebugCRT ? "*d.so*" : "*.so*", false, new[] { Params.bDebugBuildsActuallyUseDebugCRT ? "*.so*" : "*d.so*" }, CombinePaths("Engine/Binaries", SC.PlatformDir));
 		}
 
 		// assume that we always have to deploy Steam (FIXME: should be automatic)
@@ -57,39 +57,65 @@ public abstract class BaseLinuxPlatform : Platform
 		SC.StageFiles(StagedFileType.NonUFS, CombinePaths(SC.ProjectRoot, "Content/Splash"), "Splash.bmp", false, null, null, true);
 		SC.StageFiles(StagedFileType.UFS, CombinePaths(SC.LocalRoot, "Engine/Content/Localization/ICU"), "*", true, null, null, false, !Params.Pak);
 
-		List<string> Exes = GetExecutableNames(SC);
 
-		foreach (var Exe in Exes)
-		{
-
-			if (Exe.StartsWith(CombinePaths(SC.RuntimeProjectRootDir, "Binaries", SC.PlatformDir)))
-			{
-				// remap the project root. For content-only projects, rename the executable to the project name.
-				if (!Params.IsCodeBasedProject && Exe == Exes[0])
-				{
-					SC.StageFiles(StagedFileType.NonUFS, CombinePaths(SC.ProjectRoot, "Binaries", SC.PlatformDir), Path.GetFileNameWithoutExtension(Exe), true, null, CommandUtils.CombinePaths(SC.RelativeProjectRootForStage, "Binaries", SC.PlatformDir), false, true, SC.ShortProjectName);
-				}
-				else
-				{
-					SC.StageFiles(StagedFileType.NonUFS, CombinePaths(SC.ProjectRoot, "Binaries", SC.PlatformDir), Path.GetFileNameWithoutExtension(Exe), true, null, CommandUtils.CombinePaths(SC.RelativeProjectRootForStage, "Binaries", SC.PlatformDir), false);
-				}
-			}
-			else if (Exe.StartsWith(CombinePaths(SC.RuntimeRootDir, "Engine/Binaries", SC.PlatformDir)))
-			{
-				// Move the executable for content-only projects into the project directory, using the project name, so it can figure out the UProject to look for and is consistent with code projects.
-				if (!Params.IsCodeBasedProject && Exe == Exes[0])
-				{
-					SC.StageFiles(StagedFileType.NonUFS, CombinePaths(SC.LocalRoot, "Engine/Binaries", SC.PlatformDir), Path.GetFileNameWithoutExtension(Exe), true, null, CommandUtils.CombinePaths(SC.RelativeProjectRootForStage, "Binaries", SC.PlatformDir), false, true, SC.ShortProjectName);
-				}
-				else
-				{
-					SC.StageFiles(StagedFileType.NonUFS, CombinePaths(SC.LocalRoot, "Engine/Binaries", SC.PlatformDir), Path.GetFileNameWithoutExtension(Exe), true, null, null, false);
-				}
-			}
-			else
-			{
-				throw new AutomationException("Can't stage the exe {0} because it doesn't start with {1} or {2}", Exe, CombinePaths(SC.RuntimeProjectRootDir, "Binaries", SC.PlatformDir), CombinePaths(SC.RuntimeRootDir, "Engine/Binaries", SC.PlatformDir));
-			}
+        if (Params.StageNonMonolithic)
+        {
+            SC.StageFiles(StagedFileType.NonUFS, CombinePaths(SC.LocalRoot, "Engine/Binaries/ThirdParty/jemalloc/", SC.PlatformDir, "x86_64-unknown-linux-gnu"), "libjemalloc.so.1", false, null, CombinePaths("Engine/Binaries", SC.PlatformDir));
+            SC.StageFiles(StagedFileType.NonUFS, CombinePaths(SC.LocalRoot, "Engine/Binaries/ThirdParty/SDL2/", SC.PlatformDir, "x86_64-unknown-linux-gnu"), "libSDL2-2.0.so.1", false, null, CombinePaths("Engine/Binaries", SC.PlatformDir));
+
+            if (SC.DedicatedServer)
+            {
+                SC.StageFiles(StagedFileType.NonUFS, CombinePaths(SC.LocalRoot, "Engine/Binaries", SC.PlatformDir), "UE4Server");
+                SC.StageFiles(StagedFileType.NonUFS, CombinePaths(SC.LocalRoot, "Engine/Binaries", SC.PlatformDir), "libUE4Server-*.so");
+                SC.StageFiles(StagedFileType.NonUFS, CombinePaths(SC.LocalRoot, "Engine/Plugins"), "libUE4Server-*.so", true, null, null, true);
+                SC.StageFiles(StagedFileType.NonUFS, CombinePaths(SC.ProjectRoot, "Binaries", SC.PlatformDir), "libUE4Server-*.so");
+                SC.StageFiles(StagedFileType.NonUFS, CombinePaths(SC.ProjectRoot, "Plugins"), "libUE4Server-*.so", true, null, null, true);
+            }
+            else
+            {
+                SC.StageFiles(StagedFileType.NonUFS, CombinePaths(SC.LocalRoot, "Engine/Binaries", SC.PlatformDir), "UE4");
+                SC.StageFiles(StagedFileType.NonUFS, CombinePaths(SC.LocalRoot, "Engine/Binaries", SC.PlatformDir), "libUE4-*.so");
+                SC.StageFiles(StagedFileType.NonUFS, CombinePaths(SC.LocalRoot, "Engine/Plugins"), "libUE4-*.so", true, null, null, true);
+                SC.StageFiles(StagedFileType.NonUFS, CombinePaths(SC.ProjectRoot, "Binaries", SC.PlatformDir), "libUE4-*.so");
+                SC.StageFiles(StagedFileType.NonUFS, CombinePaths(SC.ProjectRoot, "Plugins"), "libUE4-*.so", true, null, null, true);
+            }
+        }
+        else
+        {
+	    	List<string> Exes = GetExecutableNames(SC);
+
+		    foreach (var Exe in Exes)
+		    {
+
+			    if (Exe.StartsWith(CombinePaths(SC.RuntimeProjectRootDir, "Binaries", SC.PlatformDir)))
+			    {
+				    // remap the project root. For content-only projects, rename the executable to the project name.
+				    if (!Params.IsCodeBasedProject && Exe == Exes[0])
+				    {
+					    SC.StageFiles(StagedFileType.NonUFS, CombinePaths(SC.ProjectRoot, "Binaries", SC.PlatformDir), Path.GetFileNameWithoutExtension(Exe), true, null, CommandUtils.CombinePaths(SC.RelativeProjectRootForStage, "Binaries", SC.PlatformDir), false, true, SC.ShortProjectName);
+				    }
+				    else
+				    {
+					    SC.StageFiles(StagedFileType.NonUFS, CombinePaths(SC.ProjectRoot, "Binaries", SC.PlatformDir), Path.GetFileNameWithoutExtension(Exe), true, null, CommandUtils.CombinePaths(SC.RelativeProjectRootForStage, "Binaries", SC.PlatformDir), false);
+				    }
+			    }
+			    else if (Exe.StartsWith(CombinePaths(SC.RuntimeRootDir, "Engine/Binaries", SC.PlatformDir)))
+			    {
+				    // Move the executable for content-only projects into the project directory, using the project name, so it can figure out the UProject to look for and is consistent with code projects.
+				    if (!Params.IsCodeBasedProject && Exe == Exes[0])
+				    {
+					    SC.StageFiles(StagedFileType.NonUFS, CombinePaths(SC.LocalRoot, "Engine/Binaries", SC.PlatformDir), Path.GetFileNameWithoutExtension(Exe), true, null, CommandUtils.CombinePaths(SC.RelativeProjectRootForStage, "Binaries", SC.PlatformDir), false, true, SC.ShortProjectName);
+				    }
+				    else
+				    {
+					    SC.StageFiles(StagedFileType.NonUFS, CombinePaths(SC.LocalRoot, "Engine/Binaries", SC.PlatformDir), Path.GetFileNameWithoutExtension(Exe), true, null, null, false);
+				    }
+			    }
+			    else
+			    {
+				    throw new AutomationException("Can't stage the exe {0} because it doesn't start with {1} or {2}", Exe, CombinePaths(SC.RuntimeProjectRootDir, "Binaries", SC.PlatformDir), CombinePaths(SC.RuntimeRootDir, "Engine/Binaries", SC.PlatformDir));
+                }
+            }
 		}
 	}
 
diff --git a/Engine/Source/Programs/AutomationTool/ProjectParams.cs b/Engine/Source/Programs/AutomationTool/ProjectParams.cs
--- a/Engine/Source/Programs/AutomationTool/ProjectParams.cs
+++ b/Engine/Source/Programs/AutomationTool/ProjectParams.cs
@@ -226,7 +226,8 @@ namespace AutomationTool
 			this.LogWindow = InParams.LogWindow;
 			this.Stage = InParams.Stage;
 			this.SkipStage = InParams.SkipStage;
-			this.StageDirectoryParam = InParams.StageDirectoryParam;
+            this.StageDirectoryParam = InParams.StageDirectoryParam;
+            this.StageNonMonolithic = InParams.StageNonMonolithic;
 			this.Manifests = InParams.Manifests;
             this.CreateChunkInstall = InParams.CreateChunkInstall;
 			this.UE4Exe = InParams.UE4Exe;
@@ -291,7 +292,8 @@ namespace AutomationTool
 			string RunCommandline = null,						
 			string StageCommandline = null,
             string BundleName = null,
-			string StageDirectoryParam = null,
+            string StageDirectoryParam = null,
+            bool? StageNonMonolithic = null,
 			string UE4Exe = null,
 			string SignPak = null,
 			List<UnrealTargetConfiguration> ClientConfigsToBuild = null,
@@ -445,6 +447,7 @@ namespace AutomationTool
 				this.Stage = true;
 			}
 			this.StageDirectoryParam = ParseParamValueIfNotSpecified(Command, StageDirectoryParam, "stagingdirectory", String.Empty).Trim(new char[]{'\"'});
+            this.StageNonMonolithic = GetParamValueIfNotSpecified(Command, StageNonMonolithic, this.StageNonMonolithic, "StageNonMonolithic");
 			this.Manifests = GetParamValueIfNotSpecified(Command, Manifests, this.Manifests, "manifests");
             this.CreateChunkInstall = GetParamValueIfNotSpecified(Command, CreateChunkInstall, this.CreateChunkInstall, "createchunkinstall");
 			this.Archive = GetParamValueIfNotSpecified(Command, Archive, this.Archive, "archive");
@@ -736,6 +739,11 @@ namespace AutomationTool
 		[Help("stagingdirectory=Path", "Directory to copy the builds to, i.e. -stagingdirectory=C:\\Stage")]
 		public string StageDirectoryParam;
 
+        /// <summary>
+        /// Whether the project should use non monolithic staging
+        /// </summary>
+        public bool StageNonMonolithic;
+
 		[Help("ue4exe=ExecutableName", "Name of the UE4 Editor executable, i.e. -ue4exe=UE4Editor.exe")]
 		public string UE4Exe;
 
diff --git a/Engine/Source/Programs/AutomationTool/Win/WinPlatform.Automation.cs b/Engine/Source/Programs/AutomationTool/Win/WinPlatform.Automation.cs
--- a/Engine/Source/Programs/AutomationTool/Win/WinPlatform.Automation.cs
+++ b/Engine/Source/Programs/AutomationTool/Win/WinPlatform.Automation.cs
@@ -90,42 +90,65 @@ public abstract class BaseWinPlatform : Platform
 
 		SC.StageFiles(StagedFileType.UFS, CombinePaths(SC.LocalRoot, "Engine/Content/Localization/ICU"), "*", true, null, null, false, !Params.Pak);
 
-		List<string> Exes = GetExecutableNames(SC);
+        if (Params.StageNonMonolithic)
+        {
+            if (SC.DedicatedServer)
+            {
+                StageExecutable("exe", SC, CommandUtils.CombinePaths(SC.LocalRoot, "Engine/Binaries", SC.PlatformDir), "UE4Server.");
+                StageExecutable("dll", SC, CommandUtils.CombinePaths(SC.LocalRoot, "Engine/Binaries", SC.PlatformDir), "UE4Server-*.");
+                StageExecutable("dll", SC, CommandUtils.CombinePaths(SC.LocalRoot, "Engine/Plugins"), "UE4Server-*.", true);
+                StageExecutable("dll", SC, CommandUtils.CombinePaths(SC.ProjectRoot, "Binaries", SC.PlatformDir), "UE4Server-*.");
+                StageExecutable("dll", SC, CommandUtils.CombinePaths(SC.ProjectRoot, "Plugins"), "UE4Server-*.", true);
+
+                StageExecutable("dll", SC, CommandUtils.CombinePaths(SC.LocalRoot, "Engine/Binaries/ThirdParty/ICU/icu4c-53_1", SC.PlatformDir, "VS2013"), "*.");
+            }
+            else
+            {
+                StageExecutable("exe", SC, CommandUtils.CombinePaths(SC.LocalRoot, "Engine/Binaries", SC.PlatformDir), "UE4.");
+                StageExecutable("dll", SC, CommandUtils.CombinePaths(SC.LocalRoot, "Engine/Binaries", SC.PlatformDir), "UE4-*.");
+                StageExecutable("dll", SC, CommandUtils.CombinePaths(SC.LocalRoot, "Engine/Plugins"), "UE4-*.", true);
+                StageExecutable("dll", SC, CommandUtils.CombinePaths(SC.ProjectRoot, "Binaries", SC.PlatformDir), "UE4-*.");
+                StageExecutable("dll", SC, CommandUtils.CombinePaths(SC.ProjectRoot, "Plugins"), "UE4-*.", true);
+            }
+        }
+        else
+        {
+			List<string> Exes = GetExecutableNames(SC);
 
-		// the first exe is the "main" one, the rest are marked as debug files
-		StagedFileType WorkingFileType = StagedFileType.NonUFS;
-
-		foreach (var Exe in Exes)
-		{
-
-			if (Exe.StartsWith(CombinePaths(SC.RuntimeProjectRootDir, "Binaries", SC.PlatformDir)))
-			{
-				// remap the project root.
-				string SourceFile = CombinePaths(SC.ProjectRoot, "Binaries", SC.PlatformDir, Path.GetFileName(Exe));
-				StageExecutable("exe", SC, Path.GetDirectoryName(SourceFile), Path.GetFileNameWithoutExtension(SourceFile) + ".", true, null, CommandUtils.CombinePaths(SC.RelativeProjectRootForStage, "Binaries", SC.PlatformDir), false, WorkingFileType);
-				if(Exe == Exes[0])
+			// the first exe is the "main" one, the rest are marked as debug files
+			StagedFileType WorkingFileType = StagedFileType.NonUFS;
+				
+		    foreach (var Exe in Exes)
+            {
+				if (Exe.StartsWith(CombinePaths(SC.RuntimeProjectRootDir, "Binaries", SC.PlatformDir)))
 				{
-					StageBootstrapExecutable(SC, SourceFile, CombinePaths(SC.RelativeProjectRootForStage, "Binaries", SC.PlatformDir, Path.GetFileName(Exe)), "");
+					// remap the project root.
+					string SourceFile = CombinePaths(SC.ProjectRoot, "Binaries", SC.PlatformDir, Path.GetFileName(Exe));
+					StageExecutable("exe", SC, Path.GetDirectoryName(SourceFile), Path.GetFileNameWithoutExtension(SourceFile) + ".", true, null, CommandUtils.CombinePaths(SC.RelativeProjectRootForStage, "Binaries", SC.PlatformDir), false, WorkingFileType);
+					if(Exe == Exes[0])
+					{
+						StageBootstrapExecutable(SC, SourceFile, CombinePaths(SC.RelativeProjectRootForStage, "Binaries", SC.PlatformDir, Path.GetFileName(Exe)), "");
+					}
 				}
-			}
-			else if (Exe.StartsWith(CombinePaths(SC.RuntimeRootDir, "Engine/Binaries", SC.PlatformDir)))
-			{
-				// keep it in the engine directory.
-				string SourceFile = CombinePaths(SC.LocalRoot, "Engine", "Binaries", SC.PlatformDir, Path.GetFileName(Exe));
-				StageExecutable("exe", SC, CombinePaths(SC.LocalRoot, "Engine/Binaries", SC.PlatformDir), Path.GetFileNameWithoutExtension(SourceFile) + ".", true, null, null, false, WorkingFileType);
-				if(Exe == Exes[0])
+				else if (Exe.StartsWith(CombinePaths(SC.RuntimeRootDir, "Engine/Binaries", SC.PlatformDir)))
 				{
-					StageBootstrapExecutable(SC, SourceFile, CombinePaths("Engine", "Binaries", SC.PlatformDir, Path.GetFileName(Exe)), String.Format("..\\..\\..\\{0}\\{0}.uproject", SC.ShortProjectName));
+					// keep it in the engine directory.
+					string SourceFile = CombinePaths(SC.LocalRoot, "Engine", "Binaries", SC.PlatformDir, Path.GetFileName(Exe));
+					StageExecutable("exe", SC, CombinePaths(SC.LocalRoot, "Engine/Binaries", SC.PlatformDir), Path.GetFileNameWithoutExtension(SourceFile) + ".", true, null, null, false, WorkingFileType);
+					if(Exe == Exes[0])
+					{
+						StageBootstrapExecutable(SC, SourceFile, CombinePaths("Engine", "Binaries", SC.PlatformDir, Path.GetFileName(Exe)), String.Format("..\\..\\..\\{0}\\{0}.uproject", SC.ShortProjectName));
+					}
 				}
+				else
+				{
+					throw new AutomationException("Can't stage the exe {0} because it doesn't start with {1} or {2}", Exe, CombinePaths(SC.RuntimeProjectRootDir, "Binaries", SC.PlatformDir), CombinePaths(SC.RuntimeRootDir, "Engine/Binaries", SC.PlatformDir));
+				}
+				// the first exe is the "main" one, the rest are marked as debug files
+				WorkingFileType = StagedFileType.DebugNonUFS;
 			}
-			else
-			{
-				throw new AutomationException("Can't stage the exe {0} because it doesn't start with {1} or {2}", Exe, CombinePaths(SC.RuntimeProjectRootDir, "Binaries", SC.PlatformDir), CombinePaths(SC.RuntimeRootDir, "Engine/Binaries", SC.PlatformDir));
-			}
-			// the first exe is the "main" one, the rest are marked as debug files
-			WorkingFileType = StagedFileType.DebugNonUFS;
 		}
-
+		
 		if(Params.Prereqs)
 		{
 			string InstallerRelativePath = CombinePaths("Engine", "Extras", "Redist", "en-us");
