Pak file creation - add the ability override UAT command line arguments per platform with a 3 state setting - Always, Never, Don't care. Fixes an issue where pak files were created even when we don't want (for now, atleast).

diff --git a/Engine/Source/Programs/AutomationTool/Android/AndroidPlatform.Automation.cs b/Engine/Source/Programs/AutomationTool/Android/AndroidPlatform.Automation.cs
--- a/Engine/Source/Programs/AutomationTool/Android/AndroidPlatform.Automation.cs
+++ b/Engine/Source/Programs/AutomationTool/Android/AndroidPlatform.Automation.cs
@@ -813,11 +813,11 @@ public class AndroidPlatform : Platform
 		return Dest;
 	}
 
-	public override bool RequiresPak(ProjectParams Params)
+	public override PakType RequiresPak(ProjectParams Params)
 	{
-		return Params.Package;
+		return PakType.Always;
 	}
-
+    
 	#region Hooks
 
 	public override void PostBuildTarget(UE4Build Build, string ProjectName, string UProjectPath, string Config)
diff --git a/Engine/Source/Programs/AutomationTool/HTML5/HTML5Platform.Automation.cs b/Engine/Source/Programs/AutomationTool/HTML5/HTML5Platform.Automation.cs
--- a/Engine/Source/Programs/AutomationTool/HTML5/HTML5Platform.Automation.cs
+++ b/Engine/Source/Programs/AutomationTool/HTML5/HTML5Platform.Automation.cs
@@ -410,6 +410,11 @@ public class HTML5Platform : Platform
 		return false;
 	}
 
+    public override PakType RequiresPak(ProjectParams Params)
+    {
+        return PakType.Never;
+    }
+
 	public override bool DeployLowerCaseFilenames(bool bUFSFile)
 	{
 		return false;
diff --git a/Engine/Source/Programs/AutomationTool/IOS/IOSPlatform.Automation.cs b/Engine/Source/Programs/AutomationTool/IOS/IOSPlatform.Automation.cs
--- a/Engine/Source/Programs/AutomationTool/IOS/IOSPlatform.Automation.cs
+++ b/Engine/Source/Programs/AutomationTool/IOS/IOSPlatform.Automation.cs
@@ -833,6 +833,11 @@ public class IOSPlatform : Platform
 		return false;
 	}
 
+    public override PakType RequiresPak(ProjectParams Params)
+    {
+        return PakType.Always;
+    }
+
 	public override bool DeployLowerCaseFilenames(bool bUFSFile)
 	{
 		// we shouldn't modify the case on files like Info.plist or the icons
diff --git a/Engine/Source/Programs/AutomationTool/Platform.cs b/Engine/Source/Programs/AutomationTool/Platform.cs
--- a/Engine/Source/Programs/AutomationTool/Platform.cs
+++ b/Engine/Source/Programs/AutomationTool/Platform.cs
@@ -376,11 +376,14 @@ namespace AutomationTool
 		}
 
 		/// <summary>
-		/// determines if the platform needs to be pak'd when packaged
+		/// Tri-state - The intent is to override command line parameters for pak if needed per platform.
 		/// </summary>
-		public virtual bool RequiresPak(ProjectParams Params)
+        /// 
+        public enum PakType { Always, Never, DontCare };
+
+        public virtual PakType RequiresPak(ProjectParams Params)
 		{
-			return false;
+            return PakType.DontCare;
 		}
 
         /// <summary>
diff --git a/Engine/Source/Programs/AutomationTool/Scripts/CopyBuildToStagingDirectory.Automation.cs b/Engine/Source/Programs/AutomationTool/Scripts/CopyBuildToStagingDirectory.Automation.cs
--- a/Engine/Source/Programs/AutomationTool/Scripts/CopyBuildToStagingDirectory.Automation.cs
+++ b/Engine/Source/Programs/AutomationTool/Scripts/CopyBuildToStagingDirectory.Automation.cs
@@ -390,7 +390,7 @@ public partial class Project : CommandUtils
 			CopyManifestFilesToStageDir(SC.NonUFSStagingFilesDebug, SC.StageDirectory, "DebugFiles");
 		}
 
-		bool bStageUnrealFileSystemFiles = !Params.CookOnTheFly && !Params.Pak && !Params.FileServer;
+        bool bStageUnrealFileSystemFiles = !Params.CookOnTheFly && !ShouldCreatePak(Params) && !Params.FileServer;
 		if (bStageUnrealFileSystemFiles)
 		{
 			CopyManifestFilesToStageDir(SC.UFSStagingFiles, SC.StageDirectory, "UFSFiles");
@@ -646,12 +646,46 @@ public partial class Project : CommandUtils
 
 	private static bool ShouldCreatePak(ProjectParams Params, DeploymentContext SC)
 	{
-		return (Params.Pak || Params.SignedPak || !String.IsNullOrEmpty(Params.SignPak) || SC.StageTargetPlatform.RequiresPak(Params)) && !Params.SkipPak;
+        Platform.PakType Pak = SC.StageTargetPlatform.RequiresPak(Params);
+
+        // we may care but we don't want. 
+        if (Params.SkipPak)
+            return false; 
+
+        if (Pak == Platform.PakType.Always)
+        {
+            return true;
+        }
+        else if (Pak == Platform.PakType.Never)
+        {
+            return false;
+        }
+        else // DontCare
+        {
+            return (Params.Pak || Params.SignedPak || !String.IsNullOrEmpty(Params.SignPak));
+        }
 	}
 
 	private static bool ShouldCreatePak(ProjectParams Params)
 	{
-		return (Params.Pak || Params.SignedPak || !String.IsNullOrEmpty(Params.SignPak) || Params.ClientTargetPlatformInstances[0].RequiresPak(Params)) && !Params.SkipPak;
+        Platform.PakType Pak = Params.ClientTargetPlatformInstances[0].RequiresPak(Params);
+
+        // we may care but we don't want. 
+        if (Params.SkipPak)
+            return false; 
+
+        if (Pak == Platform.PakType.Always)
+        {
+            return true;
+        }
+        else if (Pak == Platform.PakType.Never)
+        {
+            return false;
+        }
+        else // DontCare
+        {
+            return (Params.Pak || Params.SignedPak || !String.IsNullOrEmpty(Params.SignPak));
+        }
 	}
 
 	public static void ApplyStagingManifest(ProjectParams Params, DeploymentContext SC)
