Changed Emscripten setting read function to handle comments.
Also added more logging, changed autoSDK setup.bat to output helpful varibles into emscripten file.

diff --git a/Engine/Source/Programs/AutomationTool/HTML5/HTML5Platform.Automation.cs b/Engine/Source/Programs/AutomationTool/HTML5/HTML5Platform.Automation.cs
--- a/Engine/Source/Programs/AutomationTool/HTML5/HTML5Platform.Automation.cs
+++ b/Engine/Source/Programs/AutomationTool/HTML5/HTML5Platform.Automation.cs
@@ -16,7 +16,7 @@ public class HTML5Platform : Platform
 
     Dictionary<string, string> ReadEmscriptenSettings()
     {
-        // Check HTML5ToolChain.cs
+        // Check HTML5ToolChain.cs for duplicate
         if (!System.IO.File.Exists(EmscriptenSettingsPath))
         {
             return new Dictionary<string, string>();
@@ -24,15 +24,19 @@ public class HTML5Platform : Platform
 
         Dictionary<string, string> Settings = new Dictionary<string, string>();
         System.IO.StreamReader SettingFile = new System.IO.StreamReader(EmscriptenSettingsPath);
-        string EMLine = SettingFile.ReadToEnd();
-        string Pattern = @"(\w+)\s*=\s*['\[]((?:[^'\\]|\\.^)*)['\]]";
-        Regex Rgx = new Regex(Pattern, RegexOptions.IgnoreCase);
-        MatchCollection Matches = Rgx.Matches(EMLine);
-        foreach (Match Matched in Matches)
+        string EMLine = null;
+        while ((EMLine = SettingFile.ReadLine()) != null)
         {
-            if (Matched.Groups.Count == 3 && Matched.Groups[2].ToString() != "")
+            EMLine = EMLine.Split('#')[0];
+            string Pattern1 = @"(\w*)\s*=\s*['\[]?([^'\]\r\n]*)['\]]?";
+            Regex Rgx = new Regex(Pattern1, RegexOptions.IgnoreCase);
+            MatchCollection Matches = Rgx.Matches(EMLine);
+            foreach (Match Matched in Matches)
             {
-                Settings[Matched.Groups[1].ToString()] = Matched.Groups[2].ToString();
+                if (Matched.Groups.Count == 3 && Matched.Groups[2].ToString() != "")
+                {
+                    Settings[Matched.Groups[1].ToString()] = Matched.Groups[2].ToString();
+                }
             }
         }
 
@@ -67,13 +71,27 @@ public class HTML5Platform : Platform
             if (EmscriptenSettings.ContainsKey("PYTHON"))
             {
                 PythonPath = EmscriptenSettings["PYTHON"];
+                Log("Found python path {0} in emscripten file", PythonPath);
             }
             // The AutoSDK defines this env var as part of its install. See setup.bat/unsetup.bat
             // If it's missing then just assume that python lives on the path
-            if (PythonPath == null)
+            if (PythonPath == null && Environment.GetEnvironmentVariable("PYTHON") != null)
             {
                 PythonPath = Environment.GetEnvironmentVariable("PYTHON");
+                Log("Found python path {0} in PYTHON Environment Variable", PythonPath);
+            }
+
+            // Check if the exe exists
+            if (!System.IO.File.Exists(PythonPath))
+            {
+                PythonPath = null;
             }
+
+            if (PythonPath == null)
+            {
+                Log("Either no python path can be found or it doesn't exist. Using python on PATH");
+            }
+
 			// make the file_packager command line
 			if (Utils.IsRunningOnMono)
 			{
@@ -334,15 +352,20 @@ public class HTML5Platform : Platform
             if (EmscriptenSettings.ContainsKey("PYTHON"))
             {
                 PythonName = EmscriptenSettings["PYTHON"];
+                Log("Found python path {0} in emscripten file", PythonName);
             }
             // The AutoSDK defines this env var as part of its install. See setup.bat/unsetup.bat
             // If it's missing then just assume that python lives on the path
-            if (PythonName == null)
+            if (PythonName == null && Environment.GetEnvironmentVariable("PYTHON") != null)
             {
                 PythonName = Environment.GetEnvironmentVariable("PYTHON");
+                Log("Found python path {0} in PYTHON Environment Variable", PythonName);
             }
-            if (PythonName == null)
+
+            // Check if the exe exists
+            if (!System.IO.File.Exists(PythonName))
             {
+				Log("Either no python path can be found or it doesn't exist. Using python on PATH");
                 PythonName = Utils.IsRunningOnMono ? "python" : "python.exe";
             }
 
diff --git a/Engine/Source/Programs/UnrealBuildTool/HTML5/HTML5ToolChain.cs b/Engine/Source/Programs/UnrealBuildTool/HTML5/HTML5ToolChain.cs
--- a/Engine/Source/Programs/UnrealBuildTool/HTML5/HTML5ToolChain.cs
+++ b/Engine/Source/Programs/UnrealBuildTool/HTML5/HTML5ToolChain.cs
@@ -18,7 +18,7 @@ namespace UnrealBuildTool
 
         Dictionary<string, string> ReadEmscriptenSettings()
         {
-            // Check HTML5Platform.Automation.cs
+            // Check HTML5ToolChain.cs for duplicate
             if (!System.IO.File.Exists(EmscriptenSettingsPath))
             {
                 return new Dictionary<string, string>();
@@ -26,15 +26,19 @@ namespace UnrealBuildTool
 
             Dictionary<string, string> Settings = new Dictionary<string, string>();
             System.IO.StreamReader SettingFile = new System.IO.StreamReader(EmscriptenSettingsPath);
-            string EMLine = SettingFile.ReadToEnd();
-            string Pattern = @"(\w+)\s*=\s*['\[]((?:[^'\\]|\\.^)*)['\]]";
-            Regex Rgx = new Regex(Pattern, RegexOptions.IgnoreCase);
-            MatchCollection Matches = Rgx.Matches(EMLine);
-            foreach (Match Matched in Matches)
+            string EMLine = null;
+            while ((EMLine = SettingFile.ReadLine()) != null)
             {
-                if (Matched.Groups.Count == 3 && Matched.Groups[2].ToString() != "")
+                EMLine = EMLine.Split('#')[0];
+                string Pattern1 = @"(\w*)\s*=\s*['\[]?([^'\]\r\n]*)['\]]?";
+                Regex Rgx = new Regex(Pattern1, RegexOptions.IgnoreCase);
+                MatchCollection Matches = Rgx.Matches(EMLine);
+                foreach (Match Matched in Matches)
                 {
-                    Settings[Matched.Groups[1].ToString()] = Matched.Groups[2].ToString();
+                    if (Matched.Groups.Count == 3 && Matched.Groups[2].ToString() != "")
+                    {
+                        Settings[Matched.Groups[1].ToString()] = Matched.Groups[2].ToString();
+                    }
                 }
             }
 
