Quick fix for Platform Run Client for cook on the fly for HTML5 - some higher level changes passed in unexpected data. This needs a bit more work.

diff --git a/Engine/Source/Programs/AutomationTool/HTML5/HTML5Platform.Automation.cs b/Engine/Source/Programs/AutomationTool/HTML5/HTML5Platform.Automation.cs
--- a/Engine/Source/Programs/AutomationTool/HTML5/HTML5Platform.Automation.cs
+++ b/Engine/Source/Programs/AutomationTool/HTML5/HTML5Platform.Automation.cs
@@ -314,26 +314,16 @@ public class HTML5Platform : Platform
 		// Are we running via cook on the fly server?
 		// find our http url - This is awkward because RunClient doesn't have real information that NFS is running or not.
 		bool IsCookOnTheFly = false;
-		string[] Arguments = ClientCmdLine.Split(' ');
-		foreach (var Argument in Arguments)
-		{
-			string[] KeyVal = Argument.Split('=');
-			if (KeyVal[0].ToLower().Contains("filehostip"))
-			{
-				// split urls.
-				string[] Urls = KeyVal[1].Split('+');
-				foreach (var Url in Urls)
-				{
-					if (Url.Contains("http"))
-					{
-						url = Url + "/" + Path.GetFileName(url);
-						IsCookOnTheFly = true;
-					}
-				}
-			}
-			if (IsCookOnTheFly)
-				break;
-		}
+
+        // 9/24/2014 @fixme - All this is convoluted, clean up.
+        // looks like cookonthefly commandline stopped adding protocol or the port :/ hard coding to DEFAULT_TCP_FILE_SERVING_PORT+1 (DEFAULT_HTTP_FILE_SERVING_PORT)
+        // This will fail if the NFS server is started with a different port - we need to modify driver .cs script to pass in IP/Port data correctly. 
+
+        if (ClientCmdLine.Contains("filehostip"))
+        {
+            IsCookOnTheFly = true; 
+            url = "http://127.0.0.1:41898/" + url; 
+        }
 
         if (IsCookOnTheFly)
         {
