build script, overrides for package and cook, cook android wex with etc1

diff --git a/Engine/Source/Programs/AutomationTool/Scripts/GUBP.Automation.cs b/Engine/Source/Programs/AutomationTool/Scripts/GUBP.Automation.cs
--- a/Engine/Source/Programs/AutomationTool/Scripts/GUBP.Automation.cs
+++ b/Engine/Source/Programs/AutomationTool/Scripts/GUBP.Automation.cs
@@ -2270,7 +2270,29 @@ public class GUBP : BuildCommand
             {
                 // not sure if we need something here or if the cook commandlet will automatically convert the exe name
             }
-            CommandUtils.CookCommandlet(GameProj.FilePath, "UE4Editor-Cmd.exe", null, null, null, CookPlatform);
+			string AdditionalParams = "";
+			var Target = GameProj.Properties.Targets[TargetRules.TargetType.Game];
+			if(CookPlatform.Contains("Server"))
+			{
+				Target = GameProj.Properties.Targets[TargetRules.TargetType.Server];
+			}
+			if(CookPlatform.Contains("Client"))
+			{
+				Target = GameProj.Properties.Targets[TargetRules.TargetType.Client];
+			}
+			
+			if(Target.Rules.GUBP_AdditionalCookParameters(HostPlatform, CookPlatform) != "")
+			{
+				AdditionalParams = Target.Rules.GUBP_AdditionalCookParameters(HostPlatform, CookPlatform);
+			}
+			if (AdditionalParams == "")
+			{
+				CommandUtils.CookCommandlet(GameProj.FilePath, "UE4Editor-Cmd.exe", null, null, null, CookPlatform);
+			}
+			else
+			{
+				CommandUtils.CookCommandlet(GameProj.FilePath, "UE4Editor-Cmd.exe", null, null, null, CookPlatform, AdditionalParams);
+			}
             var CookedPath = RootIfAnyForTempStorage();
             var CookedFiles = CommandUtils.FindFiles("*", true, CookedPath);
             if (CookedFiles.GetLength(0) < 1)
@@ -2711,6 +2733,15 @@ public class GUBP : BuildCommand
                     {
                         Args += String.Format("+{0}", Plat.ToString());
                     }
+					var Target =  GameProj.Properties.Targets[TargetRules.TargetType.Game];
+					if(ClientNotGame)
+					{
+						Target = GameProj.Properties.Targets[TargetRules.TargetType.Client];						
+					}					
+					if (Target.Rules.GUBP_AdditionalPackageParameters(HostPlatform, Plat) != "")
+					{
+						Args += " " + Target.Rules.GUBP_AdditionalPackageParameters(HostPlatform, Plat);
+					}
                 }
                 bool bFirstClientConfig = true;
                 foreach (var Config in ClientConfigs)
@@ -2727,7 +2758,8 @@ public class GUBP : BuildCommand
                 }
                 if (ClientNotGame)
                 {
-                    Args += " -client";
+                    Args += " -client";					
+					
                 }
             }
             else
@@ -2753,6 +2785,11 @@ public class GUBP : BuildCommand
                     {
                         Args += String.Format("+{0}", Plat.ToString());
                     }
+					var Target = GameProj.Properties.Targets[TargetRules.TargetType.Server];
+					if (Target.Rules.GUBP_AdditionalPackageParameters(HostPlatform, Plat) != "")
+					{
+						Args += " " + Target.Rules.GUBP_AdditionalPackageParameters(HostPlatform, Plat);
+					}
                 }
                 bool bFirstServerConfig = true;
                 foreach (var Config in ServerConfigs)
@@ -2796,7 +2833,7 @@ public class GUBP : BuildCommand
                 }
 				Args += String.Format(" -Archive -archivedirectory={0}", CommandUtils.MakePathSafeToUseWithCommandLine(IntermediateArchiveDirectory));
             }
-
+			
             string LogFile = CommandUtils.RunUAT(CommandUtils.CmdEnv, Args);
 
             if (P4Enabled)
diff --git a/Engine/Source/Programs/UnrealBuildTool/System/RulesCompiler.cs b/Engine/Source/Programs/UnrealBuildTool/System/RulesCompiler.cs
--- a/Engine/Source/Programs/UnrealBuildTool/System/RulesCompiler.cs
+++ b/Engine/Source/Programs/UnrealBuildTool/System/RulesCompiler.cs
@@ -852,6 +852,21 @@ namespace UnrealBuildTool
 #endif
             return Result;
         }
+		/// <summary>
+		/// Return additional parameters to cook commandlet
+		/// </summary>
+		public virtual string GUBP_AdditionalCookParameters(UnrealTargetPlatform HostPlatform, string Platform)
+		{
+			return "";
+		}
+
+		/// <summary>
+		/// Return additional parameters to package commandlet
+		/// </summary>
+		public virtual string GUBP_AdditionalPackageParameters(UnrealTargetPlatform HostPlatform, UnrealTargetPlatform Platform)
+		{
+			return "";
+		}
 	}
 
 
