build script, fix for noncode cook/package nodes

diff --git a/Engine/Source/Programs/AutomationTool/Scripts/GUBP.Automation.cs b/Engine/Source/Programs/AutomationTool/Scripts/GUBP.Automation.cs
--- a/Engine/Source/Programs/AutomationTool/Scripts/GUBP.Automation.cs
+++ b/Engine/Source/Programs/AutomationTool/Scripts/GUBP.Automation.cs
@@ -2271,20 +2271,24 @@ public class GUBP : BuildCommand
                 // not sure if we need something here or if the cook commandlet will automatically convert the exe name
             }
 			string AdditionalParams = "";
-			var Target = GameProj.Properties.Targets[TargetRules.TargetType.Game];
-			if(CookPlatform.Contains("Server"))
+			if(GameProj.GameName != bp.Branch.BaseEngineProject.GameName)
 			{
-				Target = GameProj.Properties.Targets[TargetRules.TargetType.Server];
-			}
-			if(CookPlatform.Contains("Client"))
-			{
-				Target = GameProj.Properties.Targets[TargetRules.TargetType.Client];
-			}
-			
-			if(Target.Rules.GUBP_AdditionalCookParameters(HostPlatform, CookPlatform) != "")
-			{
-				AdditionalParams = Target.Rules.GUBP_AdditionalCookParameters(HostPlatform, CookPlatform);
+				var Target = GameProj.Properties.Targets[TargetRules.TargetType.Game];
+				if (CookPlatform.Contains("Server"))
+				{
+					Target = GameProj.Properties.Targets[TargetRules.TargetType.Server];
+				}
+				if (CookPlatform.Contains("Client"))
+				{
+					Target = GameProj.Properties.Targets[TargetRules.TargetType.Client];
+				}
+
+				if (Target.Rules.GUBP_AdditionalCookParameters(HostPlatform, CookPlatform) != "")
+				{
+					AdditionalParams = Target.Rules.GUBP_AdditionalCookParameters(HostPlatform, CookPlatform);
+				}
 			}
+		
 			if (AdditionalParams == "")
 			{
 				CommandUtils.CookCommandlet(GameProj.FilePath, "UE4Editor-Cmd.exe", null, null, null, CookPlatform);
@@ -2733,14 +2737,17 @@ public class GUBP : BuildCommand
                     {
                         Args += String.Format("+{0}", Plat.ToString());
                     }
-					var Target =  GameProj.Properties.Targets[TargetRules.TargetType.Game];
-					if(ClientNotGame)
+					if(GameProj.GameName != bp.Branch.BaseEngineProject.GameName)
 					{
-						Target = GameProj.Properties.Targets[TargetRules.TargetType.Client];						
-					}					
-					if (Target.Rules.GUBP_AdditionalPackageParameters(HostPlatform, Plat) != "")
-					{
-						Args += " " + Target.Rules.GUBP_AdditionalPackageParameters(HostPlatform, Plat);
+						var Target =  GameProj.Properties.Targets[TargetRules.TargetType.Game];
+						if(ClientNotGame)
+						{
+							Target = GameProj.Properties.Targets[TargetRules.TargetType.Client];						
+						}					
+						if (Target.Rules.GUBP_AdditionalPackageParameters(HostPlatform, Plat) != "")
+						{
+							Args += " " + Target.Rules.GUBP_AdditionalPackageParameters(HostPlatform, Plat);
+						}
 					}
                 }
                 bool bFirstClientConfig = true;
@@ -2759,7 +2766,6 @@ public class GUBP : BuildCommand
                 if (ClientNotGame)
                 {
                     Args += " -client";					
-					
                 }
             }
             else
@@ -2785,10 +2791,13 @@ public class GUBP : BuildCommand
                     {
                         Args += String.Format("+{0}", Plat.ToString());
                     }
-					var Target = GameProj.Properties.Targets[TargetRules.TargetType.Server];
-					if (Target.Rules.GUBP_AdditionalPackageParameters(HostPlatform, Plat) != "")
+					if (GameProj.GameName != bp.Branch.BaseEngineProject.GameName)
 					{
-						Args += " " + Target.Rules.GUBP_AdditionalPackageParameters(HostPlatform, Plat);
+						var Target = GameProj.Properties.Targets[TargetRules.TargetType.Server];
+						if (Target.Rules.GUBP_AdditionalPackageParameters(HostPlatform, Plat) != "")
+						{
+							Args += " " + Target.Rules.GUBP_AdditionalPackageParameters(HostPlatform, Plat);
+						}
 					}
                 }
                 bool bFirstServerConfig = true;
