build script, fix for noncode project cooks

diff --git a/Engine/Source/Programs/AutomationTool/Scripts/GUBP.Automation.cs b/Engine/Source/Programs/AutomationTool/Scripts/GUBP.Automation.cs
--- a/Engine/Source/Programs/AutomationTool/Scripts/GUBP.Automation.cs
+++ b/Engine/Source/Programs/AutomationTool/Scripts/GUBP.Automation.cs
@@ -2166,6 +2166,7 @@ public class GUBP : BuildCommand
         UnrealTargetPlatform TargetPlatform;
         string CookPlatform;
         bool bIsMassive;
+		bool bisCode;
 
         public CookNode(GUBP bp, UnrealTargetPlatform InHostPlatform, BranchInfo.BranchUProject InGameProj, UnrealTargetPlatform InTargetPlatform, string InCookPlatform)
             : base(InHostPlatform)
@@ -2189,7 +2190,7 @@ public class GUBP : BuildCommand
             {
                 var Options = InGameProj.Options(HostPlatform);
                 bIsMassive = Options.bIsMassive;
-
+				bisCode = true;
                 AddDependency(EditorGameNode.StaticGetFullName(HostPlatform, GameProj));
                 // add an arc to prevent cooks from running until promotable is labeled
                 if (Options.bIsPromotable)
@@ -2208,6 +2209,7 @@ public class GUBP : BuildCommand
             else
             {
                 bIsShared = true;
+				bisCode = false;
                 AddPseudodependency(GamePlatformMonolithicsNode.StaticGetFullName(HostPlatform, bp.Branch.BaseEngineProject, TargetPlatform));
             }
             if (bIsShared)
@@ -2271,7 +2273,7 @@ public class GUBP : BuildCommand
                 // not sure if we need something here or if the cook commandlet will automatically convert the exe name
             }
 			string AdditionalParams = "";
-			if(GameProj.GameName != bp.Branch.BaseEngineProject.GameName)
+			if (bisCode)
 			{
 				var Target = GameProj.Properties.Targets[TargetRules.TargetType.Game];
 				if (CookPlatform.Contains("Server"))
@@ -2514,6 +2516,7 @@ public class GUBP : BuildCommand
         List<UnrealTargetConfiguration> ClientConfigs;
         List<UnrealTargetConfiguration> ServerConfigs;
         bool ClientNotGame;
+		bool bIsCode;
         UnrealBuildTool.TargetRules.TargetType GameOrClient;
 
         public FormalBuildNode(GUBP bp,
@@ -2540,7 +2543,14 @@ public class GUBP : BuildCommand
             {
                 GameOrClient = TargetRules.TargetType.Client;
             }
-            
+			if (InGameProj.GameName != bp.Branch.BaseEngineProject.GameName && GameProj.Properties.Targets.ContainsKey(TargetRules.TargetType.Editor))
+			{
+				bIsCode = true;
+			}
+			else
+			{
+				bIsCode = false;
+			}
 
             // verify we actually built these
             var WorkingGameProject = InGameProj;
@@ -2737,7 +2747,7 @@ public class GUBP : BuildCommand
                     {
                         Args += String.Format("+{0}", Plat.ToString());
                     }
-					if(GameProj.GameName != bp.Branch.BaseEngineProject.GameName)
+					if(bIsCode)
 					{
 						var Target =  GameProj.Properties.Targets[TargetRules.TargetType.Game];
 						if(ClientNotGame)
@@ -2791,7 +2801,7 @@ public class GUBP : BuildCommand
                     {
                         Args += String.Format("+{0}", Plat.ToString());
                     }
-					if (GameProj.GameName != bp.Branch.BaseEngineProject.GameName)
+					if (bIsCode)
 					{
 						var Target = GameProj.Properties.Targets[TargetRules.TargetType.Server];
 						if (Target.Rules.GUBP_AdditionalPackageParameters(HostPlatform, Plat) != "")
