Fixes for BootstrapPackagedGame:

- Should now build Win32 Shipping as well as Win64 Shipping on CIS
- Supports x86 and x64 installer packages; uses execuables rather than MSIs for XP support.
- Fix exception when packaging a non-Rocket game with no arguments

diff --git a/Engine/Source/Programs/AutomationTool/Win/WinPlatform.Automation.cs b/Engine/Source/Programs/AutomationTool/Win/WinPlatform.Automation.cs
--- a/Engine/Source/Programs/AutomationTool/Win/WinPlatform.Automation.cs
+++ b/Engine/Source/Programs/AutomationTool/Win/WinPlatform.Automation.cs
@@ -148,12 +148,6 @@ public abstract class BaseWinPlatform : Platform
 				WorkingFileType = StagedFileType.DebugNonUFS;
 			}
 		}
-		
-		if(Params.Prereqs)
-		{
-			string InstallerRelativePath = CombinePaths("Engine", "Extras", "Redist", "en-us");
-			SC.StageFiles(StagedFileType.NonUFS, CombinePaths(SC.LocalRoot, InstallerRelativePath), "UE4PrereqSetup.msi", false, null, InstallerRelativePath);
-		}
 	}
 
 	void StageBootstrapExecutable(DeploymentContext SC, string TargetFile, string StagedRelativeTargetPath, string StagedArguments)
@@ -188,10 +182,10 @@ public abstract class BaseWinPlatform : Platform
 				if(GroupIcon != null) Update.SetIcons(IconResourceId, GroupIcon);
 
 				const int ExecFileResourceId = 201;
-				Update.SetData(ExecFileResourceId, ResourceType.RawData, Encoding.Unicode.GetBytes(StagedRelativeTargetPath));
+				Update.SetData(ExecFileResourceId, ResourceType.RawData, Encoding.Unicode.GetBytes(StagedRelativeTargetPath + "\0"));
 
 				const int ExecArgsResourceId = 202;
-				Update.SetData(ExecArgsResourceId, ResourceType.RawData, Encoding.Unicode.GetBytes(StagedArguments));
+				Update.SetData(ExecArgsResourceId, ResourceType.RawData, Encoding.Unicode.GetBytes(StagedArguments + "\0"));
 			}
 
 			// Copy it to the staging directory
@@ -323,6 +317,17 @@ public class Win64Platform : BaseWinPlatform
 	}
 
 	public override bool IsSupported { get { return true; } }
+
+	public override void GetFilesToDeployOrStage(ProjectParams Params, DeploymentContext SC)
+	{
+		base.GetFilesToDeployOrStage(Params, SC);
+		
+		if(Params.Prereqs)
+		{
+			string InstallerRelativePath = CombinePaths("Engine", "Extras", "Redist", "en-us");
+			SC.StageFiles(StagedFileType.NonUFS, CombinePaths(SC.LocalRoot, InstallerRelativePath), "UE4PrereqSetup_x64.exe", false, null, InstallerRelativePath);
+		}
+	}
 }
 
 public class Win32Platform : BaseWinPlatform
@@ -333,4 +338,15 @@ public class Win32Platform : BaseWinPlatform
 	}
 
 	public override bool IsSupported { get { return true; } }
+
+	public override void GetFilesToDeployOrStage(ProjectParams Params, DeploymentContext SC)
+	{
+		base.GetFilesToDeployOrStage(Params, SC);
+		
+		if(Params.Prereqs)
+		{
+			string InstallerRelativePath = CombinePaths("Engine", "Extras", "Redist", "en-us");
+			SC.StageFiles(StagedFileType.NonUFS, CombinePaths(SC.LocalRoot, InstallerRelativePath), "UE4PrereqSetup_x86.exe", false, null, InstallerRelativePath);
+		}
+	}
 }
diff --git a/Engine/Source/Programs/Windows/BootstrapPackagedGame/BootstrapPackagedGame.Target.cs b/Engine/Source/Programs/Windows/BootstrapPackagedGame/BootstrapPackagedGame.Target.cs
--- a/Engine/Source/Programs/Windows/BootstrapPackagedGame/BootstrapPackagedGame.Target.cs
+++ b/Engine/Source/Programs/Windows/BootstrapPackagedGame/BootstrapPackagedGame.Target.cs
@@ -55,7 +55,16 @@ public class BootstrapPackagedGameTarget : TargetRules
 	{
 		bInternalToolOnly = false;
 		SeparateNode = false;
-		return (InHostPlatform == UnrealTargetPlatform.Win32 || InHostPlatform == UnrealTargetPlatform.Win64);
+		return (InHostPlatform == UnrealTargetPlatform.Win64);
+	}
+
+	public override List<UnrealTargetPlatform> GUBP_ToolPlatforms(UnrealTargetPlatform InHostPlatform)
+	{
+		if (InHostPlatform == UnrealTargetPlatform.Win64)
+		{
+			return new List<UnrealTargetPlatform> { UnrealTargetPlatform.Win64, UnrealTargetPlatform.Win32 };
+		}
+		return base.GUBP_ToolPlatforms(InHostPlatform);
 	}
 	
 	public override List<UnrealTargetConfiguration> GUBP_ToolConfigs(UnrealTargetPlatform InHostPlatform)
diff --git a/Engine/Source/Programs/Windows/BootstrapPackagedGame/Private/BootstrapPackagedGame.cpp b/Engine/Source/Programs/Windows/BootstrapPackagedGame/Private/BootstrapPackagedGame.cpp
--- a/Engine/Source/Programs/Windows/BootstrapPackagedGame/Private/BootstrapPackagedGame.cpp
+++ b/Engine/Source/Programs/Windows/BootstrapPackagedGame/Private/BootstrapPackagedGame.cpp
@@ -46,7 +46,11 @@ int InstallMissingPrerequisites(const WCHAR* BaseDirectory)
 
 		// If we don't have the installer, just notify the user and quit
 		WCHAR PrereqInstaller[MAX_PATH];
-		PathCombine(PrereqInstaller, BaseDirectory, L"Engine\\Extras\\Redist\\en-us\\UE4PrereqSetup.msi");
+#ifdef _M_X64
+		PathCombine(PrereqInstaller, BaseDirectory, L"Engine\\Extras\\Redist\\en-us\\UE4PrereqSetup_x64.exe");
+#else
+		PathCombine(PrereqInstaller, BaseDirectory, L"Engine\\Extras\\Redist\\en-us\\UE4PrereqSetup_x86.exe");
+#endif
 		if(GetFileAttributes(PrereqInstaller) == INVALID_FILE_ATTRIBUTES)
 		{
 			MessageBox(NULL, MissingPrerequisitesMsg, NULL, MB_OK);
