fix for Tappy Chicken not launching past splash screen on iOS
fix for Android unable to launch projects via UFE
fix for Project Launcher default roles settings not being used
TTP347511, TTP347613, TTP347768
#ue4

diff --git a/Engine/Source/Developer/LauncherServices/Private/Launcher/LauncherWorker.cpp b/Engine/Source/Developer/LauncherServices/Private/Launcher/LauncherWorker.cpp
--- a/Engine/Source/Developer/LauncherServices/Private/Launcher/LauncherWorker.cpp
+++ b/Engine/Source/Developer/LauncherServices/Private/Launcher/LauncherWorker.cpp
@@ -195,13 +195,12 @@ FString FLauncherWorker::CreateUATCommand( const ILauncherProfileRef& InProfile,
 	FString CommandLine = FString::Printf(TEXT(" -cmdline=\"%s -Messaging\""),
 		*InitialMap);
 
-	FString AdditionalCommandLine = FString::Printf(TEXT(" -addcmdline=\"%s -SessionId=%s -SessionOwner=%s -SessionName='%s'%s\""),
+	FString AdditionalCommandLine = FString::Printf(TEXT(" -addcmdline=\"%s -SessionId=%s -SessionOwner=%s -SessionName='%s'\""),
 		*InitialMap,
-		*InstanceId.ToString(),
+//		*InstanceId.ToString(),
 		*SessionId.ToString(),
 		FPlatformProcess::UserName(false),
-		*InProfile->GetName(),
-		InProfile->GetLaunchRoles().Num() > 0 ? (InProfile->GetLaunchRoles()[0]->IsVsyncEnabled() ? TEXT(" -vsync") : TEXT("")) : TEXT(""));
+		*InProfile->GetName());
 
 	// staging directory
 	FString StageDirectory = TEXT("");
@@ -219,35 +218,37 @@ FString FLauncherWorker::CreateUATCommand( const ILauncherProfileRef& InProfile,
 	FString OptionalParams = TEXT("");
 	for (int32 PlatformIndex = 0; PlatformIndex < InPlatforms.Num(); ++PlatformIndex)
 	{
+		// Platform info for the given platform
+		const PlatformInfo::FPlatformInfo* PlatformInfo = PlatformInfo::FindPlatformInfo(FName(*InPlatforms[PlatformIndex]));
+
 		// switch server and no editor platforms to the proper type
-		if (InPlatforms[PlatformIndex] == TEXT("LinuxServer"))
+		if (PlatformInfo->TargetPlatformName == FName("LinuxServer"))
 		{
 			ServerPlatforms += TEXT("+Linux");
 		}
-		else if (InPlatforms[PlatformIndex] == TEXT("WindowsServer"))
+		else if (PlatformInfo->TargetPlatformName == FName("WindowsServer"))
 		{
 			ServerPlatforms += TEXT("+Win64");
 		}
-		else if (InPlatforms[PlatformIndex] == TEXT("LinuxNoEditor"))
+		else if (PlatformInfo->TargetPlatformName == FName("LinuxNoEditor"))
 		{
 			Platforms += TEXT("+Linux");
 		}
-		else if (InPlatforms[PlatformIndex] == TEXT("WindowsNoEditor") || InPlatforms[PlatformIndex] == TEXT("Windows"))
+		else if (PlatformInfo->TargetPlatformName == FName("WindowsNoEditor") || PlatformInfo->TargetPlatformName == FName("Windows"))
 		{
 			Platforms += TEXT("+Win64");
 		}
-		else if (InPlatforms[PlatformIndex] == TEXT("MacNoEditor"))
+		else if (PlatformInfo->TargetPlatformName == FName("MacNoEditor"))
 		{
 			Platforms += TEXT("+Mac");
 		}
 		else
 		{
 			Platforms += TEXT("+");
-			Platforms += InPlatforms[PlatformIndex];
+			Platforms += PlatformInfo->TargetPlatformName.ToString();
 		}
 
 		// Append any extra UAT flags specified for this platform flavor
-		const PlatformInfo::FPlatformInfo* PlatformInfo = PlatformInfo::FindPlatformInfo(FName(*InPlatforms[PlatformIndex]));
 		if (!PlatformInfo->UATCommandLine.IsEmpty())
 		{
 			OptionalParams += TEXT(" ");
@@ -266,14 +267,18 @@ FString FLauncherWorker::CreateUATCommand( const ILauncherProfileRef& InProfile,
 
 	UATCommand += PlatformCommand;
 	UATCommand += ServerCommand;
+	UATCommand += OptionalParams;
 
 	// device list
 	FString DeviceNames = TEXT("");
 	FString DeviceCommand = TEXT("");
+	FString RoleCommands = TEXT("");
 	ILauncherDeviceGroupPtr DeviceGroup = InProfile->GetDeployedDeviceGroup();
 	if (DeviceGroup.IsValid())
 	{
 		const TArray<FString>& Devices = DeviceGroup->GetDeviceIDs();
+		bool bVsyncAdded = false;
+
 		// for each deployed device...
 		for (int32 DeviceIndex = 0; DeviceIndex < Devices.Num(); ++DeviceIndex)
 		{
@@ -285,6 +290,19 @@ FString FLauncherWorker::CreateUATCommand( const ILauncherProfileRef& InProfile,
 			{
 				// add the platform
 				DeviceNames += TEXT("+\"") + DeviceId + TEXT("\"");
+				TArray<ILauncherProfileLaunchRolePtr> Roles;
+				if (InProfile->GetLaunchRolesFor(DeviceId, Roles) > 0)
+				{
+					for (int32 RoleIndex = 0; RoleIndex < Roles.Num(); RoleIndex++)
+					{
+						if (!bVsyncAdded && Roles[RoleIndex]->IsVsyncEnabled())
+						{
+							RoleCommands += TEXT(" -vsync");
+							bVsyncAdded = true;
+						}
+						RoleCommands += *(TEXT(" ") + Roles[RoleIndex]->GetCommandLine());
+					}
+				}
 			}			
 		}
 	}
@@ -428,6 +446,7 @@ FString FLauncherWorker::CreateUATCommand( const ILauncherProfileRef& InProfile,
 		if (InProfile->GetLaunchMode() != ELauncherProfileLaunchModes::DoNotLaunch)
 		{
 			UATCommand += TEXT(" -run");
+			UATCommand += RoleCommands;
 
 			FCommandDesc Desc;
 			FText Command = FText::Format(LOCTEXT("LauncherRunDesc", "Launching on {0}"), FText::FromString(DeviceNames.RightChop(1)));
diff --git a/Engine/Source/Programs/AutomationTool/Android/AndroidPlatform.Automation.cs b/Engine/Source/Programs/AutomationTool/Android/AndroidPlatform.Automation.cs
--- a/Engine/Source/Programs/AutomationTool/Android/AndroidPlatform.Automation.cs
+++ b/Engine/Source/Programs/AutomationTool/Android/AndroidPlatform.Automation.cs
@@ -770,7 +770,7 @@ public class AndroidPlatform : Platform
 
 		if (!Params.CookOnTheFly && !Params.SkipCookOnTheFly) // only stage the UFS files if we are not using cook on the fly
 		{
-			SC.StageFiles(StagedFileType.UFS, CombinePaths(SC.LocalRoot, "Engine/Content/Localization/ICU"), "*", true, null, null, false, !Params.Pak);
+			SC.StageFiles(StagedFileType.UFS, CombinePaths(SC.LocalRoot, "Engine/Content/Localization/ICU"), "*", true, null, null, false, !Params.UsePak(SC.StageTargetPlatform));
 		}
 	}
 
@@ -815,7 +815,8 @@ public class AndroidPlatform : Platform
 
 	public override PakType RequiresPak(ProjectParams Params)
 	{
-		return PakType.Always;
+		// if packaging is enabled, always create a pak, otherwise use the Params.Pak value
+		return Params.Package ? PakType.Always : PakType.DontCare;
 	}
     
 	#region Hooks
diff --git a/Engine/Source/Programs/AutomationTool/IOS/IOSPlatform.Automation.cs b/Engine/Source/Programs/AutomationTool/IOS/IOSPlatform.Automation.cs
--- a/Engine/Source/Programs/AutomationTool/IOS/IOSPlatform.Automation.cs
+++ b/Engine/Source/Programs/AutomationTool/IOS/IOSPlatform.Automation.cs
@@ -741,7 +741,7 @@ public class IOSPlatform : Platform
 		}
 
 		// stage required icu files
-		SC.StageFiles (StagedFileType.UFS, CombinePaths (SC.LocalRoot, "Engine/Content/Localization/ICU"), "*", true, null, null, false, !Params.Pak);
+		SC.StageFiles(StagedFileType.UFS, CombinePaths(SC.LocalRoot, "Engine/Content/Localization/ICU"), "*", true, null, null, false, !Params.UsePak(SC.StageTargetPlatform));
 	}
 
 	public override void GetFilesToArchive(ProjectParams Params, DeploymentContext SC)
diff --git a/Engine/Source/Programs/AutomationTool/Linux/LinuxPlatform.Automation.cs b/Engine/Source/Programs/AutomationTool/Linux/LinuxPlatform.Automation.cs
--- a/Engine/Source/Programs/AutomationTool/Linux/LinuxPlatform.Automation.cs
+++ b/Engine/Source/Programs/AutomationTool/Linux/LinuxPlatform.Automation.cs
@@ -55,7 +55,7 @@ public abstract class BaseLinuxPlatform : Platform
 		}
 
 		SC.StageFiles(StagedFileType.NonUFS, CombinePaths(SC.ProjectRoot, "Content/Splash"), "Splash.bmp", false, null, null, true);
-		SC.StageFiles(StagedFileType.UFS, CombinePaths(SC.LocalRoot, "Engine/Content/Localization/ICU"), "*", true, null, null, false, !Params.Pak);
+		SC.StageFiles(StagedFileType.UFS, CombinePaths(SC.LocalRoot, "Engine/Content/Localization/ICU"), "*", true, null, null, false, !Params.UsePak(SC.StageTargetPlatform));
 
 
         if (Params.StageNonMonolithic)
@@ -198,7 +198,7 @@ EOF
 
 # Set permissions
 chmod 755 $HOME/{3}
-chmod 700 $HOME/Desktop/{1}.desktop", DesiredGLVersion, SC.ShortProjectName, SC.ShortProjectName, binPath, (Params.Pak ? " -pak" : ""), iconPath);
+chmod 700 $HOME/Desktop/{1}.desktop", DesiredGLVersion, SC.ShortProjectName, SC.ShortProjectName, binPath, (Params.UsePak(SC.StageTargetPlatform) ? " -pak" : ""), iconPath);
 			// End Bash Shell Script
 
 			string scriptFile = Path.GetTempFileName();
diff --git a/Engine/Source/Programs/AutomationTool/Mac/MacPlatform.Automation.cs b/Engine/Source/Programs/AutomationTool/Mac/MacPlatform.Automation.cs
--- a/Engine/Source/Programs/AutomationTool/Mac/MacPlatform.Automation.cs
+++ b/Engine/Source/Programs/AutomationTool/Mac/MacPlatform.Automation.cs
@@ -49,7 +49,7 @@ public class MacPlatform : Platform
 		// Copy the splash screen, Mac specific
 		SC.StageFiles(StagedFileType.NonUFS, CombinePaths(SC.ProjectRoot, "Content/Splash"), "Splash.bmp", false, null, null, true);
 
-		SC.StageFiles(StagedFileType.UFS, CombinePaths(SC.LocalRoot, "Engine/Content/Localization/ICU"), "*", true, null, null, false, !Params.Pak);
+		SC.StageFiles(StagedFileType.UFS, CombinePaths(SC.LocalRoot, "Engine/Content/Localization/ICU"), "*", true, null, null, false, !Params.UsePak(SC.StageTargetPlatform));
 	}
 
 	public override void Package(ProjectParams Params, DeploymentContext SC, int WorkingCL)
diff --git a/Engine/Source/Programs/AutomationTool/ProjectParams.cs b/Engine/Source/Programs/AutomationTool/ProjectParams.cs
--- a/Engine/Source/Programs/AutomationTool/ProjectParams.cs
+++ b/Engine/Source/Programs/AutomationTool/ProjectParams.cs
@@ -683,6 +683,14 @@ namespace AutomationTool
 		[Help("pak", "generate a pak file")]
 		public bool Pak { private set; get; }
 
+		/// <summary>
+		/// 
+		/// </summary>
+		public bool UsePak(Platform PlatformToCheck)
+		{
+			return Pak || PlatformToCheck.RequiresPak(this) == Platform.PakType.Always;
+		}
+
 		/// <summary>
 		/// Shared: Encryption keys used for signing the pak file.
 		/// </summary>
diff --git a/Engine/Source/Programs/AutomationTool/Scripts/CopyBuildToStagingDirectory.Automation.cs b/Engine/Source/Programs/AutomationTool/Scripts/CopyBuildToStagingDirectory.Automation.cs
--- a/Engine/Source/Programs/AutomationTool/Scripts/CopyBuildToStagingDirectory.Automation.cs
+++ b/Engine/Source/Programs/AutomationTool/Scripts/CopyBuildToStagingDirectory.Automation.cs
@@ -133,11 +133,11 @@ public partial class Project : CommandUtils
 			SC.NonUFSStagingFiles = ConvertToLower(SC.NonUFSStagingFiles);
 			SC.NonUFSStagingFilesDebug = ConvertToLower(SC.NonUFSStagingFilesDebug);
 		}
-		if (Params.Pak && BuildPlatform.DeployPakInternalLowerCaseFilenames())
+		if (Params.UsePak(SC.StageTargetPlatform) && BuildPlatform.DeployPakInternalLowerCaseFilenames())
 		{
 			SC.UFSStagingFiles = ConvertToLower(SC.UFSStagingFiles);
 		}
-		else if (!Params.Pak && BuildPlatform.DeployLowerCaseFilenames(true))
+		else if (!Params.UsePak(SC.StageTargetPlatform) && BuildPlatform.DeployLowerCaseFilenames(true))
 		{
 			SC.UFSStagingFiles = ConvertToLower(SC.UFSStagingFiles);
 		}
@@ -203,32 +203,32 @@ public partial class Project : CommandUtils
             }
 
 			// Engine ufs (content)
-			SC.StageFiles(StagedFileType.UFS, CombinePaths(SC.LocalRoot, "Engine/Config"), "*", true, null, null, false, !Params.Pak); // TODO: Exclude localization data generation config files.
+			SC.StageFiles(StagedFileType.UFS, CombinePaths(SC.LocalRoot, "Engine/Config"), "*", true, null, null, false, !Params.UsePak(SC.StageTargetPlatform)); // TODO: Exclude localization data generation config files.
 
 			if (Params.bUsesSlate)
 			{
 				if (Params.bUsesSlateEditorStyle)
 				{
-					SC.StageFiles(StagedFileType.UFS, CombinePaths(SC.LocalRoot, "Engine/Content/Editor/Slate"), "*", true, null, null, false, !Params.Pak);
+					SC.StageFiles(StagedFileType.UFS, CombinePaths(SC.LocalRoot, "Engine/Content/Editor/Slate"), "*", true, null, null, false, !Params.UsePak(SC.StageTargetPlatform));
 				}
-				SC.StageFiles(StagedFileType.UFS, CombinePaths(SC.LocalRoot, "Engine/Content/Slate"), "*", true, null, null, false, !Params.Pak);
-				SC.StageFiles(StagedFileType.UFS, CombinePaths(SC.ProjectRoot, "Content/Slate"), "*", true, null, CombinePaths(SC.RelativeProjectRootForStage, "Content/Slate"), true, !Params.Pak);
+				SC.StageFiles(StagedFileType.UFS, CombinePaths(SC.LocalRoot, "Engine/Content/Slate"), "*", true, null, null, false, !Params.UsePak(SC.StageTargetPlatform));
+				SC.StageFiles(StagedFileType.UFS, CombinePaths(SC.ProjectRoot, "Content/Slate"), "*", true, null, CombinePaths(SC.RelativeProjectRootForStage, "Content/Slate"), true, !Params.UsePak(SC.StageTargetPlatform));
 
 			}
             foreach (string Culture in CulturesToStage)
             {
-                SC.StageFiles(StagedFileType.UFS, CombinePaths(SC.LocalRoot, "Engine/Content/Localization/Engine", Culture), "*.locres", true, null, null, true, !Params.Pak);
+				SC.StageFiles(StagedFileType.UFS, CombinePaths(SC.LocalRoot, "Engine/Content/Localization/Engine", Culture), "*.locres", true, null, null, true, !Params.UsePak(SC.StageTargetPlatform));
             }
-			SC.StageFiles(StagedFileType.UFS, CombinePaths(SC.LocalRoot, "Engine/Plugins"), "*.uplugin", true, null, null, true, !Params.Pak);
+			SC.StageFiles(StagedFileType.UFS, CombinePaths(SC.LocalRoot, "Engine/Plugins"), "*.uplugin", true, null, null, true, !Params.UsePak(SC.StageTargetPlatform));
 
 			// Game ufs (content)
-            
-			SC.StageFiles(StagedFileType.UFS, CombinePaths(SC.ProjectRoot), "*.uproject", false, null, CombinePaths(SC.RelativeProjectRootForStage), true, !Params.Pak);
-			SC.StageFiles(StagedFileType.UFS, CombinePaths(SC.ProjectRoot, "Config"), "*", true, null, CombinePaths(SC.RelativeProjectRootForStage, "Config"), true, !Params.Pak); // TODO: Exclude localization data generation config files.
-			SC.StageFiles(StagedFileType.UFS, CombinePaths(SC.ProjectRoot, "Plugins"), "*.uplugin", true, null, null, true, !Params.Pak);
+
+			SC.StageFiles(StagedFileType.UFS, CombinePaths(SC.ProjectRoot), "*.uproject", false, null, CombinePaths(SC.RelativeProjectRootForStage), true, !Params.UsePak(SC.StageTargetPlatform));
+			SC.StageFiles(StagedFileType.UFS, CombinePaths(SC.ProjectRoot, "Config"), "*", true, null, CombinePaths(SC.RelativeProjectRootForStage, "Config"), true, !Params.UsePak(SC.StageTargetPlatform)); // TODO: Exclude localization data generation config files.
+			SC.StageFiles(StagedFileType.UFS, CombinePaths(SC.ProjectRoot, "Plugins"), "*.uplugin", true, null, null, true, !Params.UsePak(SC.StageTargetPlatform));
 			foreach (string Culture in CulturesToStage)
 			{
-				SC.StageFiles(StagedFileType.UFS, CombinePaths(SC.ProjectRoot, "Content/Localization/Game", Culture), "*.locres", true, null, CombinePaths(SC.RelativeProjectRootForStage, "Content/Localization/Game", Culture), true, !Params.Pak);
+				SC.StageFiles(StagedFileType.UFS, CombinePaths(SC.ProjectRoot, "Content/Localization/Game", Culture), "*.locres", true, null, CombinePaths(SC.RelativeProjectRootForStage, "Content/Localization/Game", Culture), true, !Params.UsePak(SC.StageTargetPlatform));
 			}
 
 			// Stage any additional UFS and NonUFS paths specified in the project ini files; these dirs are relative to the game content directory
@@ -246,7 +246,7 @@ public partial class Project : CommandUtils
 						if (PathParts.Length == 3)
 						{
 							var RelativePath = PathParts[1];
-							SC.StageFiles(StagedFileType.UFS, CombinePaths(ProjectContentRoot, RelativePath), "*", true, null, CombinePaths(StageContentRoot, RelativePath), true, !Params.Pak);
+							SC.StageFiles(StagedFileType.UFS, CombinePaths(ProjectContentRoot, RelativePath), "*", true, null, CombinePaths(StageContentRoot, RelativePath), true, !Params.UsePak(SC.StageTargetPlatform));
 						}
 					}
 				}
@@ -276,12 +276,12 @@ public partial class Project : CommandUtils
 
 			if (SC.StageTargetPlatform.StageMovies)
 			{
-				SC.StageFiles(StagedFileTypeForMovies, CombinePaths(SC.LocalRoot, "Engine/Content/Movies"), "*", true, null, CombinePaths(SC.RelativeProjectRootForStage, "Engine/Content/Movies"), true, !Params.Pak);
-				SC.StageFiles(StagedFileTypeForMovies, CombinePaths(SC.ProjectRoot, "Content/Movies"), "*", true, null, CombinePaths(SC.RelativeProjectRootForStage, "Content/Movies"), true, !Params.Pak);
+				SC.StageFiles(StagedFileTypeForMovies, CombinePaths(SC.LocalRoot, "Engine/Content/Movies"), "*", true, null, CombinePaths(SC.RelativeProjectRootForStage, "Engine/Content/Movies"), true, !Params.UsePak(SC.StageTargetPlatform));
+				SC.StageFiles(StagedFileTypeForMovies, CombinePaths(SC.ProjectRoot, "Content/Movies"), "*", true, null, CombinePaths(SC.RelativeProjectRootForStage, "Content/Movies"), true, !Params.UsePak(SC.StageTargetPlatform));
 			}
 
 			// eliminate the sand box
-			SC.StageFiles(StagedFileType.UFS, CombinePaths(SC.ProjectRoot, "Saved", "Cooked", SC.CookPlatform), "*", true, null, "", true, !Params.Pak);
+			SC.StageFiles(StagedFileType.UFS, CombinePaths(SC.ProjectRoot, "Saved", "Cooked", SC.CookPlatform), "*", true, null, "", true, !Params.UsePak(SC.StageTargetPlatform));
 
 			// CrashReportClient is a standalone slate app that does not look in the generated pak file, so it needs the Content/Slate and Shaders/StandaloneRenderer folders Non-UFS
 			// @todo Make CrashReportClient more portable so we don't have to do this
@@ -710,7 +710,7 @@ public partial class Project : CommandUtils
 		{
 			return;
 		}
-		DumpManifest(SC, CombinePaths(CmdEnv.LogFolder, "FinalCopy" + (SC.DedicatedServer ? "_Server" : "")), !Params.Pak);
+		DumpManifest(SC, CombinePaths(CmdEnv.LogFolder, "FinalCopy" + (SC.DedicatedServer ? "_Server" : "")), !Params.UsePak(SC.StageTargetPlatform));
 		CopyUsingStagingManifest(Params, SC);
 
 		var ThisPlatform = SC.StageTargetPlatform;
diff --git a/Engine/Source/Programs/AutomationTool/Scripts/RunProjectCommand.Automation.cs b/Engine/Source/Programs/AutomationTool/Scripts/RunProjectCommand.Automation.cs
--- a/Engine/Source/Programs/AutomationTool/Scripts/RunProjectCommand.Automation.cs
+++ b/Engine/Source/Programs/AutomationTool/Scripts/RunProjectCommand.Automation.cs
@@ -740,7 +740,7 @@ public partial class Project : CommandUtils
 					TempCmdLine += "-allowcaching ";
 				}
 			}
-			else if (Params.Pak || Params.SignedPak)
+			else if (Params.UsePak(SC.StageTargetPlatform) || Params.SignedPak)
 			{
 				if (Params.SignedPak)
 				{
@@ -925,7 +925,7 @@ public partial class Project : CommandUtils
 			Args += String.Format("{0} -server -abslog={1}  -unattended -FORCELOGFLUSH -log -Messaging -nomcp", Map, CommandUtils.MakePathSafeToUseWithCommandLine(ServerLogFile));
 		}
 
-		if (ServerParams.Pak || ServerParams.SignedPak)
+		if (ServerParams.UsePak(SC.StageTargetPlatform) || ServerParams.SignedPak)
 		{
 			if (ServerParams.SignedPak)
 			{
diff --git a/Engine/Source/Programs/AutomationTool/Win/WinPlatform.Automation.cs b/Engine/Source/Programs/AutomationTool/Win/WinPlatform.Automation.cs
--- a/Engine/Source/Programs/AutomationTool/Win/WinPlatform.Automation.cs
+++ b/Engine/Source/Programs/AutomationTool/Win/WinPlatform.Automation.cs
@@ -88,7 +88,7 @@ public abstract class BaseWinPlatform : Platform
 		// Copy the splash screen, windows specific
 		SC.StageFiles(StagedFileType.NonUFS, CombinePaths(SC.ProjectRoot, "Content/Splash"), "Splash.bmp", false, null, null, true);
 
-		SC.StageFiles(StagedFileType.UFS, CombinePaths(SC.LocalRoot, "Engine/Content/Localization/ICU"), "*", true, null, null, false, !Params.Pak);
+		SC.StageFiles(StagedFileType.UFS, CombinePaths(SC.LocalRoot, "Engine/Content/Localization/ICU"), "*", true, null, null, false, !Params.UsePak(SC.StageTargetPlatform));
 
         if (Params.StageNonMonolithic)
         {
