UAT fixes for removing RocketGame

diff --git a/Engine/Source/Programs/AutomationTool/ProjectParams.cs b/Engine/Source/Programs/AutomationTool/ProjectParams.cs
--- a/Engine/Source/Programs/AutomationTool/ProjectParams.cs
+++ b/Engine/Source/Programs/AutomationTool/ProjectParams.cs
@@ -1201,11 +1201,6 @@ namespace AutomationTool
 						EditorTarget = TargetData.TargetName;
 					}
 
-					if (DetectedTargets.TryGetValue(TargetRules.TargetType.RocketGame, out TargetData))
-					{
-						GameTarget = TargetData.TargetName;
-					}
-
 					if (DetectedTargets.TryGetValue(TargetRules.TargetType.Program, out TargetData))
 					{
 						ProgramTarget = TargetData.TargetName;
@@ -1255,7 +1250,6 @@ namespace AutomationTool
 				var ValidGameTargetTypes = new TargetRules.TargetType[]
 				{
 					GameTargetType,
-					TargetRules.TargetType.RocketGame,
 					TargetRules.TargetType.Program		
 				};
 
diff --git a/Engine/Source/Programs/AutomationTool/ProjectUtils.cs b/Engine/Source/Programs/AutomationTool/ProjectUtils.cs
--- a/Engine/Source/Programs/AutomationTool/ProjectUtils.cs
+++ b/Engine/Source/Programs/AutomationTool/ProjectUtils.cs
@@ -231,7 +231,6 @@ namespace AutomationTool
 					break;
 				case TargetRules.TargetType.Client:
 				case TargetRules.TargetType.Game:
-				case TargetRules.TargetType.RocketGame:
 					if (!bIsCodeBasedProject)
 					{
 						BinPath = CommandUtils.CombinePaths(CommandUtils.CmdEnv.LocalRoot, "Engine", "Binaries");
