Report the bootstrap program as the executable to run on Windows.

diff --git a/Engine/Source/Programs/AutomationTool/Win/WinPlatform.Automation.cs b/Engine/Source/Programs/AutomationTool/Win/WinPlatform.Automation.cs
--- a/Engine/Source/Programs/AutomationTool/Win/WinPlatform.Automation.cs
+++ b/Engine/Source/Programs/AutomationTool/Win/WinPlatform.Automation.cs
@@ -244,16 +244,16 @@ public abstract class BaseWinPlatform : Platform
 				foreach (var StageExecutable in SC.StageExecutables)
 				{
 					string ExeName = SC.StageTargetPlatform.GetPlatformExecutableName(StageExecutable);
-					if (!SC.IsCodeBasedProject && (!bIsRun || !SC.Stage))
+					if(bIsRun)
+					{
+						ExecutableNames.Add(CombinePaths(SC.RuntimeRootDir, SC.ShortProjectName + Ext));
+					}
+					else if(!SC.IsCodeBasedProject)
 					{
 						ExecutableNames.Add(CombinePaths(SC.RuntimeRootDir, "Engine/Binaries", SC.PlatformDir, ExeName + Ext));
 					}
 					else
 					{
-						if (ExeName.Contains("UE4Game") && !SC.IsCodeBasedProject && bIsRun)
-						{
-							ExeName = ExeName.Replace("UE4Game", SC.ShortProjectName);
-						}
 						ExecutableNames.Add(CombinePaths(SC.RuntimeProjectRootDir, "Binaries", SC.PlatformDir, ExeName + Ext));
 					}
 				}
