Fix linux staging

diff --git a/Engine/Source/Programs/AutomationTool/Linux/LinuxPlatform.Automation.cs b/Engine/Source/Programs/AutomationTool/Linux/LinuxPlatform.Automation.cs
--- a/Engine/Source/Programs/AutomationTool/Linux/LinuxPlatform.Automation.cs
+++ b/Engine/Source/Programs/AutomationTool/Linux/LinuxPlatform.Automation.cs
@@ -60,9 +60,6 @@ public abstract class BaseLinuxPlatform : Platform
 
         if (Params.StageNonMonolithic)
         {
-            SC.StageFiles(StagedFileType.NonUFS, CombinePaths(SC.LocalRoot, "Engine/Binaries/ThirdParty/jemalloc/", SC.PlatformDir, "x86_64-unknown-linux-gnu"), "libjemalloc.so.1", false, null, CombinePaths("Engine/Binaries", SC.PlatformDir));
-            SC.StageFiles(StagedFileType.NonUFS, CombinePaths(SC.LocalRoot, "Engine/Binaries/ThirdParty/SDL2/", SC.PlatformDir, "x86_64-unknown-linux-gnu"), "libSDL2-2.0.so.1", false, null, CombinePaths("Engine/Binaries", SC.PlatformDir));
-
             if (SC.DedicatedServer)
             {
                 SC.StageFiles(StagedFileType.NonUFS, CombinePaths(SC.LocalRoot, "Engine/Binaries", SC.PlatformDir), "UE4Server");
