Fix path to generated headers on Android. No longer has CPU suffix in the path.

diff --git a/Engine/Source/Programs/AutomationTool/Scripts/GUBP.Automation.cs b/Engine/Source/Programs/AutomationTool/Scripts/GUBP.Automation.cs
--- a/Engine/Source/Programs/AutomationTool/Scripts/GUBP.Automation.cs
+++ b/Engine/Source/Programs/AutomationTool/Scripts/GUBP.Automation.cs
@@ -1541,10 +1541,6 @@ public class GUBP : BuildCommand
 			{
 				// host platform overlaps with the editor, so we don't want them here
 				var PlatformDir = TargetPlatform.ToString();
-				if (TargetPlatform == UnrealTargetPlatform.Android)
-				{
-					PlatformDir = "Android/-armv7";
-				}
 				foreach (var FileToCopy in CommandUtils.FindFiles("*.h", true, CommandUtils.CombinePaths(CmdEnv.LocalRoot, @"Engine\Intermediate\Build\", PlatformDir, "Inc")))
 				{
 					AddBuildProduct(FileToCopy);
