Add IPP to be built on Mac as part of the tools node, but overwrite it with the Win64 one in a GUBP Rocket build. Add it to the list of distilled build products for iOS on Mac.

diff --git a/Engine/Source/Programs/AutomationTool/Scripts/GUBP.Automation.cs b/Engine/Source/Programs/AutomationTool/Scripts/GUBP.Automation.cs
--- a/Engine/Source/Programs/AutomationTool/Scripts/GUBP.Automation.cs
+++ b/Engine/Source/Programs/AutomationTool/Scripts/GUBP.Automation.cs
@@ -944,6 +944,16 @@ public class GUBP : BuildCommand
 						}
                     );
             }
+			else if(HostPlatform == UnrealTargetPlatform.Mac)
+			{
+                Agenda.IOSDotNetProjects.AddRange(
+                        new string[]
+						{
+							CombinePaths(@"Engine\Source\Programs\IOS\iPhonePackager\iPhonePackager.csproj"),
+							CombinePaths(@"Engine\Source\Programs\IOS\DeploymentInterface\DeploymentInterface.csproj"),
+						}
+                    );
+			}
 
             string AddArgs = "-nobuilduht -skipactionhistory -CopyAppBundleBackToDevice" + bp.RocketUBTArgs(); ;
 
