Set UAT to always use the original executable to launch packaged games; passing project on the command line conflicts with the one added by the bootstrapper.

diff --git a/Engine/Source/Programs/AutomationTool/Win/WinPlatform.Automation.cs b/Engine/Source/Programs/AutomationTool/Win/WinPlatform.Automation.cs
--- a/Engine/Source/Programs/AutomationTool/Win/WinPlatform.Automation.cs
+++ b/Engine/Source/Programs/AutomationTool/Win/WinPlatform.Automation.cs
@@ -244,17 +244,13 @@ public abstract class BaseWinPlatform : Platform
 				foreach (var StageExecutable in SC.StageExecutables)
 				{
 					string ExeName = SC.StageTargetPlatform.GetPlatformExecutableName(StageExecutable);
-					if(bIsRun)
+					if(SC.IsCodeBasedProject)
 					{
-						ExecutableNames.Add(CombinePaths(SC.RuntimeRootDir, SC.ShortProjectName + Ext));
-					}
-					else if(!SC.IsCodeBasedProject)
-					{
-						ExecutableNames.Add(CombinePaths(SC.RuntimeRootDir, "Engine/Binaries", SC.PlatformDir, ExeName + Ext));
+						ExecutableNames.Add(CombinePaths(SC.RuntimeProjectRootDir, "Binaries", SC.PlatformDir, ExeName + Ext));
 					}
 					else
 					{
-						ExecutableNames.Add(CombinePaths(SC.RuntimeProjectRootDir, "Binaries", SC.PlatformDir, ExeName + Ext));
+						ExecutableNames.Add(CombinePaths(SC.RuntimeRootDir, "Engine/Binaries", SC.PlatformDir, ExeName + Ext));
 					}
 				}
 			}
