fix for renamed projects not properly naming the xcode schemes
consolidation of IPA name code
TTP343209
#ue4

diff --git a/Engine/Source/Programs/AutomationTool/IOS/IOSPlatform.Automation.cs b/Engine/Source/Programs/AutomationTool/IOS/IOSPlatform.Automation.cs
--- a/Engine/Source/Programs/AutomationTool/IOS/IOSPlatform.Automation.cs
+++ b/Engine/Source/Programs/AutomationTool/IOS/IOSPlatform.Automation.cs
@@ -82,9 +82,9 @@ public class IOSPlatform : Platform
 		return 4;
 	}
 
-	protected string MakeIPAFileName( UnrealTargetConfiguration TargetConfiguration, string ProjectGameExeFilename )
+	protected string MakeIPAFileName( UnrealTargetConfiguration TargetConfiguration, ProjectParams Params )
 	{
-		var ProjectIPA = Path.ChangeExtension(ProjectGameExeFilename, null);
+		string ProjectIPA = Path.Combine(Path.GetDirectoryName(Params.RawProjectPath), "Binaries", "IOS", (Params.Distribution ? "Distro_" : "") + Params.ShortProjectName);
 		if (TargetConfiguration != UnrealTargetConfiguration.Development)
 		{
 			ProjectIPA += "-" + PlatformType.ToString() + "-" + TargetConfiguration.ToString();
@@ -171,8 +171,8 @@ public class IOSPlatform : Platform
 			}
 			var TargetConfiguration = SC.StageTargetConfigurations[0];
 
-			var ProjectStub = Params.ProjectGameExeFilename;
-			var ProjectIPA = MakeIPAFileName(TargetConfiguration, Params.ProjectGameExeFilename);
+			var ProjectStub = Path.GetFullPath (Params.ProjectGameExeFilename) + "/" + Params.ShortProjectName + Path.GetExtension(Params.ProjectGameExeFilename);
+			var ProjectIPA = MakeIPAFileName(TargetConfiguration, Params);
 
 			// package a .ipa from the now staged directory
 			var IPPExe = CombinePaths(CmdEnv.LocalRoot, "Engine/Binaries/DotNET/IOS/IPhonePackager.exe");
@@ -184,16 +184,6 @@ public class IOSPlatform : Platform
 
 			bool cookonthefly = Params.CookOnTheFly || Params.SkipCookOnTheFly;
 
-			// rename the .ipa if not code based
-			if (!Params.IsCodeBasedProject)
-			{
-				ProjectIPA = Path.Combine(Path.GetDirectoryName(Params.RawProjectPath), "Binaries", "IOS", (Params.Distribution ? "Distro_" : "") + Params.ShortProjectName + ".ipa");
-				if (TargetConfiguration != UnrealTargetConfiguration.Development)
-				{
-					ProjectIPA = Path.Combine(Path.GetDirectoryName(Params.RawProjectPath), "Binaries", "IOS", (Params.Distribution ? "Distro_" : "") + Params.ShortProjectName + "-" + PlatformType.ToString() + "-" + TargetConfiguration.ToString() + ".ipa");
-				}
-			}
-
 			// delete the .ipa to make sure it was made
 			DeleteFile(ProjectIPA);
 
@@ -751,22 +741,7 @@ public class IOSPlatform : Platform
 			throw new AutomationException("iOS is currently only able to package one target configuration at a time, but StageTargetConfigurations contained {0} configurations", SC.StageTargetConfigurations.Count);
 		}
 		var TargetConfiguration = SC.StageTargetConfigurations[0];
-
-		string ProjectGameExeFilename = Params.ProjectGameExeFilename;
-		if (UnrealBuildTool.BuildHostPlatform.Current.Platform == UnrealTargetPlatform.Mac)
-		{
-			ProjectGameExeFilename = CombinePaths (Path.GetDirectoryName(Params.RawProjectPath), "Binaries", "IOS", Path.GetFileName (Params.ProjectGameExeFilename));
-		}
-		var ProjectIPA = MakeIPAFileName( TargetConfiguration, ProjectGameExeFilename );
-		// rename the .ipa if not code based
-		if (!Params.IsCodeBasedProject)
-		{
-			ProjectIPA = Path.Combine(Path.GetDirectoryName(Params.RawProjectPath), "Binaries", "IOS", (Params.Distribution ? "Distro_" : "") + Params.ShortProjectName + ".ipa");
-			if (TargetConfiguration != UnrealTargetConfiguration.Development)
-			{
-				ProjectIPA = Path.Combine(Path.GetDirectoryName(Params.RawProjectPath), "Binaries", "IOS", (Params.Distribution ? "Distro_" : "") + Params.ShortProjectName + "-" + PlatformType.ToString() + "-" + TargetConfiguration.ToString() + ".ipa");
-			}
-		}
+		var ProjectIPA = MakeIPAFileName( TargetConfiguration, Params );
 
 		// verify the .ipa exists
 		if (!FileExists(ProjectIPA))
@@ -790,20 +765,9 @@ public class IOSPlatform : Platform
 				throw new AutomationException("iOS cannot deploy a package made for distribution.");
 			}
 			var TargetConfiguration = SC.StageTargetConfigurations[0];
-
-//			var ProjectStub = Params.ProjectGameExeFilename;	
-			var ProjectIPA = MakeIPAFileName(TargetConfiguration, Params.ProjectGameExeFilename);
-			// rename the .ipa if not code based
-			if (!Params.IsCodeBasedProject)
-			{
-				ProjectIPA = Path.Combine(Path.GetDirectoryName(Params.RawProjectPath), "Binaries", "IOS", Params.ShortProjectName + ".ipa");
-				if (TargetConfiguration != UnrealTargetConfiguration.Development)
-				{
-					ProjectIPA = Path.Combine(Path.GetDirectoryName(Params.RawProjectPath), "Binaries", "IOS", Params.ShortProjectName + "-" + PlatformType.ToString() + "-" + TargetConfiguration.ToString() + ".ipa");
-				}
-			}
-
+			var ProjectIPA = MakeIPAFileName(TargetConfiguration, Params);
 			var StagedIPA = SC.StageDirectory + "\\" + Path.GetFileName(ProjectIPA);
+
 			// verify the .ipa exists
 			if (!FileExists(StagedIPA))
 			{
@@ -833,7 +797,6 @@ public class IOSPlatform : Platform
 		return false;
 	}
 
-
 	public override bool DeployLowerCaseFilenames(bool bUFSFile)
 	{
 		// we shouldn't modify the case on files like Info.plist or the icons
diff --git a/Engine/Source/Programs/UnrealBuildTool/System/XcodeProjectFileGenerator.cs b/Engine/Source/Programs/UnrealBuildTool/System/XcodeProjectFileGenerator.cs
--- a/Engine/Source/Programs/UnrealBuildTool/System/XcodeProjectFileGenerator.cs
+++ b/Engine/Source/Programs/UnrealBuildTool/System/XcodeProjectFileGenerator.cs
@@ -1068,10 +1068,9 @@ namespace UnrealBuildTool
 			{
 				string TargetName = Utils.GetFilenameWithoutAnyExtensions(Path.GetFileName(TargetPath));
 				bool WantProjectFileForTarget = true;
+				bool IsEngineTarget = false;
 				if (bGeneratingGameProjectFiles || bGeneratingRocketProjectFiles)
 				{
-					bool IsEngineTarget = false;
-
 					// Check to see if this is an Engine target.  That is, the target is located under the "Engine" folder
 					string TargetFileRelativeToEngineDirectory = Utils.MakePathRelativeTo(TargetPath, Path.Combine(EngineRelativePath), AlwaysTreatSourceAsDirectory: false);
 					if (!TargetFileRelativeToEngineDirectory.StartsWith("..") && !Path.IsPathRooted(TargetFileRelativeToEngineDirectory))
@@ -1110,10 +1109,27 @@ namespace UnrealBuildTool
 						TargetsThatNeedApp.Add(TargetName);
 					}
 
+					// if the project is not an engine project check to make sure we have the correct name
+					string DisplayName = TargetName;
+					if (!IsEngineTarget && TargetRulesObject.Type != TargetRules.TargetType.Program)
+					{
+						List<UProjectInfo> AllGames = UProjectInfo.FilterGameProjects(true, bGeneratingGameProjectFiles ? GameProjectName : null);
+						Console.WriteLine("Game Count: " + AllGames.Count);
+						UProjectInfo ProjectInfo = FindGameContainingFile (AllGames, TargetFilePath);
+						if (ProjectInfo != null)
+						{
+							DisplayName = ProjectInfo.GameName;
+							if (TargetName.Contains("Editor"))
+							{
+								DisplayName += "Editor";
+							}
+						}
+					}
+
 					// @todo: Remove target platform param and merge Mac and iOS targets. For now BuildTarget knows how to build iOS, but cannot run iOS apps, so we need separate DeployTarget.
 					bool bIsMacOnly = !SupportedPlatforms.Contains(UnrealTargetPlatform.IOS);
 
-					XcodeProjectTarget BuildTarget = new XcodeProjectTarget(TargetName + " - Mac", TargetName, XcodeTargetType.Legacy, TargetFilePath, "", UnrealTargetPlatform.Mac, bIsMacOnly);
+					XcodeProjectTarget BuildTarget = new XcodeProjectTarget(DisplayName + " - Mac", TargetName, XcodeTargetType.Legacy, TargetFilePath, "", UnrealTargetPlatform.Mac, bIsMacOnly);
 					if (!bGeneratingRunIOSProject)
 					{
 						ProjectTargets.Add(BuildTarget);
@@ -1130,14 +1146,14 @@ namespace UnrealBuildTool
 								FrameworkRefs.Add(new XcodeFrameworkRef(Framework));
 							}
 
-							XcodeProjectTarget IOSDeployTarget = new XcodeProjectTarget(TargetName + " - iOS", TargetName, XcodeTargetType.Native, TargetFilePath, TargetName + ".app", UnrealTargetPlatform.IOS, false, null, true, FrameworkRefs);
+							XcodeProjectTarget IOSDeployTarget = new XcodeProjectTarget(DisplayName + " - iOS", TargetName, XcodeTargetType.Native, TargetFilePath, TargetName + ".app", UnrealTargetPlatform.IOS, false, null, true, FrameworkRefs);
 							ProjectTargets.Add(IOSDeployTarget);
 						}
 						else
 						{
 							XcodeContainerItemProxy ContainerProxy = new XcodeContainerItemProxy(ProjectTarget.Guid, BuildTarget.Guid, BuildTarget.DisplayName);
 							XcodeTargetDependency TargetDependency = new XcodeTargetDependency(BuildTarget.DisplayName, BuildTarget.Guid, ContainerProxy.Guid);
-							XcodeProjectTarget IOSDeployTarget = new XcodeProjectTarget(TargetName + " - iOS", TargetName, XcodeTargetType.Native, TargetFilePath, TargetName + ".app", UnrealTargetPlatform.IOS, false, new List<XcodeTargetDependency>() { TargetDependency }, true);
+							XcodeProjectTarget IOSDeployTarget = new XcodeProjectTarget(DisplayName + " - iOS", TargetName, XcodeTargetType.Native, TargetFilePath, TargetName + ".app", UnrealTargetPlatform.IOS, false, new List<XcodeTargetDependency>() { TargetDependency }, true);
 							ProjectTargets.Add(IOSDeployTarget);
 							ContainerItemProxies.Add(ContainerProxy);
 							TargetDependencies.Add(TargetDependency);
