Cookonthefly fixed in UFE, caused by 2317591

Switched off commandline overrides for cook on the fly deploys,

diff --git a/Engine/Source/Programs/AutomationTool/IOS/IOSPlatform.Automation.cs b/Engine/Source/Programs/AutomationTool/IOS/IOSPlatform.Automation.cs
--- a/Engine/Source/Programs/AutomationTool/IOS/IOSPlatform.Automation.cs
+++ b/Engine/Source/Programs/AutomationTool/IOS/IOSPlatform.Automation.cs
@@ -778,11 +778,14 @@ public class IOSPlatform : Platform
 				}
 			}
 
+			// Add a commandline for this deploy, if the config allows it.
+			string AdditionalCommandline = (Params.FileServer || Params.CookOnTheFly) ? "" : (" -additionalcommandline " + "\"" + Params.RunCommandline + "\"");
+
 			// deploy the .ipa
 			var IPPExe = CombinePaths(CmdEnv.LocalRoot, "Engine/Binaries/DotNET/IOS/IPhonePackager.exe");
 
 			// check for it in the stage directory
-			RunAndLog(CmdEnv, IPPExe, "Deploy \"" + Path.GetFullPath(StagedIPA) + "\"" + (String.IsNullOrEmpty(Params.Device) ? "" : " -device " + Params.Device.Substring(4)) + (" -additionalcommandline " + "\"" + Params.RunCommandline + "\""));
+			RunAndLog(CmdEnv, IPPExe, "Deploy \"" + Path.GetFullPath(StagedIPA) + "\"" + (String.IsNullOrEmpty(Params.Device) ? "" : " -device " + Params.Device.Substring(4)) + AdditionalCommandline);
 		}
         PrintRunTime();
     }
