Removed default fallback logic in BuildCookRun that was making it impossible to use configuration files for the default list of cultures to stage.

diff --git a/Engine/Source/Programs/AutomationTool/Scripts/BuildCookRun.Automation.cs b/Engine/Source/Programs/AutomationTool/Scripts/BuildCookRun.Automation.cs
--- a/Engine/Source/Programs/AutomationTool/Scripts/BuildCookRun.Automation.cs
+++ b/Engine/Source/Programs/AutomationTool/Scripts/BuildCookRun.Automation.cs
@@ -111,10 +111,6 @@ public class BuildCookRun : BuildCommand
         {
             Params.CulturesToCook = new ParamList<string>(CulturesToCook.Split('+'));
         }
-		else
-		{
-			Params.CulturesToCook = new ParamList<string>(new string[] { "en" });
-		}
 
 		if (Params.DedicatedServer)
 		{
