Improved staging logic for localization data that properly attempts to stage parent cultures.

diff --git a/Engine/Source/Programs/AutomationTool/Scripts/CopyBuildToStagingDirectory.Automation.cs b/Engine/Source/Programs/AutomationTool/Scripts/CopyBuildToStagingDirectory.Automation.cs
--- a/Engine/Source/Programs/AutomationTool/Scripts/CopyBuildToStagingDirectory.Automation.cs
+++ b/Engine/Source/Programs/AutomationTool/Scripts/CopyBuildToStagingDirectory.Automation.cs
@@ -143,6 +143,46 @@ public partial class Project : CommandUtils
 		}
 	}
 
+    private static int StageLocalizationDataForCulture(DeploymentContext SC, string CultureName, string SourceDirectory, string DestinationDirectory = null, bool bRemap = true)
+    {
+        int FilesAdded = 0;
+
+        string[] LocaleTags = CultureName.Split('_');
+
+        List<string> PotentialParentCultures = new List<string>();
+        
+        if (LocaleTags.Length > 0)
+        {
+            if (LocaleTags.Length > 1 && LocaleTags.Length > 2)
+            {
+                PotentialParentCultures.Add(string.Join("_", LocaleTags[0], LocaleTags[1], LocaleTags[2]));
+            }
+            if (LocaleTags.Length > 2)
+            {
+                PotentialParentCultures.Add(string.Join("_", LocaleTags[0], LocaleTags[2]));
+            }
+            if (LocaleTags.Length > 1)
+            {
+                PotentialParentCultures.Add(string.Join("_", LocaleTags[0], LocaleTags[1]));
+            }
+            PotentialParentCultures.Add(LocaleTags[0]);
+        }
+
+        string[] FoundDirectories = CommandUtils.FindDirectories(true, "*", false, new string[] { SourceDirectory });
+        foreach (string FoundDirectory in FoundDirectories)
+        {
+            string DirectoryName = CommandUtils.GetLastDirectoryName(FoundDirectory);
+            string CanonicalizedPotentialCulture = DirectoryName.Replace('-', '_');
+
+            if (PotentialParentCultures.Contains(CanonicalizedPotentialCulture))
+            {
+                FilesAdded += SC.StageFiles(StagedFileType.UFS, CombinePaths(SourceDirectory, DirectoryName), "*.locres", true, null, DestinationDirectory != null ? CombinePaths(DestinationDirectory, DirectoryName) : null, true, bRemap);
+            }
+        }
+
+        return FilesAdded;
+    }
+
 	public static void CreateStagingManifest(ProjectParams Params, DeploymentContext SC)
 	{
 		if (!Params.Stage)
@@ -175,20 +215,18 @@ public partial class Project : CommandUtils
 
 		if (!Params.CookOnTheFly && !Params.SkipCookOnTheFly) // only stage the UFS files if we are not using cook on the fly
 		{
-            ConfigCacheIni PlatformGameConfig = null;
-            Params.GameConfigs.TryGetValue(SC.StageTargetPlatform.PlatformType, out PlatformGameConfig);
+            ConfigCacheIni PlatformGameConfig = new ConfigCacheIni(SC.StageTargetPlatform.PlatformType, "Game", CommandUtils.GetDirectoryName(Params.RawProjectPath));
 
             // Initialize cultures to stage.
             List<string> CulturesToStage = null;
-            if (CulturesToStage != null)
-            {
-                CulturesToStage.AddRange(Params.CulturesToCook);
-            }
-            else
+
+            // Use parameters if provided.
+            if (Params.CulturesToCook != null && Params.CulturesToCook.Count > 0)
             {
-                CulturesToStage = new List<string>(Params.CulturesToCook);
+                CulturesToStage = Params.CulturesToCook;
             }
 
+            // Use configuration if otherwise lacking cultures to stage.
             if (CulturesToStage == null || CulturesToStage.Count == 0)
             {
                 if (PlatformGameConfig != null)
@@ -197,6 +235,7 @@ public partial class Project : CommandUtils
                 }
             }
 
+            // Error if no cultures have been provided.
             if (CulturesToStage == null || CulturesToStage.Count == 0)
             {
                 throw new AutomationException("No cultures were specified for cooking and packaging. This will lead to fatal errors when launching. Specify culture codes via commandline (-CookCultures=) or using project packaging settings (+CulturesToStage).");
@@ -213,11 +252,10 @@ public partial class Project : CommandUtils
 				}
 				SC.StageFiles(StagedFileType.UFS, CombinePaths(SC.LocalRoot, "Engine/Content/Slate"), "*", true, null, null, false, !Params.UsePak(SC.StageTargetPlatform));
 				SC.StageFiles(StagedFileType.UFS, CombinePaths(SC.ProjectRoot, "Content/Slate"), "*", true, null, CombinePaths(SC.RelativeProjectRootForStage, "Content/Slate"), true, !Params.UsePak(SC.StageTargetPlatform));
-
 			}
             foreach (string Culture in CulturesToStage)
             {
-				SC.StageFiles(StagedFileType.UFS, CombinePaths(SC.LocalRoot, "Engine/Content/Localization/Engine", Culture), "*.locres", true, null, null, true, !Params.UsePak(SC.StageTargetPlatform));
+                StageLocalizationDataForCulture(SC, Culture, CombinePaths(SC.LocalRoot, "Engine/Content/Localization/Engine"), null, !Params.UsePak(SC.StageTargetPlatform));
             }
 			SC.StageFiles(StagedFileType.UFS, CombinePaths(SC.LocalRoot, "Engine/Plugins"), "*.uplugin", true, null, null, true, !Params.UsePak(SC.StageTargetPlatform));
 
@@ -228,7 +266,7 @@ public partial class Project : CommandUtils
 			SC.StageFiles(StagedFileType.UFS, CombinePaths(SC.ProjectRoot, "Plugins"), "*.uplugin", true, null, null, true, !Params.UsePak(SC.StageTargetPlatform));
 			foreach (string Culture in CulturesToStage)
 			{
-				SC.StageFiles(StagedFileType.UFS, CombinePaths(SC.ProjectRoot, "Content/Localization/Game", Culture), "*.locres", true, null, CombinePaths(SC.RelativeProjectRootForStage, "Content/Localization/Game", Culture), true, !Params.UsePak(SC.StageTargetPlatform));
+                StageLocalizationDataForCulture(SC, Culture, CombinePaths(SC.ProjectRoot, "Content/Localization/Game"), CombinePaths(SC.RelativeProjectRootForStage, "Content/Localization/Game"), !Params.UsePak(SC.StageTargetPlatform));
 			}
 
 			// Stage any additional UFS and NonUFS paths specified in the project ini files; these dirs are relative to the game content directory
