Post binary launcher manifests to GameDev MCP during build

diff --git a/Engine/Source/Programs/AutomationTool/MCPPublic.cs b/Engine/Source/Programs/AutomationTool/MCPPublic.cs
--- a/Engine/Source/Programs/AutomationTool/MCPPublic.cs
+++ b/Engine/Source/Programs/AutomationTool/MCPPublic.cs
@@ -141,7 +141,7 @@ namespace EpicGames.MCP.Automation
         /// <summary>
         /// Gets the base filename of the manifest that would be created by invoking the BuildPatchTool with the given parameters.
         /// </summary>
-        public string ManifestFilename
+        public virtual string ManifestFilename
         {
             get
             {
@@ -244,6 +244,28 @@ namespace EpicGames.MCP.Automation
 		}
     }
 
+	public class BuildPatchToolBinaryStagingInfo : BuildPatchToolStagingInfo
+	{
+		public BuildPatchToolBinaryStagingInfo(BuildPatchToolStagingInfo StagingInfo) : base(
+			StagingInfo.OwnerCommand,
+			StagingInfo.AppName,
+			StagingInfo.McpConfigKey,
+			StagingInfo.AppID,
+			StagingInfo.BuildVersion,
+			StagingInfo.Platform,
+			StagingInfo.StagingDir)
+		{
+		}
+
+		public override string ManifestFilename
+		{
+			get
+			{
+				return base.ManifestFilename.Replace(".manifest", ".binary.manifest");
+			}
+		}
+	}
+
 
     /// <summary>
     /// Class that provides programmatic access to the BuildPatchTool
