build script, override cook platform for wex to be android_etc1

diff --git a/Engine/Source/Programs/AutomationTool/Scripts/GUBP.Automation.cs b/Engine/Source/Programs/AutomationTool/Scripts/GUBP.Automation.cs
--- a/Engine/Source/Programs/AutomationTool/Scripts/GUBP.Automation.cs
+++ b/Engine/Source/Programs/AutomationTool/Scripts/GUBP.Automation.cs
@@ -2163,7 +2163,7 @@ public class GUBP : BuildCommand
         UnrealTargetPlatform TargetPlatform;
         string CookPlatform;
         bool bIsMassive;
-		bool bisCode;
+		
 
         public CookNode(GUBP bp, UnrealTargetPlatform InHostPlatform, BranchInfo.BranchUProject InGameProj, UnrealTargetPlatform InTargetPlatform, string InCookPlatform)
             : base(InHostPlatform)
@@ -2179,15 +2179,13 @@ public class GUBP : BuildCommand
                 {
                     AddDependency(EditorPlatformNode.StaticGetFullName(HostPlatform, TargetPlatform));
                 }
-            }
-
+            }			
             bool bIsShared = false;
             // is this the "base game" or a non code project?
             if (InGameProj.GameName != bp.Branch.BaseEngineProject.GameName && GameProj.Properties.Targets.ContainsKey(TargetRules.TargetType.Editor))
             {
                 var Options = InGameProj.Options(HostPlatform);
-                bIsMassive = Options.bIsMassive;
-				bisCode = true;
+                bIsMassive = Options.bIsMassive;				
                 AddDependency(EditorGameNode.StaticGetFullName(HostPlatform, GameProj));
                 // add an arc to prevent cooks from running until promotable is labeled
                 if (Options.bIsPromotable)
@@ -2205,8 +2203,7 @@ public class GUBP : BuildCommand
             }
             else
             {
-                bIsShared = true;
-				bisCode = false;
+                bIsShared = true;				
                 AddPseudodependency(GamePlatformMonolithicsNode.StaticGetFullName(HostPlatform, bp.Branch.BaseEngineProject, TargetPlatform));
             }
             if (bIsShared)
@@ -2228,7 +2225,7 @@ public class GUBP : BuildCommand
                 {
                     AddPseudodependency(CookNode.StaticGetFullName(HostPlatform, bp.Branch.BaseEngineProject, BaseCookedPlatform));
                 }
-            }
+            }			
         }
         public static string StaticGetFullName(UnrealTargetPlatform InHostPlatform, BranchInfo.BranchUProject InGameProj, string InCookPlatform)
         {
@@ -2269,33 +2266,9 @@ public class GUBP : BuildCommand
             {
                 // not sure if we need something here or if the cook commandlet will automatically convert the exe name
             }
-			string AdditionalParams = "";
-			if (bisCode)
-			{
-				var Target = GameProj.Properties.Targets[TargetRules.TargetType.Game];
-				if (CookPlatform.Contains("Server"))
-				{
-					Target = GameProj.Properties.Targets[TargetRules.TargetType.Server];
-				}
-				if (CookPlatform.Contains("Client"))
-				{
-					Target = GameProj.Properties.Targets[TargetRules.TargetType.Client];
-				}
 
-				if (Target.Rules.GUBP_AdditionalCookParameters(HostPlatform, CookPlatform) != "")
-				{
-					AdditionalParams = Target.Rules.GUBP_AdditionalCookParameters(HostPlatform, CookPlatform);
-				}
-			}
+			CommandUtils.CookCommandlet(GameProj.FilePath, "UE4Editor-Cmd.exe", null, null, null, CookPlatform);
 		
-			if (AdditionalParams == "")
-			{
-				CommandUtils.CookCommandlet(GameProj.FilePath, "UE4Editor-Cmd.exe", null, null, null, CookPlatform);
-			}
-			else
-			{
-				CommandUtils.CookCommandlet(GameProj.FilePath, "UE4Editor-Cmd.exe", null, null, null, CookPlatform, AdditionalParams);
-			}
             var CookedPath = RootIfAnyForTempStorage();
             var CookedFiles = CommandUtils.FindFiles("*", true, CookedPath);
             if (CookedFiles.GetLength(0) < 1)
@@ -2448,6 +2421,10 @@ public class GUBP : BuildCommand
                         {
                             //@todo how do we get the client target platform?
                             string CookedPlatform = Platform.Platforms[TargetPlatform].GetCookPlatform(Kind == TargetRules.TargetType.Server, Kind == TargetRules.TargetType.Client, "");
+							if (Target.Rules.GUBP_AlternateCookPlatform(HostPlatform, CookedPlatform) != "")
+							{
+								CookedPlatform = Target.Rules.GUBP_AlternateCookPlatform(HostPlatform, CookedPlatform);
+							}
                             AddDependency(CookNode.StaticGetFullName(HostPlatform, GameProj, CookedPlatform));
                             AddDependency(GamePlatformMonolithicsNode.StaticGetFullName(HostPlatform, GameProj, TargetPlatform));
                         }
@@ -5245,7 +5222,7 @@ public class GUBP : BuildCommand
 
                             if (ActivePlatforms.Contains(Plat))
                             {
-                                string CookedPlatform = Platform.Platforms[Plat].GetCookPlatform(Kind == TargetRules.TargetType.Server, Kind == TargetRules.TargetType.Client, "");
+                                string CookedPlatform = Platform.Platforms[Plat].GetCookPlatform(Kind == TargetRules.TargetType.Server, Kind == TargetRules.TargetType.Client, "");								
                                 if (!GUBPNodes.ContainsKey(CookNode.StaticGetFullName(HostPlatform, Branch.BaseEngineProject, CookedPlatform)))
                                 {
                                     GameCookNodes.Add(AddNode(new CookNode(this, HostPlatform, Branch.BaseEngineProject, Plat, CookedPlatform)));
@@ -5451,6 +5428,10 @@ public class GUBP : BuildCommand
 								if (!AdditionalPlatforms.Contains(Plat))
 								{
 									string CookedPlatform = Platform.Platforms[Plat].GetCookPlatform(Kind == TargetRules.TargetType.Server, Kind == TargetRules.TargetType.Client, "");
+									if (Target.Rules.GUBP_AlternateCookPlatform(HostPlatform, CookedPlatform) != "")
+									{
+										CookedPlatform = Target.Rules.GUBP_AlternateCookPlatform(HostPlatform, CookedPlatform);
+									}
 									if (!GUBPNodes.ContainsKey(CookNode.StaticGetFullName(HostPlatform, CodeProj, CookedPlatform)))
 									{
 										AddNode(new CookNode(this, HostPlatform, CodeProj, Plat, CookedPlatform));
diff --git a/Engine/Source/Programs/UnrealBuildTool/System/RulesCompiler.cs b/Engine/Source/Programs/UnrealBuildTool/System/RulesCompiler.cs
--- a/Engine/Source/Programs/UnrealBuildTool/System/RulesCompiler.cs
+++ b/Engine/Source/Programs/UnrealBuildTool/System/RulesCompiler.cs
@@ -860,6 +860,14 @@ namespace UnrealBuildTool
 		{
 			return "";
 		}
+
+		/// <summary>
+		/// Allow Cook Platform Override from a target file
+		/// </summary>
+		public virtual string GUBP_AlternateCookPlatform(UnrealTargetPlatform HostPlatform, string Platform)
+		{
+			return "";
+		}
 	}
 
 
