Fix crash when interrupting a child process.

diff --git a/Engine/Source/Programs/AutomationTool/ProcessUtils.cs b/Engine/Source/Programs/AutomationTool/ProcessUtils.cs
--- a/Engine/Source/Programs/AutomationTool/ProcessUtils.cs
+++ b/Engine/Source/Programs/AutomationTool/ProcessUtils.cs
@@ -262,8 +262,11 @@ namespace AutomationTool
 		/// </summary>
         public void DisposeProcess()
         {
-            Proc.Dispose();
-            Proc = null;
+			if(Proc != null)
+			{
+				Proc.Dispose();
+				Proc = null;
+			}
         }
 
 		/// <summary>
