build script, cleanup hastests logic

diff --git a/Engine/Source/Programs/AutomationTool/Scripts/GUBP.Automation.cs b/Engine/Source/Programs/AutomationTool/Scripts/GUBP.Automation.cs
--- a/Engine/Source/Programs/AutomationTool/Scripts/GUBP.Automation.cs
+++ b/Engine/Source/Programs/AutomationTool/Scripts/GUBP.Automation.cs
@@ -6232,7 +6232,7 @@ public class GUBP : BuildCommand
             {
                 if (GUBPNodes[NodeToDo].RunInEC() && !NodeIsAlreadyComplete(NodeToDo, LocalOnly)) // if something is already finished, we don't put it into EC  
                 {
-                    if ((NodeToDo.Contains("Test")) && !(NodeToDo.Contains("MakeBuild")) && !(NodeToDo.Contains("Unity")) && !(NodesToDo.Contains("TestBuild")) && !(NodesToDo.Contains("Compile")))
+                    if ((NodeToDo.Contains("Test")) && !(NodeToDo.Contains("MakeBuild")) && !(NodeToDo.Contains("Unity")) && !(NodesToDo.Contains("TestBuild")) && !(NodesToDo.Contains("Compile")) && !(NodesToDo.Contains("Product")))
                     {
                         bHasTests = true;
                     }                    
