Prevent duplicate builds of launcher being posted to MCP and ensure single latest build is labelled as LatestBinary in addition to Latest.

diff --git a/Engine/Source/Programs/AutomationTool/MCPPublic.cs b/Engine/Source/Programs/AutomationTool/MCPPublic.cs
--- a/Engine/Source/Programs/AutomationTool/MCPPublic.cs
+++ b/Engine/Source/Programs/AutomationTool/MCPPublic.cs
@@ -244,29 +244,6 @@ namespace EpicGames.MCP.Automation
 		}
     }
 
-	public class BuildPatchToolBinaryStagingInfo : BuildPatchToolStagingInfo
-	{
-		public BuildPatchToolBinaryStagingInfo(BuildPatchToolStagingInfo StagingInfo) : base(
-			StagingInfo.OwnerCommand,
-			StagingInfo.AppName,
-			StagingInfo.McpConfigKey,
-			StagingInfo.AppID,
-			StagingInfo.BuildVersion,
-			StagingInfo.Platform,
-			StagingInfo.StagingDir)
-		{
-		}
-
-		public override string ManifestFilename
-		{
-			get
-			{
-				return base.ManifestFilename.Replace(".manifest", ".binary.manifest");
-			}
-		}
-	}
-
-
     /// <summary>
     /// Class that provides programmatic access to the BuildPatchTool
     /// </summary>
