Linux fixes from PR #496, #497.

- use /bin/bash instead of /bin/sh for RunUAT.sh since the script has bashisms.
- name casing change for UnrealPak.
- contributed by [REDACTED_PERSON].

diff --git a/Engine/Source/Editor/MainFrame/Private/Frame/MainFrameActions.cpp b/Engine/Source/Editor/MainFrame/Private/Frame/MainFrameActions.cpp
--- a/Engine/Source/Editor/MainFrame/Private/Frame/MainFrameActions.cpp
+++ b/Engine/Source/Editor/MainFrame/Private/Frame/MainFrameActions.cpp
@@ -1085,7 +1085,7 @@ void FMainFrameActionCallbacks::CreateUatTask( const FString& CommandLine, const
 	FString CmdExe = TEXT("cmd.exe");
 #elif PLATFORM_LINUX
 	FString RunUATScriptName = TEXT("RunUAT.sh");
-	FString CmdExe = TEXT("/bin/sh");
+	FString CmdExe = TEXT("/bin/bash");
 #else
 	FString RunUATScriptName = TEXT("RunUAT.command");
 	FString CmdExe = TEXT("/bin/sh");
diff --git a/Engine/Source/Programs/AutomationTool/Scripts/CopyBuildToStagingDirectory.Automation.cs b/Engine/Source/Programs/AutomationTool/Scripts/CopyBuildToStagingDirectory.Automation.cs
--- a/Engine/Source/Programs/AutomationTool/Scripts/CopyBuildToStagingDirectory.Automation.cs
+++ b/Engine/Source/Programs/AutomationTool/Scripts/CopyBuildToStagingDirectory.Automation.cs
@@ -65,7 +65,7 @@ public partial class Project : CommandUtils
 		string UnrealPakResponseFileName = CombinePaths(CmdEnv.LogFolder, "PakList_" + PakName + ".txt");
 		WritePakResponseFile(UnrealPakResponseFileName, UnrealPakResponseFile, Compressed);
 
-		var UnrealPakExe = CombinePaths(CmdEnv.LocalRoot, "Engine/Binaries/Win64/unrealpak.exe");
+		var UnrealPakExe = CombinePaths(CmdEnv.LocalRoot, "Engine/Binaries/Win64/UnrealPak.exe");
 
 		Log("Running UnrealPak *******");
 		string CmdLine = CommandUtils.MakePathSafeToUseWithCommandLine(OutputLocation) + " -create=" + CommandUtils.MakePathSafeToUseWithCommandLine(UnrealPakResponseFileName);
