build script, moving formal shared builds to agent sharing

diff --git a/Engine/Source/Programs/AutomationTool/Scripts/GUBP.Automation.cs b/Engine/Source/Programs/AutomationTool/Scripts/GUBP.Automation.cs
--- a/Engine/Source/Programs/AutomationTool/Scripts/GUBP.Automation.cs
+++ b/Engine/Source/Programs/AutomationTool/Scripts/GUBP.Automation.cs
@@ -5200,6 +5200,7 @@ public class GUBP : BuildCommand
 
                 var GameTestNodes = new List<string>();
                 var GameCookNodes = new List<string>();
+				var FormalAgentSharingGroup = "Shared_FormalBuilds" + HostPlatformNode.StaticGetHostPlatformSuffix(HostPlatform);
 
 
                 //foreach (var Kind in BranchInfo.MonolithicKinds)//for now, non-code projects don't do client or server.
@@ -5289,7 +5290,14 @@ public class GUBP : BuildCommand
 														GUBPNodes[GamePlatformCookedAndCompiledNode.StaticGetFullName(HostPlatform, NonCodeProject, Plat)].AgentSharingGroup = BuildAgentSharingGroup;
 														GUBPNodes[NodeName].AgentSharingGroup = BuildAgentSharingGroup;
 													}
-													
+													else
+													{
+														GUBPNodes[NodeName].AgentSharingGroup = FormalAgentSharingGroup;
+														if(Plat == UnrealTargetPlatform.XboxOne)
+														{
+															GUBPNodes[NodeName].AgentSharingGroup = "";
+														}
+													}
 													if (PlatPair.bTest)
 													{
 														AddNode(new FormalBuildTestNode(this, NonCodeProject, HostPlatform, Plat, PlatPair.TargetConfig));
@@ -5353,7 +5361,7 @@ public class GUBP : BuildCommand
                     AddNode(new GameAggregateNode(this, HostPlatform, Branch.BaseEngineProject, "AllCookedTests", GameTestNodes));
                 }
             }
-
+			
             foreach (var CodeProj in Branch.CodeProjects)
             {
                 var Options = CodeProj.Options(HostPlatform);
@@ -5389,7 +5397,7 @@ public class GUBP : BuildCommand
                 }
 
                 var CookedAgentSharingGroup = AgentShareName + "_CookedTests" + HostPlatformNode.StaticGetHostPlatformSuffix(HostPlatform);
-
+				var FormalAgentSharingGroup = "Shared_FormalBuilds" + HostPlatformNode.StaticGetHostPlatformSuffix(HostPlatform);
                 var ServerPlatforms = new List<UnrealTargetPlatform>();
                 var GamePlatforms = new List<UnrealTargetPlatform>();
                 var GameTestNodes = new List<string>();
@@ -5477,6 +5485,7 @@ public class GUBP : BuildCommand
 												CookNode.StaticGetFullName(HostPlatform, CodeProj, CookedPlatform),
 												CookNode.StaticGetFullName(HostPlatform, Branch.BaseEngineProject, CookedPlatform));
 											string BuildAgentSharingGroup = "";
+											
 											if (Options.bSeparateGamePromotion)
 											{
 												BuildAgentSharingGroup = CodeProj.GameName + "_MakeFormalBuild_" + Plat.ToString() + HostPlatformNode.StaticGetHostPlatformSuffix(HostPlatform);
@@ -5488,7 +5497,14 @@ public class GUBP : BuildCommand
 												GUBPNodes[GamePlatformCookedAndCompiledNode.StaticGetFullName(HostPlatform, CodeProj, Plat)].AgentSharingGroup = BuildAgentSharingGroup;
 												GUBPNodes[FormalNodeName].AgentSharingGroup = BuildAgentSharingGroup;
 											}
-
+											else
+											{
+												GUBPNodes[FormalNodeName].AgentSharingGroup = FormalAgentSharingGroup;
+												if (Plat == UnrealTargetPlatform.XboxOne)
+												{
+													GUBPNodes[FormalNodeName].AgentSharingGroup = "";
+												}
+											}
 											if (Config.bTest)
 											{
 												AddNode(new FormalBuildTestNode(this, CodeProj, HostPlatform, Plat, Config.TargetConfig));
