Minor fixes to UAT projects (PR #492).

- Android/HTML5 automation projects reference OneSky, UAT compiled for AnyCPU platform.
- contributed by [REDACTED_PERSON].

diff --git a/Engine/Build/BatchFiles/Linux/GenerateProjectFiles.sh b/Engine/Build/BatchFiles/Linux/GenerateProjectFiles.sh
--- a/Engine/Build/BatchFiles/Linux/GenerateProjectFiles.sh
+++ b/Engine/Build/BatchFiles/Linux/GenerateProjectFiles.sh
@@ -151,6 +151,7 @@ xbuild Source/Programs/UnrealBuildTool/UnrealBuildTool_Mono.csproj \
 xbuild Source/Programs/AutomationTool/AutomationTool_Mono.csproj \
   /verbosity:quiet /nologo \
   /p:TargetFrameworkVersion=v4.0 \
+  /p:Platform="AnyCPU" \
   /p:Configuration="Development"
 
 xbuild Source/Programs/AutomationTool/Scripts/AutomationScripts.Automation.csproj \
diff --git a/Engine/Source/Programs/AutomationTool/Android/Android.Automation.csproj b/Engine/Source/Programs/AutomationTool/Android/Android.Automation.csproj
--- a/Engine/Source/Programs/AutomationTool/Android/Android.Automation.csproj
+++ b/Engine/Source/Programs/AutomationTool/Android/Android.Automation.csproj
@@ -34,6 +34,11 @@
     <TreatWarningsAsErrors>true</TreatWarningsAsErrors>
   </PropertyGroup>
   <ItemGroup>
+    <Reference Include="OneSky, Version=1.0.0.0, Culture=neutral, processorArchitecture=MSIL">
+      <HintPath>..\..\..\..\Binaries\DotNET\OneSky.dll</HintPath>
+      <SpecificVersion>False</SpecificVersion>
+      <Private>False</Private>
+    </Reference>
     <Reference Include="System" />
     <Reference Include="System.Core" />
     <Reference Include="System.Runtime.Serialization" />
@@ -70,4 +75,4 @@
     </ProjectReference>
   </ItemGroup>
   <Import Project="$(MSBuildToolsPath)\Microsoft.CSharp.targets" />
-</Project>
\ No newline at end of file
+</Project>
diff --git a/Engine/Source/Programs/AutomationTool/HTML5/HTML5.Automation.csproj b/Engine/Source/Programs/AutomationTool/HTML5/HTML5.Automation.csproj
--- a/Engine/Source/Programs/AutomationTool/HTML5/HTML5.Automation.csproj
+++ b/Engine/Source/Programs/AutomationTool/HTML5/HTML5.Automation.csproj
@@ -34,6 +34,11 @@
     <TreatWarningsAsErrors>true</TreatWarningsAsErrors>
   </PropertyGroup>
   <ItemGroup>
+    <Reference Include="OneSky, Version=1.0.0.0, Culture=neutral, processorArchitecture=MSIL">
+      <HintPath>..\..\..\..\Binaries\DotNET\OneSky.dll</HintPath>
+      <SpecificVersion>False</SpecificVersion>
+      <Private>False</Private>
+    </Reference>
     <Reference Include="System" />
     <Reference Include="System.Core" />
     <Reference Include="Microsoft.CSharp" />
@@ -63,4 +68,4 @@
     </ProjectReference>
   </ItemGroup>
   <Import Project="$(MSBuildToolsPath)\Microsoft.CSharp.targets" />
-</Project>
\ No newline at end of file
+</Project>
