Add support to automation tool for parsing version strings with only branch and changelist.
Increased minimum deletion time of build cleanup script to one week.

diff --git a/Engine/Source/Programs/AutomationTool/Utils.cs b/Engine/Source/Programs/AutomationTool/Utils.cs
--- a/Engine/Source/Programs/AutomationTool/Utils.cs
+++ b/Engine/Source/Programs/AutomationTool/Utils.cs
@@ -795,20 +795,35 @@ namespace AutomationTool
         /// Creates a <see cref="FEngineVersionSupport"/> from a string that matches the format given in <see cref="ToString"/>.
         /// </summary>
         /// <param name="versionString">Version string that should match the FEngineVersion::ToString() format.</param>
+		/// <param name="bAllowVersion">Optional parameter which if set to true, allows version strings with no version number specified.</param>
         /// <returns>a new instance with fields initialized to the match those given in the string.</returns>
-        public static FEngineVersionSupport FromString(string versionString)
+        public static FEngineVersionSupport FromString(string versionString, bool bAllowNoVersion = false)
         {
             try
             {
-                var dotSplit = versionString.Split(new[] { '.' }, 3);
-                var dashSplit = dotSplit[2].Split(new[] { '-' }, 2);
-                var plusSplit = dashSplit[1].Split(new[] { '+' }, 2);
-                var major = int.Parse(dotSplit[0]);
-                var minor = int.Parse(dotSplit[1]);
-                var patch = int.Parse(dashSplit[0]);
-                var changelist = int.Parse(plusSplit[0]);
-                var branchName = plusSplit[1];
-                return new FEngineVersionSupport(new Version(major, minor, patch), changelist, branchName); ;
+				if (bAllowNoVersion && versionString.StartsWith("++depot"))
+				{
+					// This form of version is used when a product has no major.minor.patch version
+					// E.g. ++depot+UE4-ProdName-CL-12345678
+					var clSplit = versionString.Split(new string[] { "-CL-" }, 2, StringSplitOptions.None);
+					var dashSplit = clSplit[1].Split(new[] { '-' }, 2);
+					var changelist = int.Parse(dashSplit[0]);
+					var branchName = clSplit[0];
+					return new FEngineVersionSupport(new Version(0, 0, 0), changelist, branchName);
+				}
+				else
+				{
+					// This is the standard Rocket versioning scheme, e.g. "4.5.0-12345678+++depot+UE4"
+					var dotSplit = versionString.Split(new[] { '.' }, 3);
+					var dashSplit = dotSplit[2].Split(new[] { '-' }, 2);
+					var plusSplit = dashSplit[1].Split(new[] { '+' }, 2);
+					var major = int.Parse(dotSplit[0]);
+					var minor = int.Parse(dotSplit[1]);
+					var patch = int.Parse(dashSplit[0]);
+					var changelist = int.Parse(plusSplit[0]);
+					var branchName = plusSplit[1];
+					return new FEngineVersionSupport(new Version(major, minor, patch), changelist, branchName);
+				}
             }
             catch (Exception ex)
             {
