UAT: Correct path for BlandProject

diff --git a/Engine/Source/Programs/AutomationTool/ProjectUtils.cs b/Engine/Source/Programs/AutomationTool/ProjectUtils.cs
--- a/Engine/Source/Programs/AutomationTool/ProjectUtils.cs
+++ b/Engine/Source/Programs/AutomationTool/ProjectUtils.cs
@@ -592,7 +592,7 @@ namespace AutomationTool
             }
             if (String.IsNullOrEmpty(BaseEngineProject.FilePath))
             {
-                throw new AutomationException("All branches must have the blank project /Samples/SampleGames/BlankProject");
+                throw new AutomationException("All branches must have the blank project /Samples/Sandbox/BlankProject");
             }
 
             CommandUtils.Log("  Base Engine:");
