Improved logic for staging localization data to psuedo-canonicalize culture names from settings.

diff --git a/Engine/Source/Programs/AutomationTool/Scripts/CopyBuildToStagingDirectory.Automation.cs b/Engine/Source/Programs/AutomationTool/Scripts/CopyBuildToStagingDirectory.Automation.cs
--- a/Engine/Source/Programs/AutomationTool/Scripts/CopyBuildToStagingDirectory.Automation.cs
+++ b/Engine/Source/Programs/AutomationTool/Scripts/CopyBuildToStagingDirectory.Automation.cs
@@ -147,7 +147,7 @@ public partial class Project : CommandUtils
     {
         int FilesAdded = 0;
 
-        string[] LocaleTags = CultureName.Split('_');
+        string[] LocaleTags = CultureName.Replace('-', '_').Split('_');
 
         List<string> PotentialParentCultures = new List<string>();
         
