build script, fix logic to determine if we're running automated tests

diff --git a/Engine/Source/Programs/AutomationTool/Scripts/GUBP.Automation.cs b/Engine/Source/Programs/AutomationTool/Scripts/GUBP.Automation.cs
--- a/Engine/Source/Programs/AutomationTool/Scripts/GUBP.Automation.cs
+++ b/Engine/Source/Programs/AutomationTool/Scripts/GUBP.Automation.cs
@@ -400,6 +400,10 @@ public class GUBP : BuildCommand
         {
             return true;
         }
+		public virtual bool IsTest()
+		{
+			return false;
+		}
         public virtual bool IsAggregate()
         {
             return false;
@@ -3207,6 +3211,10 @@ public class GUBP : BuildCommand
         {
             return ECPriority;
         }
+		public override bool IsTest()
+		{
+			return true;
+		}
         public override int CISFrequencyQuantumShift(GUBP bp)
         {
             return base.CISFrequencyQuantumShift(bp) + 5;
@@ -3799,7 +3807,7 @@ public class GUBP : BuildCommand
                 FrequencyString,
                 NodeIsAlreadyComplete(NodeToDo, LocalOnly) ? " - (Completed)" : "",
                 GUBPNodes[NodeToDo].TriggerNode() ? " - (TriggerNode)" : "",
-                GUBPNodes[NodeToDo].IsSticky() ? " - (Sticky)" : "",
+                GUBPNodes[NodeToDo].IsSticky() ? " - (Sticky)" : "",				
                 Agent,
                 MemoryReq,
                 EMails,
@@ -5126,7 +5134,7 @@ public class GUBP : BuildCommand
                     foreach (var Test in EditorTests)
                     {
                         if (!bNoAutomatedTesting)
-                        {
+                        {							
                             EditorTestNodes.Add(AddNode(new UATTestNode(this, HostPlatform, Branch.BaseEngineProject, Test.Key, Test.Value, AgentSharingGroup)));
                         
                             foreach (var NonCodeProject in Branch.NonCodeProjects)
@@ -5236,7 +5244,7 @@ public class GUBP : BuildCommand
                                 var RequiredPlatforms = new List<UnrealTargetPlatform> { Plat };
                                 if (!bNoAutomatedTesting)
                                 {
-                                    var ThisMonoGameTestNodes = new List<string>();
+                                    var ThisMonoGameTestNodes = new List<string>();	
                                     
                                     foreach (var Test in GameTests)
                                     {
@@ -5244,7 +5252,7 @@ public class GUBP : BuildCommand
                                         ThisMonoGameTestNodes.Add(AddNode(new UATTestNode(this, HostPlatform, Branch.BaseEngineProject, TestName, Test.Value, CookedAgentSharingGroup, false, RequiredPlatforms)));
                                     }
                                     if (ThisMonoGameTestNodes.Count > 0)
-                                    {
+                                    {										
                                         GameTestNodes.Add(AddNode(new GameAggregateNode(this, HostPlatform, Branch.BaseEngineProject, "CookedTests_" + Plat.ToString() + "_" + Kind.ToString() + HostPlatformNode.StaticGetHostPlatformSuffix(HostPlatform), ThisMonoGameTestNodes, 0.0f)));
                                     }
                                 }
@@ -5316,7 +5324,7 @@ public class GUBP : BuildCommand
 											ThisMonoGameTestNodes.Add(AddNode(new UATTestNode(this, HostPlatform, NonCodeProject, TestName, Test.Value, CookedAgentSharingGroup, false, RequiredPlatforms)));
 										}
 										if (ThisMonoGameTestNodes.Count > 0)
-										{                                        
+										{											
 											GameTestNodes.Add(AddNode(new GameAggregateNode(this, HostPlatform, NonCodeProject, "CookedTests_" + Plat.ToString() + "_" + Kind.ToString() + HostPlatformNode.StaticGetHostPlatformSuffix(HostPlatform), ThisMonoGameTestNodes, 0.0f)));
 										}
 									}
@@ -5879,7 +5887,7 @@ public class GUBP : BuildCommand
         Log("Desired Nodes");
         foreach (var NodeToDo in NodesToDo)
         {
-            Log("  {0}", NodeToDo);
+            Log("  {0}", NodeToDo);			
         }
         // if we are doing related to, then find things that depend on the selected nodes
         if (bRelatedToNode)
@@ -6087,6 +6095,10 @@ public class GUBP : BuildCommand
                 {
                     continue;
                 }
+				if(GUBPNodes[NodeToDo].IsTest())
+				{
+					bHasTests = true;
+				}
                 int MyIndex = OrdereredToDo.IndexOf(NodeToDo);
                 foreach (var Dep in GUBPNodes[NodeToDo].FullNamesOfDependencies)
                 {
@@ -6247,11 +6259,7 @@ public class GUBP : BuildCommand
             foreach (var NodeToDo in OrdereredToDo)
             {
                 if (GUBPNodes[NodeToDo].RunInEC() && !NodeIsAlreadyComplete(NodeToDo, LocalOnly)) // if something is already finished, we don't put it into EC  
-                {
-                    if ((NodeToDo.Contains("Test")) && !(NodeToDo.Contains("MakeBuild")) && !(NodeToDo.Contains("Unity")) && !(NodeToDo.Contains("TestBuild")) && !(NodeToDo.Contains("Compile")) && !(NodeToDo.Contains("Product")))
-                    {
-                        bHasTests = true;
-                    }                    
+                {                                      
                     string EMails;
                     var NodeProps = GetECPropsForNode(NodeToDo, CLString, out EMails);
                     ECProps.AddRange(NodeProps);
