build script, add linux editor cross compile node

diff --git a/Engine/Source/Programs/AutomationTool/Scripts/GUBP.Automation.cs b/Engine/Source/Programs/AutomationTool/Scripts/GUBP.Automation.cs
--- a/Engine/Source/Programs/AutomationTool/Scripts/GUBP.Automation.cs
+++ b/Engine/Source/Programs/AutomationTool/Scripts/GUBP.Automation.cs
@@ -885,7 +885,37 @@ public class GUBP : BuildCommand
             DeleteStaleDLLs(bp);
         }
     }
+	public class RootEditorCrossCompileLinuxNode : CompileNode
+	{
+		public RootEditorCrossCompileLinuxNode(UnrealTargetPlatform InHostPlatform)
+			: base(InHostPlatform)
+		{
+			AddDependency(RootEditorNode.StaticGetFullName(UnrealTargetPlatform.Win64));
+			AddDependency(ToolsForCompileNode.StaticGetFullName(UnrealTargetPlatform.Win64));
+		}
+		public static string StaticGetFullName(UnrealTargetPlatform InHostPlatform)
+		{
+			return "RootEditor_Linux";
+		}
+		public override string GetFullName()
+		{
+			return StaticGetFullName(HostPlatform);
+		}
+		public override int CISFrequencyQuantumShift(GUBP bp)
+		{
+			return base.CISFrequencyQuantumShift(bp) + 3;
+		}
+		public override UE4Build.BuildAgenda GetAgenda(GUBP bp)
+		{
+			var Agenda = new UE4Build.BuildAgenda();
 
+			string AddArgs = "-nobuilduht" + bp.RocketUBTArgs();
+			Agenda.AddTargets(
+				new string[] { bp.Branch.BaseEngineProject.Properties.Targets[TargetRules.TargetType.Editor].TargetName },
+				UnrealTargetPlatform.Linux, UnrealTargetConfiguration.Development, InAddArgs: AddArgs);
+			return Agenda;
+		}
+	}
     public class ToolsNode : CompileNode
     {
         public ToolsNode(UnrealTargetPlatform InHostPlatform)
@@ -5622,10 +5652,11 @@ public class GUBP : BuildCommand
         
         if (HasNode(ToolsForCompileNode.StaticGetFullName(UnrealTargetPlatform.Win64)))
         {
-            if (HasNode(GamePlatformMonolithicsNode.StaticGetFullName(UnrealTargetPlatform.Mac, Branch.BaseEngineProject, UnrealTargetPlatform.IOS)) && HasNode(ToolsNode.StaticGetFullName(UnrealTargetPlatform.Win64)))
-            {
-                //AddNode(new IOSOnPCTestNode(this)); - Disable IOSOnPCTest until a1011 crash is fixed
-            }			
+			if (HasNode(GamePlatformMonolithicsNode.StaticGetFullName(UnrealTargetPlatform.Mac, Branch.BaseEngineProject, UnrealTargetPlatform.IOS)) && HasNode(ToolsNode.StaticGetFullName(UnrealTargetPlatform.Win64)))
+			{
+				//AddNode(new IOSOnPCTestNode(this)); - Disable IOSOnPCTest until a1011 crash is fixed
+			}
+			AddNode(new RootEditorCrossCompileLinuxNode(UnrealTargetPlatform.Win64));
             if (!bPreflightBuild)
             {
                 AddNode(new CleanSharedTempStorageNode(this));
