build script, turn off formalbuild agent sharing group for shared promotables

diff --git a/Engine/Source/Programs/AutomationTool/Scripts/GUBP.Automation.cs b/Engine/Source/Programs/AutomationTool/Scripts/GUBP.Automation.cs
--- a/Engine/Source/Programs/AutomationTool/Scripts/GUBP.Automation.cs
+++ b/Engine/Source/Programs/AutomationTool/Scripts/GUBP.Automation.cs
@@ -5238,7 +5238,7 @@ public class GUBP : BuildCommand
 
                 var GameTestNodes = new List<string>();
                 var GameCookNodes = new List<string>();
-				var FormalAgentSharingGroup = "Shared_FormalBuilds" + HostPlatformNode.StaticGetHostPlatformSuffix(HostPlatform);
+				//var FormalAgentSharingGroup = "Shared_FormalBuilds" + HostPlatformNode.StaticGetHostPlatformSuffix(HostPlatform);
 
 
                 //foreach (var Kind in BranchInfo.MonolithicKinds)//for now, non-code projects don't do client or server.
@@ -5330,7 +5330,7 @@ public class GUBP : BuildCommand
 													}
 													else
 													{
-														GUBPNodes[NodeName].AgentSharingGroup = FormalAgentSharingGroup;
+														//GUBPNodes[NodeName].AgentSharingGroup = FormalAgentSharingGroup;
 														if(Plat == UnrealTargetPlatform.XboxOne)
 														{
 															GUBPNodes[NodeName].AgentSharingGroup = "";
@@ -5435,7 +5435,7 @@ public class GUBP : BuildCommand
                 }
 
                 var CookedAgentSharingGroup = AgentShareName + "_CookedTests" + HostPlatformNode.StaticGetHostPlatformSuffix(HostPlatform);
-				var FormalAgentSharingGroup = "Shared_FormalBuilds" + HostPlatformNode.StaticGetHostPlatformSuffix(HostPlatform);
+				//var FormalAgentSharingGroup = "Shared_FormalBuilds" + HostPlatformNode.StaticGetHostPlatformSuffix(HostPlatform);
                 var ServerPlatforms = new List<UnrealTargetPlatform>();
                 var GamePlatforms = new List<UnrealTargetPlatform>();
                 var GameTestNodes = new List<string>();
@@ -5537,7 +5537,7 @@ public class GUBP : BuildCommand
 											}
 											else
 											{
-												GUBPNodes[FormalNodeName].AgentSharingGroup = FormalAgentSharingGroup;
+												//GUBPNodes[FormalNodeName].AgentSharingGroup = FormalAgentSharingGroup;
 												if (Plat == UnrealTargetPlatform.XboxOne)
 												{
 													GUBPNodes[FormalNodeName].AgentSharingGroup = "";
