Fix packaging on Linux (UEPLAT-184).

- The problem was caused by files that differ only in case (pak files are case-insensitive).
- Current solution is to ignore symlinks when staging files (so far we have been fixing case-sensitiveness issues by creating a symlink to the file, so this fixes the packaging too).
- Added a warning helping to pinpoint the problem if there happen to be two non-symlink files like this.

diff --git a/Engine/Source/Programs/AutomationTool/Scripts/CopyBuildToStagingDirectory.Automation.cs b/Engine/Source/Programs/AutomationTool/Scripts/CopyBuildToStagingDirectory.Automation.cs
--- a/Engine/Source/Programs/AutomationTool/Scripts/CopyBuildToStagingDirectory.Automation.cs
+++ b/Engine/Source/Programs/AutomationTool/Scripts/CopyBuildToStagingDirectory.Automation.cs
@@ -466,6 +466,11 @@ public partial class Project : CommandUtils
 
 			Dest = CombinePaths(PathSeparator.Slash, SC.PakFileInternalRoot, Dest);
 
+			// there can be files that only differ in case only, we don't support that in paks as paks are case-insensitive
+			if (UnrealPakResponseFile.ContainsKey(Src))
+			{
+				throw new AutomationException("Staging manifest already contains {0} (or a file that differs in case only)", Src);
+			}
 			UnrealPakResponseFile.Add(Src, Dest);
 		}
 		return UnrealPakResponseFile;
diff --git a/Engine/Source/Programs/AutomationTool/Utils.cs b/Engine/Source/Programs/AutomationTool/Utils.cs
--- a/Engine/Source/Programs/AutomationTool/Utils.cs
+++ b/Engine/Source/Programs/AutomationTool/Utils.cs
@@ -401,7 +401,25 @@ namespace AutomationTool
 			{
 				Log.WriteLine(TraceEventType.Information, "FindFiles {0} {1} {2}", Path, SearchPattern, Recursive);
 			}
-			return Directory.GetFiles(Path, SearchPattern, Recursive ? SearchOption.AllDirectories : SearchOption.TopDirectoryOnly);
+
+			// filter out symlinks which can cause problems at the later stage
+			List<string> FileNames = new List<string>();
+			DirectoryInfo DirInfo = new DirectoryInfo(Path);
+			foreach( FileInfo File in DirInfo.EnumerateFiles(SearchPattern, Recursive ? SearchOption.AllDirectories : SearchOption.TopDirectoryOnly))
+			{
+				if (File.Attributes.HasFlag(FileAttributes.ReparsePoint))
+				{
+					if (!bQuiet)
+					{
+						Log.WriteLine(TraceEventType.Warning, "Ignoring symlink {0}", File.FullName);
+					}
+					continue;
+				}
+				
+				FileNames.Add(File.FullName);
+			}
+			
+			return FileNames.ToArray();
 		}
 
 		/// <summary>
