Launch On fix for HTML5/Chrome. #ttp 349166
Chrome spawns a subprocess for html5 webpage, run command will now attach to that page and not kill the http server immediately.

diff --git a/Engine/Source/Programs/AutomationTool/HTML5/HTML5Platform.Automation.cs b/Engine/Source/Programs/AutomationTool/HTML5/HTML5Platform.Automation.cs
--- a/Engine/Source/Programs/AutomationTool/HTML5/HTML5Platform.Automation.cs
+++ b/Engine/Source/Programs/AutomationTool/HTML5/HTML5Platform.Automation.cs
@@ -9,6 +9,7 @@ using AutomationTool;
 using UnrealBuildTool;
 using System.Diagnostics;
 using System.Threading;
+using System.Management;
 
 public class HTML5Platform : Platform
 {
@@ -288,112 +289,153 @@ public class HTML5Platform : Platform
 
 	public override ProcessResult RunClient(ERunOptions ClientRunFlags, string ClientApp, string ClientCmdLine, ProjectParams Params)
 	{
-        // look for browser
-        var ConfigCache = new UnrealBuildTool.ConfigCacheIni(UnrealTargetPlatform.HTML5, "Engine", Path.GetDirectoryName(Params.RawProjectPath), CombinePaths(CmdEnv.LocalRoot, "Engine"));
+		// look for browser
+		var ConfigCache = new UnrealBuildTool.ConfigCacheIni(UnrealTargetPlatform.HTML5, "Engine", Path.GetDirectoryName(Params.RawProjectPath), CombinePaths(CmdEnv.LocalRoot, "Engine"));
 
-        string DeviceSection;
+		string DeviceSection;
 
-        if ( Utils.IsRunningOnMono )
-        {
-            DeviceSection = "HTML5DevicesMac";
-        }
-        else
-        {
-            DeviceSection = "HTML5DevicesWindows";
-        }
+		if ( Utils.IsRunningOnMono )
+		{
+			DeviceSection = "HTML5DevicesMac";
+		}
+		else
+		{
+			DeviceSection = "HTML5DevicesWindows";
+		}
 
-        string browserPath;
-        string DeviceName = Params.Device.Split('@')[1];
-        DeviceName = DeviceName.Substring(0, DeviceName.LastIndexOf(" on "));
-        bool ok = ConfigCache.GetString(DeviceSection, DeviceName, out browserPath);
+		string browserPath;
+		string DeviceName = Params.Device.Split('@')[1];
+		DeviceName = DeviceName.Substring(0, DeviceName.LastIndexOf(" on "));
+		bool ok = ConfigCache.GetString(DeviceSection, DeviceName, out browserPath);
 
 		if (!ok)
 			throw new System.Exception ("Incorrect browser configuration in HTML5Engine.ini ");
 
 		// open the webpage
-        string directory = Path.GetDirectoryName(ClientApp);
-        string url = Path.GetFileName(ClientApp) +".html";
+		string directory = Path.GetDirectoryName(ClientApp);
+		string url = Path.GetFileName(ClientApp) +".html";
 		// Are we running via cook on the fly server?
 		// find our http url - This is awkward because RunClient doesn't have real information that NFS is running or not.
 		bool IsCookOnTheFly = false;
 
-        // 9/24/2014 @fixme - All this is convoluted, clean up.
-        // looks like cookonthefly commandline stopped adding protocol or the port :/ hard coding to DEFAULT_TCP_FILE_SERVING_PORT+1 (DEFAULT_HTTP_FILE_SERVING_PORT)
-        // This will fail if the NFS server is started with a different port - we need to modify driver .cs script to pass in IP/Port data correctly. 
+		// 9/24/2014 @fixme - All this is convoluted, clean up.
+		// looks like cookonthefly commandline stopped adding protocol or the port :/ hard coding to DEFAULT_TCP_FILE_SERVING_PORT+1 (DEFAULT_HTTP_FILE_SERVING_PORT)
+		// This will fail if the NFS server is started with a different port - we need to modify driver .cs script to pass in IP/Port data correctly. 
 
-        if (ClientCmdLine.Contains("filehostip"))
-        {
-            IsCookOnTheFly = true; 
-            url = "http://127.0.0.1:41898/" + url; 
-        }
+		if (ClientCmdLine.Contains("filehostip"))
+		{
+			IsCookOnTheFly = true; 
+			url = "http://127.0.0.1:41898/" + url; 
+		}
 
-        if (IsCookOnTheFly)
-        {
-            url += "?cookonthefly=true";
-        }
-        else
-        {
-            var EmscriptenSettings = ReadEmscriptenSettings();
-            url = "http://127.0.0.1:8000/" + url;
-            // this will be killed UBT instances dies.
-            string input = String.Format(" -m SimpleHTTPServer 8000");
+		if (IsCookOnTheFly)
+		{
+			url += "?cookonthefly=true";
+		}
+		else
+		{
+			var EmscriptenSettings = ReadEmscriptenSettings();
+			url = "http://127.0.0.1:8000/" + url;
+			// this will be killed UBT instances dies.
+			string input = String.Format(" -m SimpleHTTPServer 8000");
 
 
 			string PythonName = null;
-            // Check the .emscripten file for a possible python path
-            if (EmscriptenSettings.ContainsKey("PYTHON"))
-            {
-                PythonName = EmscriptenSettings["PYTHON"];
-                Log("Found python path {0} in emscripten file", PythonName);
-            }
-            // The AutoSDK defines this env var as part of its install. See setup.bat/unsetup.bat
-            // If it's missing then just assume that python lives on the path
-            if (PythonName == null && Environment.GetEnvironmentVariable("PYTHON") != null)
-            {
-                PythonName = Environment.GetEnvironmentVariable("PYTHON");
-                Log("Found python path {0} in PYTHON Environment Variable", PythonName);
-            }
+			// Check the .emscripten file for a possible python path
+			if (EmscriptenSettings.ContainsKey("PYTHON"))
+			{
+				PythonName = EmscriptenSettings["PYTHON"];
+				Log("Found python path {0} in emscripten file", PythonName);
+			}
+			// The AutoSDK defines this env var as part of its install. See setup.bat/unsetup.bat
+			// If it's missing then just assume that python lives on the path
+			if (PythonName == null && Environment.GetEnvironmentVariable("PYTHON") != null)
+			{
+				PythonName = Environment.GetEnvironmentVariable("PYTHON");
+				Log("Found python path {0} in PYTHON Environment Variable", PythonName);
+			}
 
-            // Check if the exe exists
-            if (!System.IO.File.Exists(PythonName))
-            {
+			// Check if the exe exists
+			if (!System.IO.File.Exists(PythonName))
+			{
 				Log("Either no python path can be found or it doesn't exist. Using python on PATH");
-                PythonName = Utils.IsRunningOnMono ? "python" : "python.exe";
-            }
+				PythonName = Utils.IsRunningOnMono ? "python" : "python.exe";
+			}
 
 			ProcessResult Result = ProcessManager.CreateProcess(PythonName, true, "html5server.log");
 			Result.ProcessObject.StartInfo.FileName = PythonName;
-            Result.ProcessObject.StartInfo.UseShellExecute = false;
-            Result.ProcessObject.StartInfo.RedirectStandardOutput = true;
-            Result.ProcessObject.StartInfo.RedirectStandardInput = true;
-            Result.ProcessObject.StartInfo.WorkingDirectory = directory;
-            Result.ProcessObject.StartInfo.Arguments = input;
-            Result.ProcessObject.Start();
-
-
-            Result.ProcessObject.OutputDataReceived += delegate(object sender, System.Diagnostics.DataReceivedEventArgs e)
-            {
-                System.Console.WriteLine(e.Data);
-            };
+			Result.ProcessObject.StartInfo.UseShellExecute = false;
+			Result.ProcessObject.StartInfo.RedirectStandardOutput = true;
+			Result.ProcessObject.StartInfo.RedirectStandardInput = true;
+			Result.ProcessObject.StartInfo.WorkingDirectory = directory;
+			Result.ProcessObject.StartInfo.Arguments = input;
+			Result.ProcessObject.Start();
+
+			Result.ProcessObject.OutputDataReceived += delegate(object sender, System.Diagnostics.DataReceivedEventArgs e)
+			{
+				System.Console.WriteLine(e.Data);
+			};
 
-            System.Console.WriteLine("Starting Browser Process");
+			System.Console.WriteLine("Starting Browser Process");
 
 
 			// safari specific hack.
 			string argument = url; 
 			if (browserPath.Contains ("Safari") && Utils.IsRunningOnMono)
-				argument = ""; 
+				argument = "";
 
-            ProcessResult ClientProcess = Run(browserPath, argument, null, ClientRunFlags | ERunOptions.NoWaitForExit);
+			// Chrome issue. Firefox may work like this in the future
+			bool bBrowserWillSpawnProcess = browserPath.Contains("chrome");
 
-            ClientProcess.ProcessObject.EnableRaisingEvents = true;
-            ClientProcess.ProcessObject.Exited += delegate(System.Object o, System.EventArgs e)
-            {
-                System.Console.WriteLine("Browser Process Ended - Killing Webserver");
-                // send kill.
-                Result.ProcessObject.StandardInput.Close();
-                Result.ProcessObject.Kill();
-            };
+			ProcessResult SubProcess = null;
+			ProcessResult ClientProcess = Run(browserPath, argument, null, ClientRunFlags | ERunOptions.NoWaitForExit);
+			var ProcStartTime = ClientProcess.ProcessObject.StartTime;
+			var ProcName = ClientProcess.ProcessObject.ProcessName;
+
+			ClientProcess.ProcessObject.EnableRaisingEvents = true;
+			ClientProcess.ProcessObject.Exited += delegate(System.Object o, System.EventArgs e)
+			{
+				System.Console.WriteLine("Browser Process Ended (PID={0})", ClientProcess.ProcessObject.Id);
+				var bFoundChildProcess = true;
+				if (bBrowserWillSpawnProcess)
+				{
+					bFoundChildProcess = false;
+					// Chrome spawns a process from the tab it opens and then lets the process we spawned die, so
+					// catch that process and attach to that instead.
+					var CurrentProcesses = Process.GetProcesses();
+					foreach (var item in CurrentProcesses)
+					{
+						if (item.Id != ClientProcess.ProcessObject.Id && item.ProcessName == ProcName && item.StartTime >= ProcStartTime && item.StartTime <= ClientProcess.ProcessObject.ExitTime)
+						{
+							var PID = item.Id;
+							System.Console.WriteLine("Found Process {0} with PID {1} which started at {2}. Waiting on that process to end.", item.ProcessName, item.Id, item.StartTime.ToString());
+							SubProcess = new ProcessResult(item.ProcessName, item, true, item.ProcessName);
+							item.EnableRaisingEvents = true;
+							item.Exited += delegate(System.Object o2, System.EventArgs e2)
+							{
+								System.Console.WriteLine("Browser Process Ended (PID={0}) - Killing Webserver", PID);
+								Result.ProcessObject.StandardInput.Close();
+								Result.ProcessObject.Kill();
+							};
+							bFoundChildProcess = true;
+						}
+					}
+				}
+
+				if (!bFoundChildProcess)
+				{
+					System.Console.WriteLine("- Killing Webserver", ClientProcess.ProcessObject.Id);
+					Result.ProcessObject.StandardInput.Close();
+					Result.ProcessObject.Kill();
+				}
+			};
+
+			if (bBrowserWillSpawnProcess)
+			{
+				//Wait for it to do so...
+				ClientProcess.ProcessObject.WaitForExit();
+				ClientProcess = SubProcess;
+			}
 
 			// safari needs a hack. 
 			// http://superuser.com/questions/689315/run-safari-from-terminal-with-given-url-address-without-open-command
@@ -405,13 +447,13 @@ public class HTML5Platform : Platform
 				Process.Start("/usr/bin/osascript"," -e 'tell application \"Safari\" to open location \"" + url + "\"'" );
 			}
 
-            return ClientProcess;
-        }
+			return ClientProcess;
+		}
 
-        System.Console.WriteLine("Browser Path " + browserPath);
-        ProcessResult BrowserProcess = Run(browserPath, url, null, ClientRunFlags | ERunOptions.NoWaitForExit);
+		System.Console.WriteLine("Browser Path " + browserPath);
+		ProcessResult BrowserProcess = Run(browserPath, url, null, ClientRunFlags | ERunOptions.NoWaitForExit);
 	
-        return BrowserProcess;
+		return BrowserProcess;
 
 	}
 
