CIS fix

diff --git a/Engine/Source/Programs/AutomationTool/HTML5/HTML5Platform.Automation.cs b/Engine/Source/Programs/AutomationTool/HTML5/HTML5Platform.Automation.cs
--- a/Engine/Source/Programs/AutomationTool/HTML5/HTML5Platform.Automation.cs
+++ b/Engine/Source/Programs/AutomationTool/HTML5/HTML5Platform.Automation.cs
@@ -9,7 +9,6 @@ using AutomationTool;
 using UnrealBuildTool;
 using System.Diagnostics;
 using System.Threading;
-using System.Management;
 
 public class HTML5Platform : Platform
 {
