Limit reuse of patch data to 12 days for build machines.

This is to enable concurrent cleanup of builds to occur.

diff --git a/Engine/Source/Programs/AutomationTool/MCPPublic.cs b/Engine/Source/Programs/AutomationTool/MCPPublic.cs
--- a/Engine/Source/Programs/AutomationTool/MCPPublic.cs
+++ b/Engine/Source/Programs/AutomationTool/MCPPublic.cs
@@ -274,9 +274,7 @@ namespace EpicGames.MCP.Automation
 
             public PatchGenerationOptions()
             {
-				// @todo Richard.Fawcett temporary work-around for build system issues.
-				// Hack to add 10,000 days should not remain beyond 2014-09-30!
-				DataAgeThreshold = DEFAULT_DATA_AGE_THRESHOLD + 10000;
+                DataAgeThreshold = DEFAULT_DATA_AGE_THRESHOLD;
             }
 
             /// <summary>
diff --git a/Engine/Source/Programs/AutomationTool/Scripts/CopyBuildToStagingDirectory.Automation.cs b/Engine/Source/Programs/AutomationTool/Scripts/CopyBuildToStagingDirectory.Automation.cs
--- a/Engine/Source/Programs/AutomationTool/Scripts/CopyBuildToStagingDirectory.Automation.cs
+++ b/Engine/Source/Programs/AutomationTool/Scripts/CopyBuildToStagingDirectory.Automation.cs
@@ -587,7 +587,7 @@ public partial class Project : CommandUtils
 				string AppLaunch = ""; // !!JM todo, get real value for this...?
 				string VersionString = P4Branch + "-CL-" + P4Change;
 
-				string CmdLine = String.Format("-BuildRoot=\"{0}\" -CloudDir=\"{1}\" -AppID={2} -AppName=\"{3}\" -BuildVersion=\"{4}\" -AppLaunch=\"{5}\"", BuildRoot, CloudDir, AppID, AppName, VersionString, AppLaunch);
+				string CmdLine = String.Format("-BuildRoot=\"{0}\" -CloudDir=\"{1}\" -AppID={2} -AppName=\"{3}\" -BuildVersion=\"{4}\" -AppLaunch=\"{5}\" -DataAgeThreshold=12", BuildRoot, CloudDir, AppID, AppName, VersionString, AppLaunch);
 				CmdLine += " -AppArgs=\"\"";
 				CmdLine += " -custom=\"bIsPatch=false\"";
 				CmdLine += String.Format(" -customint=\"ChunkID={0}\"", ChunkID);
