Moved ICU data from Localization subdirectory to Internationalization subdirectory. Appropriate changes to loading and staging logic to reference the new location.

diff --git a/Engine/Source/Programs/AutomationTool/Scripts/CopyBuildToStagingDirectory.Automation.cs b/Engine/Source/Programs/AutomationTool/Scripts/CopyBuildToStagingDirectory.Automation.cs
--- a/Engine/Source/Programs/AutomationTool/Scripts/CopyBuildToStagingDirectory.Automation.cs
+++ b/Engine/Source/Programs/AutomationTool/Scripts/CopyBuildToStagingDirectory.Automation.cs
@@ -329,7 +329,7 @@ public partial class Project : CommandUtils
 				SC.StageFiles(StagedFileType.NonUFS, CombinePaths(SC.LocalRoot, "Engine/Content/Slate"));
 				SC.StageFiles(StagedFileType.NonUFS, CombinePaths(SC.LocalRoot, "Engine/Shaders/StandaloneRenderer"));
 
-				SC.StageFiles(StagedFileType.NonUFS, CombinePaths(SC.LocalRoot, "Engine/Content/Localization/ICU"));
+				SC.StageFiles(StagedFileType.NonUFS, CombinePaths(SC.LocalRoot, "Engine/Content/Internationalization"));
 				// Linux platform stages ICU in GetFilesToDeployOrStage(), accounting for the actual architecture
 				if (SC.StageTargetPlatform.PlatformType == UnrealTargetPlatform.Win64 ||
 					SC.StageTargetPlatform.PlatformType == UnrealTargetPlatform.Win32 ||
diff --git a/Engine/Source/Runtime/Core/Private/Internationalization/ICUInternationalization.cpp b/Engine/Source/Runtime/Core/Private/Internationalization/ICUInternationalization.cpp
--- a/Engine/Source/Runtime/Core/Private/Internationalization/ICUInternationalization.cpp
+++ b/Engine/Source/Runtime/Core/Private/Internationalization/ICUInternationalization.cpp
@@ -82,7 +82,7 @@ bool FICUInternationalization::Initialize()
 
 	u_setMemoryFunctions(NULL, &(FICUOverrides::Malloc), &(FICUOverrides::Realloc), &(FICUOverrides::Free), &(ICUStatus));
 
-	const FString DataDirectoryRelativeToContent = FString(TEXT("Localization")) / FString(TEXT("ICU"));
+	const FString DataDirectoryRelativeToContent = FString(TEXT("Internationalization"));
 	IFileManager& FileManager = IFileManager::Get();
 
 	const FString PotentialDataDirectories[] = 
