Modified staging logic for ICU data to only stage default ICU data.

diff --git a/Engine/Source/Programs/AutomationTool/Scripts/CopyBuildToStagingDirectory.Automation.cs b/Engine/Source/Programs/AutomationTool/Scripts/CopyBuildToStagingDirectory.Automation.cs
--- a/Engine/Source/Programs/AutomationTool/Scripts/CopyBuildToStagingDirectory.Automation.cs
+++ b/Engine/Source/Programs/AutomationTool/Scripts/CopyBuildToStagingDirectory.Automation.cs
@@ -329,7 +329,7 @@ public partial class Project : CommandUtils
 				SC.StageFiles(StagedFileType.NonUFS, CombinePaths(SC.LocalRoot, "Engine/Content/Slate"));
 				SC.StageFiles(StagedFileType.NonUFS, CombinePaths(SC.LocalRoot, "Engine/Shaders/StandaloneRenderer"));
 
-				SC.StageFiles(StagedFileType.NonUFS, CombinePaths(SC.LocalRoot, "Engine/Content/Internationalization"));
+				SC.StageFiles(StagedFileType.NonUFS, CombinePaths(SC.LocalRoot, "Engine/Content/Internationalization/icudt53l"));
 				// Linux platform stages ICU in GetFilesToDeployOrStage(), accounting for the actual architecture
 				if (SC.StageTargetPlatform.PlatformType == UnrealTargetPlatform.Win64 ||
 					SC.StageTargetPlatform.PlatformType == UnrealTargetPlatform.Win32 ||
