Corrected platform-specific ICU data staging logic that was using the old paths.

diff --git a/Engine/Source/Programs/AutomationTool/Android/AndroidPlatform.Automation.cs b/Engine/Source/Programs/AutomationTool/Android/AndroidPlatform.Automation.cs
--- a/Engine/Source/Programs/AutomationTool/Android/AndroidPlatform.Automation.cs
+++ b/Engine/Source/Programs/AutomationTool/Android/AndroidPlatform.Automation.cs
@@ -770,7 +770,7 @@ public class AndroidPlatform : Platform
 
 		if (!Params.CookOnTheFly && !Params.SkipCookOnTheFly) // only stage the UFS files if we are not using cook on the fly
 		{
-			SC.StageFiles(StagedFileType.UFS, CombinePaths(SC.LocalRoot, "Engine/Content/Localization/ICU"), "*", true, null, null, false, !Params.UsePak(SC.StageTargetPlatform));
+            SC.StageFiles(StagedFileType.UFS, CombinePaths(SC.LocalRoot, "Engine/Content/Internationalization/icudt53l"), "*", true, null, null, false, !Params.UsePak(SC.StageTargetPlatform));
 		}
 	}
 
diff --git a/Engine/Source/Programs/AutomationTool/HTML5/HTML5Platform.Automation.cs b/Engine/Source/Programs/AutomationTool/HTML5/HTML5Platform.Automation.cs
--- a/Engine/Source/Programs/AutomationTool/HTML5/HTML5Platform.Automation.cs
+++ b/Engine/Source/Programs/AutomationTool/HTML5/HTML5Platform.Automation.cs
@@ -252,7 +252,7 @@ public class HTML5Platform : Platform
 	public override void GetFilesToDeployOrStage(ProjectParams Params, DeploymentContext SC)
 	{
         // stage required ICU files.
-        SC.StageFiles(StagedFileType.UFS, CombinePaths(SC.LocalRoot, "Engine/Content/Localization/ICU"), "*", true, null, null, false, !Params.Pak);
+        SC.StageFiles(StagedFileType.UFS, CombinePaths(SC.LocalRoot, "Engine/Content/Internationalization/icudt53l"), "*", true, null, null, false, !Params.Pak);
 	}
 
 	public override void GetFilesToArchive(ProjectParams Params, DeploymentContext SC)
diff --git a/Engine/Source/Programs/AutomationTool/IOS/IOSPlatform.Automation.cs b/Engine/Source/Programs/AutomationTool/IOS/IOSPlatform.Automation.cs
--- a/Engine/Source/Programs/AutomationTool/IOS/IOSPlatform.Automation.cs
+++ b/Engine/Source/Programs/AutomationTool/IOS/IOSPlatform.Automation.cs
@@ -731,7 +731,7 @@ public class IOSPlatform : Platform
 		}
 
 		// stage required icu files
-		SC.StageFiles(StagedFileType.UFS, CombinePaths(SC.LocalRoot, "Engine/Content/Localization/ICU"), "*", true, null, null, false, !Params.UsePak(SC.StageTargetPlatform));
+        SC.StageFiles(StagedFileType.UFS, CombinePaths(SC.LocalRoot, "Engine/Content/Internationalization/icudt53l"), "*", true, null, null, false, !Params.UsePak(SC.StageTargetPlatform));
 	}
 
 	public override void GetFilesToArchive(ProjectParams Params, DeploymentContext SC)
diff --git a/Engine/Source/Programs/AutomationTool/Linux/LinuxPlatform.Automation.cs b/Engine/Source/Programs/AutomationTool/Linux/LinuxPlatform.Automation.cs
--- a/Engine/Source/Programs/AutomationTool/Linux/LinuxPlatform.Automation.cs
+++ b/Engine/Source/Programs/AutomationTool/Linux/LinuxPlatform.Automation.cs
@@ -55,7 +55,7 @@ public abstract class BaseLinuxPlatform : Platform
 		}
 
 		SC.StageFiles(StagedFileType.NonUFS, CombinePaths(SC.ProjectRoot, "Content/Splash"), "Splash.bmp", false, null, null, true);
-		SC.StageFiles(StagedFileType.UFS, CombinePaths(SC.LocalRoot, "Engine/Content/Localization/ICU"), "*", true, null, null, false, !Params.UsePak(SC.StageTargetPlatform));
+        SC.StageFiles(StagedFileType.UFS, CombinePaths(SC.LocalRoot, "Engine/Content/Internationalization/icudt53l"), "*", true, null, null, false, !Params.UsePak(SC.StageTargetPlatform));
 
 
         if (Params.StageNonMonolithic)
diff --git a/Engine/Source/Programs/AutomationTool/Mac/MacPlatform.Automation.cs b/Engine/Source/Programs/AutomationTool/Mac/MacPlatform.Automation.cs
--- a/Engine/Source/Programs/AutomationTool/Mac/MacPlatform.Automation.cs
+++ b/Engine/Source/Programs/AutomationTool/Mac/MacPlatform.Automation.cs
@@ -49,7 +49,7 @@ public class MacPlatform : Platform
 		// Copy the splash screen, Mac specific
 		SC.StageFiles(StagedFileType.NonUFS, CombinePaths(SC.ProjectRoot, "Content/Splash"), "Splash.bmp", false, null, null, true);
 
-		SC.StageFiles(StagedFileType.UFS, CombinePaths(SC.LocalRoot, "Engine/Content/Localization/ICU"), "*", true, null, null, false, !Params.UsePak(SC.StageTargetPlatform));
+        SC.StageFiles(StagedFileType.UFS, CombinePaths(SC.LocalRoot, "Engine/Content/Internationalization/icudt53l"), "*", true, null, null, false, !Params.UsePak(SC.StageTargetPlatform));
 	}
 
 	public override void Package(ProjectParams Params, DeploymentContext SC, int WorkingCL)
diff --git a/Engine/Source/Programs/AutomationTool/Win/WinPlatform.Automation.cs b/Engine/Source/Programs/AutomationTool/Win/WinPlatform.Automation.cs
--- a/Engine/Source/Programs/AutomationTool/Win/WinPlatform.Automation.cs
+++ b/Engine/Source/Programs/AutomationTool/Win/WinPlatform.Automation.cs
@@ -88,7 +88,7 @@ public abstract class BaseWinPlatform : Platform
 		// Copy the splash screen, windows specific
 		SC.StageFiles(StagedFileType.NonUFS, CombinePaths(SC.ProjectRoot, "Content/Splash"), "Splash.bmp", false, null, null, true);
 
-		SC.StageFiles(StagedFileType.UFS, CombinePaths(SC.LocalRoot, "Engine/Content/Localization/ICU"), "*", true, null, null, false, !Params.UsePak(SC.StageTargetPlatform));
+        SC.StageFiles(StagedFileType.UFS, CombinePaths(SC.LocalRoot, "Engine/Content/Internationalization/icudt53l"), "*", true, null, null, false, !Params.UsePak(SC.StageTargetPlatform));
 
         if (Params.StageNonMonolithic)
         {
