Mac fix for launching html5 on chrome

diff --git a/Engine/Source/Programs/AutomationTool/HTML5/HTML5Platform.Automation.cs b/Engine/Source/Programs/AutomationTool/HTML5/HTML5Platform.Automation.cs
--- a/Engine/Source/Programs/AutomationTool/HTML5/HTML5Platform.Automation.cs
+++ b/Engine/Source/Programs/AutomationTool/HTML5/HTML5Platform.Automation.cs
@@ -384,7 +384,7 @@ public class HTML5Platform : Platform
 				argument = "";
 
 			// Chrome issue. Firefox may work like this in the future
-			bool bBrowserWillSpawnProcess = browserPath.Contains("chrome");
+			bool bBrowserWillSpawnProcess = browserPath.Contains("chrome") || (browserPath.Contains("Google Chrome") && Utils.IsRunningOnMono);
 
 			ProcessResult SubProcess = null;
 			ProcessResult ClientProcess = Run(browserPath, argument, null, ClientRunFlags | ERunOptions.NoWaitForExit);
