- Added parameter to UAT "-archivemetadata" which will copy build.properties to the archive directory [UE-2228]

diff --git a/Engine/Source/Programs/AutomationTool/ProjectParams.cs b/Engine/Source/Programs/AutomationTool/ProjectParams.cs
--- a/Engine/Source/Programs/AutomationTool/ProjectParams.cs
+++ b/Engine/Source/Programs/AutomationTool/ProjectParams.cs
@@ -269,6 +269,7 @@ namespace AutomationTool
             this.bDebugBuildsActuallyUseDebugCRT = InParams.bDebugBuildsActuallyUseDebugCRT;
 			this.Archive = InParams.Archive;
 			this.ArchiveDirectoryParam = InParams.ArchiveDirectoryParam;
+			this.ArchiveMetaData = InParams.ArchiveMetaData;
 			this.Distribution = InParams.Distribution;
 			this.Prereqs = InParams.Prereqs;
 			this.NoBootstrapExe = InParams.NoBootstrapExe;
@@ -354,6 +355,7 @@ namespace AutomationTool
 			int? NumClients = null,
 			bool? Archive = null,
 			string ArchiveDirectoryParam = null,
+			bool? ArchiveMetaData = null,
 			ParamList<string> ProgramTargets = null,
 			bool? Distribution = null,
             bool? OBBinAPK = null,
@@ -454,6 +456,7 @@ namespace AutomationTool
             this.CreateChunkInstall = GetParamValueIfNotSpecified(Command, CreateChunkInstall, this.CreateChunkInstall, "createchunkinstall");
 			this.Archive = GetParamValueIfNotSpecified(Command, Archive, this.Archive, "archive");
 			this.ArchiveDirectoryParam = ParseParamValueIfNotSpecified(Command, ArchiveDirectoryParam, "archivedirectory", String.Empty);
+			this.ArchiveMetaData = GetParamValueIfNotSpecified(Command, ArchiveMetaData, this.ArchiveMetaData, "archivemetadata");
 			this.Distribution = GetParamValueIfNotSpecified(Command, Distribution, this.Distribution, "distribution");
 			this.Prereqs = GetParamValueIfNotSpecified(Command, Prereqs, this.Prereqs, "prereqs");
 			this.NoBootstrapExe = GetParamValueIfNotSpecified(Command, NoBootstrapExe, this.NoBootstrapExe, "nobootstrapexe");
@@ -778,6 +781,12 @@ namespace AutomationTool
 		[Help("archivedirectory=Path", "Directory to archive the builds to, i.e. -archivedirectory=C:\\Archive")]
 		public string ArchiveDirectoryParam;
 
+		/// <summary>
+		/// Whether the project should use non monolithic staging
+		/// </summary>
+		[Help("archivemetadata", "Archive extra metadata files in addition to the build (e.g. build.properties)")]
+		public bool ArchiveMetaData;
+
 		#endregion
 
 		#region Build
@@ -1699,6 +1708,7 @@ namespace AutomationTool
 
 				CommandUtils.Log("AdditionalServerMapParams={0}", AdditionalServerMapParams);
 				CommandUtils.Log("Archive={0}", Archive);
+				CommandUtils.Log("ArchiveMetaData={0}", ArchiveMetaData);
 				CommandUtils.Log("BaseArchiveDirectory={0}", BaseArchiveDirectory);
 				CommandUtils.Log("BaseStageDirectory={0}", BaseStageDirectory);
 				CommandUtils.Log("Build={0}", Build);
diff --git a/Engine/Source/Programs/AutomationTool/Scripts/ArchiveCommand.Automation.cs b/Engine/Source/Programs/AutomationTool/Scripts/ArchiveCommand.Automation.cs
--- a/Engine/Source/Programs/AutomationTool/Scripts/ArchiveCommand.Automation.cs
+++ b/Engine/Source/Programs/AutomationTool/Scripts/ArchiveCommand.Automation.cs
@@ -23,6 +23,13 @@ public partial class Project : CommandUtils
 		ThisPlatform.GetFilesToArchive(Params, SC);
 
 		//@todo add any archive meta data files as needed
+
+		if (Params.ArchiveMetaData)
+		{
+			// archive the build.properties for extra info for testing, etc
+			string BuildPropertiesFile = CombinePaths(SC.LocalRoot, "Engine", "Build", "build.properties");
+			SC.ArchiveFiles(Path.GetDirectoryName(BuildPropertiesFile), Path.GetFileName(BuildPropertiesFile));
+		}
 	}
 
 	public static void ApplyArchiveManifest(ProjectParams Params, DeploymentContext SC)
