build script, add depended on by properties

diff --git a/Engine/Source/Programs/AutomationTool/Scripts/GUBP.Automation.cs b/Engine/Source/Programs/AutomationTool/Scripts/GUBP.Automation.cs
--- a/Engine/Source/Programs/AutomationTool/Scripts/GUBP.Automation.cs
+++ b/Engine/Source/Programs/AutomationTool/Scripts/GUBP.Automation.cs
@@ -357,6 +357,7 @@ public class GUBP : BuildCommand
     {
         public List<string> FullNamesOfDependencies = new List<string>();
         public List<string> FullNamesOfPseudosependencies = new List<string>(); //these are really only used for sorting. We want the editor to fail before the monolithics. Think of it as "can't possibly be useful without".
+		public List<string> FullNamesOfDependedOn = new List<string>();
         public List<string> BuildProducts = null;
         public List<string> AllDependencyBuildProducts = null;
         public List<string> AllDependencies = null;
@@ -3769,6 +3770,7 @@ public class GUBP : BuildCommand
         bool bShowChanges = (bp.ParseParam("Changes") && GUBPNodesHistory != null) || bShowAllChanges;
         bool bShowDetailedHistory = (bp.ParseParam("History") && GUBPNodesHistory != null) || bShowChanges;
         bool bShowDependencies = bp.ParseParam("ShowDependencies");
+		bool bShowDependednOn = bp.ParseParam("ShowDependedOn");
         bool bShowECDependencies = bp.ParseParam("ShowECDependencies");
         bool bShowHistory = !bp.ParseParam("NoHistory") && GUBPNodesHistory != null;
         bool AddEmailProps = bp.ParseParam("ShowEmails");
@@ -3877,6 +3879,13 @@ public class GUBP : BuildCommand
                     Log("           {0}", Dep);
                 }
             }
+			if(bShowDependednOn)
+			{
+				foreach (var Dep in GUBPNodes[NodeToDo].FullNamesOfDependedOn)
+				{
+					Log("            depOn> {0}", Dep);
+				}
+			}
         }
     }
     public void SaveGraphVisualization(List<string> Nodes)
@@ -5718,6 +5727,7 @@ public class GUBP : BuildCommand
         var FullNodeList = new Dictionary<string, string>();
         var FullNodeListSortKey = new Dictionary<string, int>();
         var FullNodeDirectDependencies = new Dictionary<string, string>();
+		var FullNodeDependedOnBy = new Dictionary<string, string>();
         {
             Log("******* {0} GUBP Nodes", GUBPNodes.Count);
             var SortedNodes = TopologicalSort(new HashSet<string>(GUBPNodes.Keys), LocalOnly: true, DoNotConsiderCompletion: true);
@@ -5992,24 +6002,61 @@ public class GUBP : BuildCommand
                 NodesToDo.UnionWith(Fringe);
             }
         }
-        if (TimeIndex != 0)
-        {
-            Log("Culling based on time index");
-            var NewNodesToDo = new HashSet<string>();
-            foreach (var NodeToDo in NodesToDo)
-            {
-                if (TimeIndex % (1 << GUBPNodes[NodeToDo].DependentCISFrequencyQuantumShift()) == 0)
-                {
-                    Log("  Keeping {0}", NodeToDo);
-                    NewNodesToDo.Add(NodeToDo);
-                }
-                else
-                {
-                    Log("  Rejecting {0}", NodeToDo);
-                }
-            }
-            NodesToDo = NewNodesToDo;
-        }
+		if (TimeIndex != 0)
+		{
+			Log("Culling based on time index");
+			var NewNodesToDo = new HashSet<string>();
+			foreach (var NodeToDo in NodesToDo)
+			{
+				if (TimeIndex % (1 << GUBPNodes[NodeToDo].DependentCISFrequencyQuantumShift()) == 0)
+				{
+					Log("  Keeping {0}", NodeToDo);
+					NewNodesToDo.Add(NodeToDo);
+				}
+				else
+				{
+					Log("  Rejecting {0}", NodeToDo);
+				}
+			}
+			NodesToDo = NewNodesToDo;
+		}
+		//find things that depend on our nodes and setup commander dictionary
+		if (!bOnlyNode)
+		{			
+			foreach(var NodeToDo in NodesToDo)
+			{
+				if (!GUBPNodes[NodeToDo].IsAggregate() && !GUBPNodes[NodeToDo].IsTest())
+				{
+					List<string> ECDependencies = new List<string>();
+					ECDependencies = GetECDependencies(NodeToDo);
+					foreach (var Dep in ECDependencies)
+					{
+						if (!GUBPNodes.ContainsKey(Dep))
+						{
+							throw new AutomationException("Node {0} is not in the graph. It is a dependency of {1}.", Dep, NodeToDo);
+						}
+						if (!GUBPNodes[Dep].FullNamesOfDependedOn.Contains(NodeToDo))
+						{
+							GUBPNodes[Dep].FullNamesOfDependedOn.Add(NodeToDo);
+						}						
+					}
+				}				
+			}	
+			foreach(var NodeToDo in NodesToDo)
+			{
+				var Deps = GUBPNodes[NodeToDo].FullNamesOfDependedOn;
+				string All = "";
+				foreach (var Dep in Deps)
+				{
+					if (All != "")
+					{
+						All += " ";
+					}
+					All += Dep;
+				}
+				FullNodeDependedOnBy.Add(NodeToDo, All);
+			}
+		}
 
         if (CommanderSetup)
         {
@@ -6171,7 +6218,11 @@ public class GUBP : BuildCommand
             foreach (var NodePair in FullNodeListSortKey)
             {
                 ECProps.Add(string.Format("SortKey/{0}={1}", NodePair.Key, NodePair.Value));
-            }            
+            }
+			foreach (var NodePair in FullNodeDependedOnBy)
+			{
+				ECProps.Add(string.Format("DependedOnBy/{0}={1}", NodePair.Key, NodePair.Value));
+			}
             var ECJobProps = new List<string>();
             if (ExplicitTrigger != "")
             {
