Add GenerateChunks (-manifests) support to Packaging Settings.  Read it properly from .ini when setting up ProjectParams.

diff --git a/Engine/Config/BaseGame.ini b/Engine/Config/BaseGame.ini
--- a/Engine/Config/BaseGame.ini
+++ b/Engine/Config/BaseGame.ini
@@ -36,6 +36,7 @@ SupportContact=
 BuildConfiguration=PPBC_Development
 FullRebuild=True
 UsePakFile=True
+bGenerateChunks=False
 IncludePrerequisites=True
 UseOBB_InAPK=False
 +CulturesToStage=en
diff --git a/Engine/Source/Editor/UnrealEd/Classes/Settings/ProjectPackagingSettings.h b/Engine/Source/Editor/UnrealEd/Classes/Settings/ProjectPackagingSettings.h
--- a/Engine/Source/Editor/UnrealEd/Classes/Settings/ProjectPackagingSettings.h
+++ b/Engine/Source/Editor/UnrealEd/Classes/Settings/ProjectPackagingSettings.h
@@ -61,6 +61,13 @@ public:
 	UPROPERTY(config, EditAnywhere, Category=Packaging)
 	bool UsePakFile;
 
+	/** 
+	 * If enabled, will generate pak file chunks.  Assets can be assigned to chunks in the editor or via a delegate (See ShooterGameDelegates.cpp). 
+	 * Can be used for streaming installs (PS4 Playgo, XboxOne Streaming Install, etc)
+	 */
+	UPROPERTY(config, EditAnywhere, Category=Packaging)
+	bool bGenerateChunks;
+
 	/** Specifies whether to include prerequisites of packaged games, such as redistributable operating system components, whenever possible. */
 	UPROPERTY(config, EditAnywhere, Category=Packaging)
 	bool IncludePrerequisites;
@@ -97,7 +104,7 @@ public:
 	 * Note: These paths are relative to your project Content directory
 	 */
 	UPROPERTY(config, EditAnywhere, Category=Packaging, AdvancedDisplay, meta=(DisplayName="Additional Non-Asset Directories To Copy", RelativeToGameContentDir))
-	TArray<FDirectoryPath> DirectoriesToAlwaysStageAsNonUFS;
+	TArray<FDirectoryPath> DirectoriesToAlwaysStageAsNonUFS;	
 
 public:
 
diff --git a/Engine/Source/Programs/AutomationTool/ProjectParams.cs b/Engine/Source/Programs/AutomationTool/ProjectParams.cs
--- a/Engine/Source/Programs/AutomationTool/ProjectParams.cs
+++ b/Engine/Source/Programs/AutomationTool/ProjectParams.cs
@@ -108,10 +108,10 @@ namespace AutomationTool
 		/// <summary>
 		/// Sets up platforms
 		/// </summary>
-        /// <param name="DependentPlatformMap">Set with a mapping from source->destination if specified on command line</param>
+		/// <param name="DependentPlatformMap">Set with a mapping from source->destination if specified on command line</param>
 		/// <param name="Command">The command line to parse</param>
 		/// <param name="OverrideTargetPlatforms">If passed use this always</param>
-        /// <param name="DefaultTargetPlatforms">Use this if nothing is passed on command line</param>
+		/// <param name="DefaultTargetPlatforms">Use this if nothing is passed on command line</param>
 		/// <param name="AllowPlatformParams">Allow raw -platform options</param>
 		/// <param name="PlatformParamNames">Possible -parameters to check for</param>
 		/// <returns>List of platforms parsed from the command line</returns>
@@ -143,20 +143,20 @@ namespace AutomationTool
 						var Platforms = new List<string>(CmdLinePlatform.Split('+'));
 						foreach (var PlatformName in Platforms)
 						{
-                            // Look for dependent platforms, Source_1.Dependent_1+Source_2.Dependent_2+Standalone_3
-                            var SubPlatforms = new List<string>(PlatformName.Split('.'));
-
-                            foreach (var SubPlatformName in SubPlatforms)
-                            {
-                                UnrealTargetPlatform NewPlatform = (UnrealTargetPlatform)Enum.Parse(typeof(UnrealTargetPlatform), SubPlatformName, true);
-                                TargetPlatforms.Add(NewPlatform);
-
-                                if (SubPlatformName != SubPlatforms[0])
-                                {
-                                    // We're a dependent platform so add ourselves to the map, pointing to the first element in the list
-                                    DependentPlatformMap.Add(NewPlatform, (UnrealTargetPlatform)Enum.Parse(typeof(UnrealTargetPlatform), SubPlatforms[0], true));
-                                }
-                            }
+							// Look for dependent platforms, Source_1.Dependent_1+Source_2.Dependent_2+Standalone_3
+							var SubPlatforms = new List<string>(PlatformName.Split('.'));
+
+							foreach (var SubPlatformName in SubPlatforms)
+							{
+								UnrealTargetPlatform NewPlatform = (UnrealTargetPlatform)Enum.Parse(typeof(UnrealTargetPlatform), SubPlatformName, true);
+								TargetPlatforms.Add(NewPlatform);
+
+								if (SubPlatformName != SubPlatforms[0])
+								{
+									// We're a dependent platform so add ourselves to the map, pointing to the first element in the list
+									DependentPlatformMap.Add(NewPlatform, (UnrealTargetPlatform)Enum.Parse(typeof(UnrealTargetPlatform), SubPlatforms[0], true));
+								}
+							}
 						}
 					}
 					else if (AllowPlatformParams)
@@ -194,15 +194,15 @@ namespace AutomationTool
 			this.RawProjectPath = InParams.RawProjectPath;
 			this.MapsToCook = InParams.MapsToCook;
 			this.DirectoriesToCook = InParams.DirectoriesToCook;
-            this.CulturesToCook = InParams.CulturesToCook;
+			this.CulturesToCook = InParams.CulturesToCook;
 			this.ClientCookedTargets = InParams.ClientCookedTargets;
 			this.ServerCookedTargets = InParams.ServerCookedTargets;
 			this.EditorTargets = InParams.EditorTargets;
 			this.ProgramTargets = InParams.ProgramTargets;
 			this.ClientTargetPlatforms = InParams.ClientTargetPlatforms;
-            this.ClientDependentPlatformMap = InParams.ClientDependentPlatformMap;
+			this.ClientDependentPlatformMap = InParams.ClientDependentPlatformMap;
 			this.ServerTargetPlatforms = InParams.ServerTargetPlatforms;
-            this.ServerDependentPlatformMap = InParams.ServerDependentPlatformMap;
+			this.ServerDependentPlatformMap = InParams.ServerDependentPlatformMap;
 			this.Build = InParams.Build;
 			this.Run = InParams.Run;
 			this.Cook = InParams.Cook;
@@ -210,8 +210,8 @@ namespace AutomationTool
 			this.CookFlavor = InParams.CookFlavor;
 			this.SkipCook = InParams.SkipCook;
 			this.SkipCookOnTheFly = InParams.SkipCookOnTheFly;
-            this.Prebuilt = InParams.Prebuilt;
-            this.RunTimeoutSeconds = InParams.RunTimeoutSeconds;
+			this.Prebuilt = InParams.Prebuilt;
+			this.RunTimeoutSeconds = InParams.RunTimeoutSeconds;
 			this.Clean = InParams.Clean;
 			this.Pak = InParams.Pak;
 			this.SignPak = InParams.SignPak;
@@ -226,10 +226,10 @@ namespace AutomationTool
 			this.LogWindow = InParams.LogWindow;
 			this.Stage = InParams.Stage;
 			this.SkipStage = InParams.SkipStage;
-            this.StageDirectoryParam = InParams.StageDirectoryParam;
-            this.StageNonMonolithic = InParams.StageNonMonolithic;
+			this.StageDirectoryParam = InParams.StageDirectoryParam;
+			this.StageNonMonolithic = InParams.StageNonMonolithic;
 			this.Manifests = InParams.Manifests;
-            this.CreateChunkInstall = InParams.CreateChunkInstall;
+			this.CreateChunkInstall = InParams.CreateChunkInstall;
 			this.UE4Exe = InParams.UE4Exe;
 			this.NoDebugInfo = InParams.NoDebugInfo;
 			this.NoCleanStage = InParams.NoCleanStage;
@@ -238,44 +238,44 @@ namespace AutomationTool
 			this.Foreign = InParams.Foreign;
 			this.ForeignCode = InParams.ForeignCode;
 			this.StageCommandline = InParams.StageCommandline;
-            this.BundleName = InParams.BundleName;
+			this.BundleName = InParams.BundleName;
 			this.RunCommandline = InParams.RunCommandline;
 			this.Package = InParams.Package;
 			this.Deploy = InParams.Deploy;
 			this.Device = InParams.Device;
 			this.ServerDevice = InParams.ServerDevice;
-            this.NullRHI = InParams.NullRHI;
-            this.FakeClient = InParams.FakeClient;
-            this.EditorTest = InParams.EditorTest;
-            this.RunAutomationTests = InParams.RunAutomationTests;
-            this.RunAutomationTest = InParams.RunAutomationTest;
-            this.CrashIndex = InParams.CrashIndex;
-            this.Port = InParams.Port;
+			this.NullRHI = InParams.NullRHI;
+			this.FakeClient = InParams.FakeClient;
+			this.EditorTest = InParams.EditorTest;
+			this.RunAutomationTests = InParams.RunAutomationTests;
+			this.RunAutomationTest = InParams.RunAutomationTest;
+			this.CrashIndex = InParams.CrashIndex;
+			this.Port = InParams.Port;
 			this.SkipServer = InParams.SkipServer;
 			this.Rocket = InParams.Rocket;
 			this.Unattended = InParams.Unattended;
-            this.ServerDeviceAddress = InParams.ServerDeviceAddress;
-            this.DeviceUsername = InParams.DeviceUsername;
-            this.DevicePassword = InParams.DevicePassword;
-            this.CrashReporter = InParams.CrashReporter;
+			this.ServerDeviceAddress = InParams.ServerDeviceAddress;
+			this.DeviceUsername = InParams.DeviceUsername;
+			this.DevicePassword = InParams.DevicePassword;
+			this.CrashReporter = InParams.CrashReporter;
 			this.ClientConfigsToBuild = InParams.ClientConfigsToBuild;
 			this.ServerConfigsToBuild = InParams.ServerConfigsToBuild;
 			this.NumClients = InParams.NumClients;
-            this.Compressed = InParams.Compressed;
-            this.UseDebugParamForEditorExe = InParams.UseDebugParamForEditorExe;
-            this.bUsesSteam = InParams.bUsesSteam;
+			this.Compressed = InParams.Compressed;
+			this.UseDebugParamForEditorExe = InParams.UseDebugParamForEditorExe;
+			this.bUsesSteam = InParams.bUsesSteam;
 			this.bUsesSlate = InParams.bUsesSlate;
 			this.bUsesSlateEditorStyle = InParams.bUsesSlateEditorStyle;
-            this.bDebugBuildsActuallyUseDebugCRT = InParams.bDebugBuildsActuallyUseDebugCRT;
+			this.bDebugBuildsActuallyUseDebugCRT = InParams.bDebugBuildsActuallyUseDebugCRT;
 			this.Archive = InParams.Archive;
 			this.ArchiveDirectoryParam = InParams.ArchiveDirectoryParam;
 			this.ArchiveMetaData = InParams.ArchiveMetaData;
 			this.Distribution = InParams.Distribution;
 			this.Prereqs = InParams.Prereqs;
 			this.NoBootstrapExe = InParams.NoBootstrapExe;
-            this.OBBinAPK = InParams.OBBinAPK;
-            this.Prebuilt = InParams.Prebuilt;
-            this.RunTimeoutSeconds = InParams.RunTimeoutSeconds;
+			this.OBBinAPK = InParams.OBBinAPK;
+			this.Prebuilt = InParams.Prebuilt;
+			this.RunTimeoutSeconds = InParams.RunTimeoutSeconds;
 		}
 
 		/// <summary>
@@ -293,32 +293,32 @@ namespace AutomationTool
 			ParamList<string> Port = null,
 			string RunCommandline = null,						
 			string StageCommandline = null,
-            string BundleName = null,
-            string StageDirectoryParam = null,
-            bool? StageNonMonolithic = null,
+			string BundleName = null,
+			string StageDirectoryParam = null,
+			bool? StageNonMonolithic = null,
 			string UE4Exe = null,
 			string SignPak = null,
 			List<UnrealTargetConfiguration> ClientConfigsToBuild = null,
 			List<UnrealTargetConfiguration> ServerConfigsToBuild = null,
 			ParamList<string> MapsToCook = null,
 			ParamList<string> DirectoriesToCook = null,
-            ParamList<string> CulturesToCook = null,
+			ParamList<string> CulturesToCook = null,
 			ParamList<string> ClientCookedTargets = null,
 			ParamList<string> EditorTargets = null,
 			ParamList<string> ServerCookedTargets = null,
 			List<UnrealTargetPlatform> ClientTargetPlatforms = null,
-            Dictionary<UnrealTargetPlatform, UnrealTargetPlatform> ClientDependentPlatformMap = null,
+			Dictionary<UnrealTargetPlatform, UnrealTargetPlatform> ClientDependentPlatformMap = null,
 			List<UnrealTargetPlatform> ServerTargetPlatforms = null,
-            Dictionary<UnrealTargetPlatform, UnrealTargetPlatform> ServerDependentPlatformMap = null,
+			Dictionary<UnrealTargetPlatform, UnrealTargetPlatform> ServerDependentPlatformMap = null,
 			bool? Build = null,
 			bool? Cook = null,
 			string CookFlavor = null,
 			bool? Run = null,
 			bool? SkipServer = null,
 			bool? Clean = null,
-            bool? Compressed = null,
-            bool? UseDebugParamForEditorExe = null,
-            bool? IterativeCooking = null,
+			bool? Compressed = null,
+			bool? UseDebugParamForEditorExe = null,
+			bool? IterativeCooking = null,
 			bool? CookOnTheFly = null,
 			bool? CrashReporter = null,
 			bool? DedicatedServer = null,
@@ -336,21 +336,21 @@ namespace AutomationTool
 			bool? Pak = null,
 			bool? Prereqs = null,
 			bool? NoBootstrapExe = null,
-            bool? SignedPak = null,
-            bool? NullRHI = null,
-            bool? FakeClient = null,
-            bool? EditorTest = null,
-            bool? RunAutomationTests = null,
-            string RunAutomationTest = null,
-            int? CrashIndex = null,
-            bool? Rocket = null,
+			bool? SignedPak = null,
+			bool? NullRHI = null,
+			bool? FakeClient = null,
+			bool? EditorTest = null,
+			bool? RunAutomationTests = null,
+			string RunAutomationTest = null,
+			int? CrashIndex = null,
+			bool? Rocket = null,
 			bool? SkipCook = null,
 			bool? SkipCookOnTheFly = null,
 			bool? SkipPak = null,
 			bool? SkipStage = null,
 			bool? Stage = null,
 			bool? Manifests = null,
-            bool? CreateChunkInstall = null,
+			bool? CreateChunkInstall = null,
 			bool? Unattended = null,
 			int? NumClients = null,
 			bool? Archive = null,
@@ -358,9 +358,9 @@ namespace AutomationTool
 			bool? ArchiveMetaData = null,
 			ParamList<string> ProgramTargets = null,
 			bool? Distribution = null,
-            bool? OBBinAPK = null,
-            bool? Prebuilt = null,
-            int? RunTimeoutSeconds = null
+			bool? OBBinAPK = null,
+			bool? Prebuilt = null,
+			int? RunTimeoutSeconds = null
 			)
 		{
 			//
@@ -376,9 +376,9 @@ namespace AutomationTool
 			{
 				this.DirectoriesToCook = DirectoriesToCook;
 			}
-            if (CulturesToCook != null)
+			if (CulturesToCook != null)
 			{
-                this.CulturesToCook = CulturesToCook;
+				this.CulturesToCook = CulturesToCook;
 			}
 			if (ClientCookedTargets != null)
 			{
@@ -398,18 +398,18 @@ namespace AutomationTool
 			}
 
 			// Parse command line params for client platforms "-TargetPlatform=Win64+Mac", "-Platform=Win64+Mac" and also "-Win64", "-Mac" etc.
-            if (ClientDependentPlatformMap != null)
-            {
-                this.ClientDependentPlatformMap = ClientDependentPlatformMap;
-            }
+			if (ClientDependentPlatformMap != null)
+			{
+				this.ClientDependentPlatformMap = ClientDependentPlatformMap;
+			}
 			this.ClientTargetPlatforms = SetupTargetPlatforms(ref this.ClientDependentPlatformMap, Command, ClientTargetPlatforms, new ParamList<UnrealTargetPlatform>() {HostPlatform.Current.HostEditorPlatform}, true, "TargetPlatform", "Platform");
 
-            // Parse command line params for server platforms "-ServerTargetPlatform=Win64+Mac", "-ServerPlatform=Win64+Mac". "-Win64" etc is not allowed here
-            if (ServerDependentPlatformMap != null)
-            {
-                this.ServerDependentPlatformMap = ServerDependentPlatformMap;
-            }
-            this.ServerTargetPlatforms = SetupTargetPlatforms(ref this.ServerDependentPlatformMap, Command, ServerTargetPlatforms, this.ClientTargetPlatforms, false, "ServerTargetPlatform", "ServerPlatform");
+			// Parse command line params for server platforms "-ServerTargetPlatform=Win64+Mac", "-ServerPlatform=Win64+Mac". "-Win64" etc is not allowed here
+			if (ServerDependentPlatformMap != null)
+			{
+				this.ServerDependentPlatformMap = ServerDependentPlatformMap;
+			}
+			this.ServerTargetPlatforms = SetupTargetPlatforms(ref this.ServerDependentPlatformMap, Command, ServerTargetPlatforms, this.ClientTargetPlatforms, false, "ServerTargetPlatform", "ServerPlatform");
 
 			this.Build = GetParamValueIfNotSpecified(Command, Build, this.Build, "build");
 			this.Run = GetParamValueIfNotSpecified(Command, Run, this.Run, "run");
@@ -431,9 +431,9 @@ namespace AutomationTool
 			}
 			this.NoXGE = GetParamValueIfNotSpecified(Command, NoXGE, this.NoXGE, "noxge");
 			this.CookOnTheFly = GetParamValueIfNotSpecified(Command, CookOnTheFly, this.CookOnTheFly, "cookonthefly");
-            this.Compressed = GetParamValueIfNotSpecified(Command, Compressed, this.Compressed, "compressed");
-            this.UseDebugParamForEditorExe = GetParamValueIfNotSpecified(Command, UseDebugParamForEditorExe, this.UseDebugParamForEditorExe, "UseDebugParamForEditorExe");
-            this.IterativeCooking = GetParamValueIfNotSpecified(Command, IterativeCooking, this.IterativeCooking, "iterativecooking", "iterate");
+			this.Compressed = GetParamValueIfNotSpecified(Command, Compressed, this.Compressed, "compressed");
+			this.UseDebugParamForEditorExe = GetParamValueIfNotSpecified(Command, UseDebugParamForEditorExe, this.UseDebugParamForEditorExe, "UseDebugParamForEditorExe");
+			this.IterativeCooking = GetParamValueIfNotSpecified(Command, IterativeCooking, this.IterativeCooking, "iterativecooking", "iterate");
 			this.SkipCookOnTheFly = GetParamValueIfNotSpecified(Command, SkipCookOnTheFly, this.SkipCookOnTheFly, "skipcookonthefly");
 			this.FileServer = GetParamValueIfNotSpecified(Command, FileServer, this.FileServer, "fileserver");
 			this.DedicatedServer = GetParamValueIfNotSpecified(Command, DedicatedServer, this.DedicatedServer, "dedicatedserver", "server");
@@ -451,32 +451,32 @@ namespace AutomationTool
 				this.Stage = true;
 			}
 			this.StageDirectoryParam = ParseParamValueIfNotSpecified(Command, StageDirectoryParam, "stagingdirectory", String.Empty).Trim(new char[]{'\"'});
-            this.StageNonMonolithic = GetParamValueIfNotSpecified(Command, StageNonMonolithic, this.StageNonMonolithic, "StageNonMonolithic");
+			this.StageNonMonolithic = GetParamValueIfNotSpecified(Command, StageNonMonolithic, this.StageNonMonolithic, "StageNonMonolithic");
 			this.Manifests = GetParamValueIfNotSpecified(Command, Manifests, this.Manifests, "manifests");
-            this.CreateChunkInstall = GetParamValueIfNotSpecified(Command, CreateChunkInstall, this.CreateChunkInstall, "createchunkinstall");
+			this.CreateChunkInstall = GetParamValueIfNotSpecified(Command, CreateChunkInstall, this.CreateChunkInstall, "createchunkinstall");
 			this.Archive = GetParamValueIfNotSpecified(Command, Archive, this.Archive, "archive");
 			this.ArchiveDirectoryParam = ParseParamValueIfNotSpecified(Command, ArchiveDirectoryParam, "archivedirectory", String.Empty);
 			this.ArchiveMetaData = GetParamValueIfNotSpecified(Command, ArchiveMetaData, this.ArchiveMetaData, "archivemetadata");
 			this.Distribution = GetParamValueIfNotSpecified(Command, Distribution, this.Distribution, "distribution");
 			this.Prereqs = GetParamValueIfNotSpecified(Command, Prereqs, this.Prereqs, "prereqs");
 			this.NoBootstrapExe = GetParamValueIfNotSpecified(Command, NoBootstrapExe, this.NoBootstrapExe, "nobootstrapexe");
-            this.OBBinAPK = GetParamValueIfNotSpecified(Command, OBBinAPK, this.OBBinAPK, "obbinapk");
-            this.Prebuilt = GetParamValueIfNotSpecified(Command, Prebuilt, this.Prebuilt, "prebuilt");
-            if (this.Prebuilt)
-            {
-                this.SkipCook = true;
-                /*this.SkipPak = true;
-                this.SkipStage = true;
-                this.Pak = true;
-                this.Stage = true;*/
-                this.Cook = true;
-                this.Archive = true;
-                
-                this.Deploy = true;
-                this.Run = true;
-                //this.StageDirectoryParam = this.PrebuiltDir;
-            }
-            this.NoDebugInfo = GetParamValueIfNotSpecified(Command, NoDebugInfo, this.NoDebugInfo, "nodebuginfo");
+			this.OBBinAPK = GetParamValueIfNotSpecified(Command, OBBinAPK, this.OBBinAPK, "obbinapk");
+			this.Prebuilt = GetParamValueIfNotSpecified(Command, Prebuilt, this.Prebuilt, "prebuilt");
+			if (this.Prebuilt)
+			{
+				this.SkipCook = true;
+				/*this.SkipPak = true;
+				this.SkipStage = true;
+				this.Pak = true;
+				this.Stage = true;*/
+				this.Cook = true;
+				this.Archive = true;
+				
+				this.Deploy = true;
+				this.Run = true;
+				//this.StageDirectoryParam = this.PrebuiltDir;
+			}
+			this.NoDebugInfo = GetParamValueIfNotSpecified(Command, NoDebugInfo, this.NoDebugInfo, "nodebuginfo");
 			this.NoCleanStage = GetParamValueIfNotSpecified(Command, NoCleanStage, this.NoCleanStage, "nocleanstage");
 			this.MapToRun = ParseParamValueIfNotSpecified(Command, MapToRun, "map", String.Empty);
 			this.AdditionalServerMapParams = ParseParamValueIfNotSpecified(Command, AdditionalServerMapParams, "AdditionalServerMapParams", String.Empty);
@@ -492,9 +492,9 @@ namespace AutomationTool
 			this.NullRHI = GetParamValueIfNotSpecified(Command, NullRHI, this.NullRHI, "nullrhi");
 			this.FakeClient = GetParamValueIfNotSpecified(Command, FakeClient, this.FakeClient, "fakeclient");
 			this.EditorTest = GetParamValueIfNotSpecified(Command, EditorTest, this.EditorTest, "editortest");
-            this.RunAutomationTest = ParseParamValueIfNotSpecified(Command, RunAutomationTest, "RunAutomationTest");
-            this.RunAutomationTests = this.RunAutomationTest != "" || GetParamValueIfNotSpecified(Command, RunAutomationTests, this.RunAutomationTests, "RunAutomationTests");
-            this.SkipServer = GetParamValueIfNotSpecified(Command, SkipServer, this.SkipServer, "skipserver");
+			this.RunAutomationTest = ParseParamValueIfNotSpecified(Command, RunAutomationTest, "RunAutomationTest");
+			this.RunAutomationTests = this.RunAutomationTest != "" || GetParamValueIfNotSpecified(Command, RunAutomationTests, this.RunAutomationTests, "RunAutomationTests");
+			this.SkipServer = GetParamValueIfNotSpecified(Command, SkipServer, this.SkipServer, "skipserver");
 			this.Rocket = GetParamValueIfNotSpecified(Command, Rocket, this.Rocket || GlobalCommandLine.Rocket, "rocket");
 			this.UE4Exe = ParseParamValueIfNotSpecified(Command, UE4Exe, "ue4exe", "UE4Editor-Cmd.exe");
 			this.Unattended = GetParamValueIfNotSpecified(Command, Unattended, this.Unattended, "unattended");
@@ -522,28 +522,28 @@ namespace AutomationTool
 				this.ClientConfigsToBuild = ClientConfigsToBuild;
 			}
 
-            if (Port == null)
-            {
-                if( Command != null )
-                {
-                    this.Port = new ParamList<string>();
-
-                    var PortString = Command.ParseParamValue("port");
-                    if (String.IsNullOrEmpty(PortString) == false)
-                    {
-                        var Ports = new ParamList<string>(PortString.Split('+'));
-                        foreach (var P in Ports)
-                        {
-                            this.Port.Add(P);
-                        }
-                    }
-                    
-                }
-            }
-            else
-            {
-                this.Port = Port;
-            }
+			if (Port == null)
+			{
+				if( Command != null )
+				{
+					this.Port = new ParamList<string>();
+
+					var PortString = Command.ParseParamValue("port");
+					if (String.IsNullOrEmpty(PortString) == false)
+					{
+						var Ports = new ParamList<string>(PortString.Split('+'));
+						foreach (var P in Ports)
+						{
+							this.Port.Add(P);
+						}
+					}
+					
+				}
+			}
+			else
+			{
+				this.Port = Port;
+			}
 
 			if (ServerConfigsToBuild == null)
 			{
@@ -573,22 +573,22 @@ namespace AutomationTool
 			{
 				this.NumClients = Command.ParseParamInt("numclients");
 			}
-            if (CrashIndex.HasValue)
-            {
-                this.CrashIndex = CrashIndex.Value;
-            }
-            else if (Command != null)
-            {
-                this.CrashIndex = Command.ParseParamInt("CrashIndex");
-            }
-            if (RunTimeoutSeconds.HasValue)
-            {
-                this.RunTimeoutSeconds = RunTimeoutSeconds.Value;
-            }
-            else if (Command != null)
-            {
-                this.RunTimeoutSeconds = Command.ParseParamInt("runtimeoutseconds");
-            }
+			if (CrashIndex.HasValue)
+			{
+				this.CrashIndex = CrashIndex.Value;
+			}
+			else if (Command != null)
+			{
+				this.CrashIndex = Command.ParseParamInt("CrashIndex");
+			}
+			if (RunTimeoutSeconds.HasValue)
+			{
+				this.RunTimeoutSeconds = RunTimeoutSeconds.Value;
+			}
+			else if (Command != null)
+			{
+				this.RunTimeoutSeconds = Command.ParseParamInt("runtimeoutseconds");
+			}
 
 			AutodetectSettings(false);
 			ValidateAndLog();
@@ -664,24 +664,24 @@ namespace AutomationTool
 		public bool Unattended { private set; get; }
 
 		/// <summary>
-        /// Shared: Sets platforms to build for non-dedicated servers. commandline: -TargetPlatform
+		/// Shared: Sets platforms to build for non-dedicated servers. commandline: -TargetPlatform
 		/// </summary>
 		public List<UnrealTargetPlatform> ClientTargetPlatforms = new List<UnrealTargetPlatform>();
 
-        /// <summary>
-        /// Shared: Dictionary that maps client dependent platforms to "source" platforms that it should copy data from. commandline: -TargetPlatform=source.dependent
-        /// </summary>
-        public Dictionary<UnrealTargetPlatform, UnrealTargetPlatform> ClientDependentPlatformMap = new Dictionary<UnrealTargetPlatform, UnrealTargetPlatform>();
+		/// <summary>
+		/// Shared: Dictionary that maps client dependent platforms to "source" platforms that it should copy data from. commandline: -TargetPlatform=source.dependent
+		/// </summary>
+		public Dictionary<UnrealTargetPlatform, UnrealTargetPlatform> ClientDependentPlatformMap = new Dictionary<UnrealTargetPlatform, UnrealTargetPlatform>();
 
 		/// <summary>
-        /// Shared: Sets platforms to build for dedicated servers. commandline: -ServerTargetPlatform
+		/// Shared: Sets platforms to build for dedicated servers. commandline: -ServerTargetPlatform
 		/// </summary>
 		public List<UnrealTargetPlatform> ServerTargetPlatforms = new List<UnrealTargetPlatform>();
 
-        /// <summary>
-        /// Shared: Dictionary that maps server dependent platforms to "source" platforms that it should copy data from: -ServerTargetPlatform=source.dependent
-        /// </summary>
-        public Dictionary<UnrealTargetPlatform, UnrealTargetPlatform> ServerDependentPlatformMap = new Dictionary<UnrealTargetPlatform, UnrealTargetPlatform>();
+		/// <summary>
+		/// Shared: Dictionary that maps server dependent platforms to "source" platforms that it should copy data from: -ServerTargetPlatform=source.dependent
+		/// </summary>
+		public Dictionary<UnrealTargetPlatform, UnrealTargetPlatform> ServerDependentPlatformMap = new Dictionary<UnrealTargetPlatform, UnrealTargetPlatform>();
 
 		/// <summary>
 		/// Shared: True if pak file should be generated.
@@ -733,11 +733,11 @@ namespace AutomationTool
 		[Help("manifests", "generate streaming install manifests when cooking data")]
 		public bool Manifests { private set; get; }
 
-        /// <summary>
-        /// Shared: true if this build chunk install streaming install data, command line: -createchunkinstalldata !!JM
-        /// </summary>
-        [Help("createchunkinstall", "generate streaming install data from manifest when cooking data, requires -stage & -manifests")]
-        public bool CreateChunkInstall { private set; get; }
+		/// <summary>
+		/// Shared: true if this build chunk install streaming install data, command line: -createchunkinstalldata !!JM
+		/// </summary>
+		[Help("createchunkinstall", "generate streaming install data from manifest when cooking data, requires -stage & -manifests")]
+		public bool CreateChunkInstall { private set; get; }
 
 		/// <summary>
 		/// Shared: Directory to copy the client to, command line: -stagingdirectory=
@@ -753,10 +753,10 @@ namespace AutomationTool
 		[Help("stagingdirectory=Path", "Directory to copy the builds to, i.e. -stagingdirectory=C:\\Stage")]
 		public string StageDirectoryParam;
 
-        /// <summary>
-        /// Whether the project should use non monolithic staging
-        /// </summary>
-        public bool StageNonMonolithic;
+		/// <summary>
+		/// Whether the project should use non monolithic staging
+		/// </summary>
+		public bool StageNonMonolithic;
 
 		[Help("ue4exe=ExecutableName", "Name of the UE4 Editor executable, i.e. -ue4exe=UE4Editor.exe")]
 		public string UE4Exe;
@@ -774,7 +774,7 @@ namespace AutomationTool
 		{
 			get
 			{
-                return Path.GetFullPath(String.IsNullOrEmpty(ArchiveDirectoryParam) ? CommandUtils.CombinePaths(Path.GetDirectoryName(RawProjectPath), "ArchivedBuilds") : ArchiveDirectoryParam);
+				return Path.GetFullPath(String.IsNullOrEmpty(ArchiveDirectoryParam) ? CommandUtils.CombinePaths(Path.GetDirectoryName(RawProjectPath), "ArchivedBuilds") : ArchiveDirectoryParam);
 			}
 		}
 
@@ -899,20 +899,20 @@ namespace AutomationTool
 		/// </summary>
 		public ParamList<string> DirectoriesToCook = new ParamList<string>();
 
-        /// <summary>
-        /// Cook: List of cultures to cook.
-        /// </summary>
-        public ParamList<string> CulturesToCook = new ParamList<string>();
+		/// <summary>
+		/// Cook: List of cultures to cook.
+		/// </summary>
+		public ParamList<string> CulturesToCook = new ParamList<string>();
 
-        /// <summary>
-        /// Compress packages during cook.
-        /// </summary>
-        public bool Compressed;
+		/// <summary>
+		/// Compress packages during cook.
+		/// </summary>
+		public bool Compressed;
 
-        /// <summary>
-        /// put -debug on the editorexe commandline
-        /// </summary>
-        public bool UseDebugParamForEditorExe;
+		/// <summary>
+		/// put -debug on the editorexe commandline
+		/// </summary>
+		public bool UseDebugParamForEditorExe;
 
 		/// <summary>
 		/// Cook: Uses the iterative cooking, command line: -iterativecooking or -iterate
@@ -943,16 +943,16 @@ namespace AutomationTool
 		[Help("cmdline", "command line to put into the stage in UE4CommandLine.txt")]
 		public string StageCommandline;
 
-        /// <summary>
+		/// <summary>
 		/// Stage: If non-empty, the contents will be used for the bundle name
 		/// </summary>
 		[Help("bundlename", "string to use as the bundle name when deploying to mobile device")]
-        public string BundleName;
+		public string BundleName;
 
-        /// <summary>
-        /// Whether the project uses Steam (todo: substitute with more generic functionality)
-        /// </summary>
-        public bool bUsesSteam;
+		/// <summary>
+		/// Whether the project uses Steam (todo: substitute with more generic functionality)
+		/// </summary>
+		public bool bUsesSteam;
 
 		/// <summary>
 		/// Whether the project uses visual Slate UI (as opposed to the low level windowing/messaging which is alway used)
@@ -965,20 +965,20 @@ namespace AutomationTool
 		/// </summary>
 		public bool bUsesSlateEditorStyle = false;
 
-        /// <summary>
-        /// By default we use the Release C++ Runtime (CRT), even when compiling Debug builds.  This is because the Debug C++
-        /// Runtime isn't very useful when debugging Unreal Engine projects, and linking against the Debug CRT libraries forces
-        /// our third party library dependencies to also be compiled using the Debug CRT (and often perform more slowly.)  Often
-        /// it can be inconvenient to require a separate copy of the debug versions of third party static libraries simply
-        /// so that you can debug your program's code.
-        /// </summary>
-        public bool bDebugBuildsActuallyUseDebugCRT = false;
-
-        /// <summary>
-        /// On Windows, adds an executable to the root of the staging directory which checks for prerequisites being 
+		/// <summary>
+		/// By default we use the Release C++ Runtime (CRT), even when compiling Debug builds.  This is because the Debug C++
+		/// Runtime isn't very useful when debugging Unreal Engine projects, and linking against the Debug CRT libraries forces
+		/// our third party library dependencies to also be compiled using the Debug CRT (and often perform more slowly.)  Often
+		/// it can be inconvenient to require a separate copy of the debug versions of third party static libraries simply
+		/// so that you can debug your program's code.
+		/// </summary>
+		public bool bDebugBuildsActuallyUseDebugCRT = false;
+
+		/// <summary>
+		/// On Windows, adds an executable to the root of the staging directory which checks for prerequisites being 
 		/// installed and launches the game with a path to the .uproject file.
 		/// </summary>
-        public bool NoBootstrapExe { get; set; }
+		public bool NoBootstrapExe { get; set; }
 
 		#endregion
 
@@ -1068,62 +1068,62 @@ namespace AutomationTool
 		[Help("addcmdline", "Additional command line arguments for the program")]
 		public string RunCommandline;
 
-        /// <summary>
-        /// Run:adds -nullrhi to the client commandline
-        /// </summary>
-        [Help("nullrhi", "add -nullrhi to the client commandlines")]
-        public bool NullRHI;
-
-        /// <summary>
-        /// Run:adds ?fake to the server URL
-        /// </summary>
-        [Help("fakeclient", "adds ?fake to the server URL")]
-        public bool FakeClient;
-
-        /// <summary>
-        /// Run:adds ?fake to the server URL
-        /// </summary>
-        [Help("editortest", "rather than running a client, run the editor instead")]
-        public bool EditorTest;
-
-        /// <summary>
-        /// Run:when running -editortest or a client, run all automation tests, not compatible with -server
-        /// </summary>
-        [Help("RunAutomationTests", "when running -editortest or a client, run all automation tests, not compatible with -server")]
-        public bool RunAutomationTests;
-
-        /// <summary>
-        /// Run:when running -editortest or a client, run all automation tests, not compatible with -server
-        /// </summary>
-        [Help("RunAutomationTests", "when running -editortest or a client, run a specific automation tests, not compatible with -server")]
-        public string RunAutomationTest;
-
-        /// <summary>
-        /// Run: Adds commands like debug crash, debug rendercrash, etc based on index
-        /// </summary>
-        [Help("Crash=index", "when running -editortest or a client, adds commands like debug crash, debug rendercrash, etc based on index")]
-        public int CrashIndex;
-
-        public ParamList<string> Port;
-
-        /// <summary>
-        /// Run: Linux username for unattended key genereation
-        /// </summary>
-        [Help("deviceuser", "Linux username for unattended key genereation")]
-        public string DeviceUsername;
-
-        /// <summary>
-        /// Run: Linux password for unattended key genereation
-        /// </summary>
-        [Help("devicepass", "Linux password for unattended key genereation")]
-        public string DevicePassword;
-
-        /// <summary>
-        /// Run: Sever device IP address
-        /// </summary>
-        public string ServerDeviceAddress;
-
-        #endregion
+		/// <summary>
+		/// Run:adds -nullrhi to the client commandline
+		/// </summary>
+		[Help("nullrhi", "add -nullrhi to the client commandlines")]
+		public bool NullRHI;
+
+		/// <summary>
+		/// Run:adds ?fake to the server URL
+		/// </summary>
+		[Help("fakeclient", "adds ?fake to the server URL")]
+		public bool FakeClient;
+
+		/// <summary>
+		/// Run:adds ?fake to the server URL
+		/// </summary>
+		[Help("editortest", "rather than running a client, run the editor instead")]
+		public bool EditorTest;
+
+		/// <summary>
+		/// Run:when running -editortest or a client, run all automation tests, not compatible with -server
+		/// </summary>
+		[Help("RunAutomationTests", "when running -editortest or a client, run all automation tests, not compatible with -server")]
+		public bool RunAutomationTests;
+
+		/// <summary>
+		/// Run:when running -editortest or a client, run all automation tests, not compatible with -server
+		/// </summary>
+		[Help("RunAutomationTests", "when running -editortest or a client, run a specific automation tests, not compatible with -server")]
+		public string RunAutomationTest;
+
+		/// <summary>
+		/// Run: Adds commands like debug crash, debug rendercrash, etc based on index
+		/// </summary>
+		[Help("Crash=index", "when running -editortest or a client, adds commands like debug crash, debug rendercrash, etc based on index")]
+		public int CrashIndex;
+
+		public ParamList<string> Port;
+
+		/// <summary>
+		/// Run: Linux username for unattended key genereation
+		/// </summary>
+		[Help("deviceuser", "Linux username for unattended key genereation")]
+		public string DeviceUsername;
+
+		/// <summary>
+		/// Run: Linux password for unattended key genereation
+		/// </summary>
+		[Help("devicepass", "Linux password for unattended key genereation")]
+		public string DevicePassword;
+
+		/// <summary>
+		/// Run: Sever device IP address
+		/// </summary>
+		public string ServerDeviceAddress;
+
+		#endregion
 
 		#region Package
 
@@ -1137,13 +1137,13 @@ namespace AutomationTool
 		public bool Prereqs { get; set; }
 
 		[Help("obbinapk", "package with OBB data in APK assets directory")]
-        public bool OBBinAPK {get; set; }
+		public bool OBBinAPK {get; set; }
 
-        [Help("Prebuilt", "this is a prebuilt cooked and packaged build")]
-        public bool Prebuilt { get; private set; }
+		[Help("Prebuilt", "this is a prebuilt cooked and packaged build")]
+		public bool Prebuilt { get; private set; }
 
-        [Help("RunTimeoutSeconds", "timeout to wait after we lunch the game")]
-        public int RunTimeoutSeconds;
+		[Help("RunTimeoutSeconds", "timeout to wait after we lunch the game")]
+		public int RunTimeoutSeconds;
 
 
 		#endregion
@@ -1178,7 +1178,7 @@ namespace AutomationTool
 			bUsesSteam = Properties.bUsesSteam;
 			bUsesSlate = Properties.bUsesSlate;
 			bUsesSlateEditorStyle = Properties.bUsesSlateEditorStyle;
-            bDebugBuildsActuallyUseDebugCRT = Properties.bDebugBuildsActuallyUseDebugCRT;
+			bDebugBuildsActuallyUseDebugCRT = Properties.bDebugBuildsActuallyUseDebugCRT;
 
 			bIsCodeBasedProject = Properties.bIsCodeBasedProject;			
 			DetectedTargets = Properties.Targets;
@@ -1276,7 +1276,7 @@ namespace AutomationTool
 					if (DetectedTargets.TryGetValue(ValidTarget, out TargetData))
 					{
 						GameTarget = TargetData.TargetName;
-                        bDebugBuildsActuallyUseDebugCRT = TargetData.Rules.bDebugBuildsActuallyUseDebugCRT;
+						bDebugBuildsActuallyUseDebugCRT = TargetData.Rules.bDebugBuildsActuallyUseDebugCRT;
 						bUsesSlate = TargetData.Rules.bUsesSlate;
 						bUsesSlateEditorStyle = TargetData.Rules.bUsesSlateEditorStyle;
 						bUsesSteam = TargetData.Rules.bUsesSteam;
@@ -1357,27 +1357,27 @@ namespace AutomationTool
 				}
 			}
 
-            if (ClientCookedTargetsList == null && !NoClient && (Cook || CookOnTheFly || Prebuilt))
+			if (ClientCookedTargetsList == null && !NoClient && (Cook || CookOnTheFly || Prebuilt))
 			{
-                if (String.IsNullOrEmpty(GameTarget))
+				if (String.IsNullOrEmpty(GameTarget))
 				{
-                    throw new AutomationException("Game target not found. Game target is required with -cook or -cookonthefly");
-                }
+					throw new AutomationException("Game target not found. Game target is required with -cook or -cookonthefly");
+				}
 				else
 				{
-                    ClientCookedTargetsList = new ParamList<string>(GameTarget);
-                }
+					ClientCookedTargetsList = new ParamList<string>(GameTarget);
+				}
+			}
+			else if (ClientCookedTargetsList == null)
+			{
+				ClientCookedTargetsList = new ParamList<string>();
 			}
-            else if (ClientCookedTargetsList == null)
-            {
-                ClientCookedTargetsList = new ParamList<string>();
-            }
 
-            if (ServerCookedTargetsList == null && DedicatedServer && (Cook || CookOnTheFly))
+			if (ServerCookedTargetsList == null && DedicatedServer && (Cook || CookOnTheFly))
 			{
 				if (String.IsNullOrEmpty(ServerTarget))
 				{
-                    throw new AutomationException("Server target not found. Server target is required with -server and -cook or -cookonthefly");
+					throw new AutomationException("Server target not found. Server target is required with -server and -cook or -cookonthefly");
 				}
 				ServerCookedTargetsList = new ParamList<string>(ServerTarget);
 			}
@@ -1395,6 +1395,19 @@ namespace AutomationTool
 					ProjectGameExePath = CommandUtils.CombinePaths(ProjectBinariesPath, GameTarget + Platform.GetExeExtension(ClientTargetPlatforms[0]));
 				}
 			}
+
+			// for the moment, allow the commandline to override the ini if set.  This will keep us from breaking licensee packaging scripts until we have
+			// the full solution for per-platform packaging settings.
+			if (!Manifests)
+			{				
+				ConfigCacheIni GameIni = new ConfigCacheIni("Game", Path.GetDirectoryName(RawProjectPath));
+				String IniPath = "/Script/UnrealEd.ProjectPackagingSettings";
+				bool bSetting = false;
+				if (!GameIni.GetBool(IniPath, "bGenerateChunks", out bSetting))
+				{
+					Manifests = bSetting;
+				}
+			}
 		}
 
 		#endregion
@@ -1436,10 +1449,10 @@ namespace AutomationTool
 			get { return !CommandUtils.IsNullOrEmpty(DirectoriesToCook); }
 		}
 
-        public bool HasCulturesToCook
-        {
-            get { return !CommandUtils.IsNullOrEmpty(CulturesToCook); }
-        }
+		public bool HasCulturesToCook
+		{
+			get { return !CommandUtils.IsNullOrEmpty(CulturesToCook); }
+		}
 
 		public bool HasGameTargetDetected
 		{
@@ -1478,7 +1491,7 @@ namespace AutomationTool
 
 		public string CodeBasedUprojectPath
 		{
-            get { return IsCodeBasedProject ? RawProjectPath : null; }
+			get { return IsCodeBasedProject ? RawProjectPath : null; }
 		}
 		/// <summary>
 		/// True if this project is a program.
@@ -1535,14 +1548,14 @@ namespace AutomationTool
 			}
 		}
 
-        public UnrealTargetPlatform GetCookedDataPlatformForClientTarget(UnrealTargetPlatform TargetPlatformType)
-        {
-            if (ClientDependentPlatformMap.ContainsKey(TargetPlatformType))
-            {
-                return ClientDependentPlatformMap[TargetPlatformType];
-            }
-            return TargetPlatformType;
-        }
+		public UnrealTargetPlatform GetCookedDataPlatformForClientTarget(UnrealTargetPlatform TargetPlatformType)
+		{
+			if (ClientDependentPlatformMap.ContainsKey(TargetPlatformType))
+			{
+				return ClientDependentPlatformMap[TargetPlatformType];
+			}
+			return TargetPlatformType;
+		}
 
 		public List<Platform> ServerTargetPlatformInstances
 		{
@@ -1557,14 +1570,14 @@ namespace AutomationTool
 			}
 		}
 
-        public UnrealTargetPlatform GetCookedDataPlatformForServerTarget(UnrealTargetPlatform TargetPlatformType)
-        {
-            if (ServerDependentPlatformMap.ContainsKey(TargetPlatformType))
-            {
-                return ServerDependentPlatformMap[TargetPlatformType];
-            }
-            return TargetPlatformType;
-        }
+		public UnrealTargetPlatform GetCookedDataPlatformForServerTarget(UnrealTargetPlatform TargetPlatformType)
+		{
+			if (ServerDependentPlatformMap.ContainsKey(TargetPlatformType))
+			{
+				return ServerDependentPlatformMap[TargetPlatformType];
+			}
+			return TargetPlatformType;
+		}
 
 		/// <summary>
 		/// All auto-detected targets for this project
@@ -1617,14 +1630,14 @@ namespace AutomationTool
 			{
 				throw new AutomationException("RawProjectPath can't be empty.");
 			}
-            if (!RawProjectPath.EndsWith(".uproject", StringComparison.InvariantCultureIgnoreCase))
-            {
-                throw new AutomationException("RawProjectPath {0} must end with .uproject", RawProjectPath);
-            }
-            if (!CommandUtils.FileExists_NoExceptions(RawProjectPath))
-            {
-                throw new AutomationException("RawProjectPath {0} file must exist", RawProjectPath);
-            }
+			if (!RawProjectPath.EndsWith(".uproject", StringComparison.InvariantCultureIgnoreCase))
+			{
+				throw new AutomationException("RawProjectPath {0} must end with .uproject", RawProjectPath);
+			}
+			if (!CommandUtils.FileExists_NoExceptions(RawProjectPath))
+			{
+				throw new AutomationException("RawProjectPath {0} file must exist", RawProjectPath);
+			}
 
 			if (FileServer && !Cook)
 			{
@@ -1646,20 +1659,20 @@ namespace AutomationTool
 				throw new AutomationException("Only staged builds can be paked, use -stage or -skipstage.");
 			}
 
-            if (Deploy && !Stage)
-            {
-                throw new AutomationException("Only staged builds can be deployed, use -stage or -skipstage.");
-            }
+			if (Deploy && !Stage)
+			{
+				throw new AutomationException("Only staged builds can be deployed, use -stage or -skipstage.");
+			}
 
-            if ((Pak || Stage || Cook || CookOnTheFly || FileServer || DedicatedServer) && EditorTest)
-            {
-                throw new AutomationException("None of pak, stage, cook, CookOnTheFly or DedicatedServer can be used with EditorTest");
-            }
+			if ((Pak || Stage || Cook || CookOnTheFly || FileServer || DedicatedServer) && EditorTest)
+			{
+				throw new AutomationException("None of pak, stage, cook, CookOnTheFly or DedicatedServer can be used with EditorTest");
+			}
 
-            if (DedicatedServer && RunAutomationTests)
-            {
-                throw new AutomationException("DedicatedServer cannot be used with RunAutomationTests");
-            }
+			if (DedicatedServer && RunAutomationTests)
+			{
+				throw new AutomationException("DedicatedServer cannot be used with RunAutomationTests");
+			}
 
 			if ((CookOnTheFly || FileServer) && DedicatedServer)
 			{
@@ -1686,15 +1699,15 @@ namespace AutomationTool
 				throw new AutomationException("Don't use both -cook and -cookonthefly.");
 			}
 
-            if (Compressed && !Pak)
-            {
-                throw new AutomationException("-compressed can only be used with -pak");
-            }
+			if (Compressed && !Pak)
+			{
+				throw new AutomationException("-compressed can only be used with -pak");
+			}
 
-            if (CreateChunkInstall && (!Manifests || !Stage))
-            {
-                throw new AutomationException("-createchunkinstall can only be used with -pak & -stage");
-            }
+			if (CreateChunkInstall && (!Manifests || !Stage))
+			{
+				throw new AutomationException("-createchunkinstall can only be used with -pak & -stage");
+			}
 		}
 
 		protected bool bLogged = false;
@@ -1718,13 +1731,13 @@ namespace AutomationTool
 				CommandUtils.Log("ClientConfigsToBuild={0}", string.Join(",", ClientConfigsToBuild));
 				CommandUtils.Log("ClientCookedTargets={0}", ClientCookedTargets.ToString());
 				CommandUtils.Log("ClientTargetPlatform={0}", string.Join(",", ClientTargetPlatforms));
-                CommandUtils.Log("Compressed={0}", Compressed);
-                CommandUtils.Log("UseDebugParamForEditorExe={0}", UseDebugParamForEditorExe);
+				CommandUtils.Log("Compressed={0}", Compressed);
+				CommandUtils.Log("UseDebugParamForEditorExe={0}", UseDebugParamForEditorExe);
 				CommandUtils.Log("CookFlavor={0}", CookFlavor);
 				CommandUtils.Log("CookOnTheFly={0}", CookOnTheFly);
 				CommandUtils.Log("DedicatedServer={0}", DedicatedServer);
 				CommandUtils.Log("DirectoriesToCook={0}", DirectoriesToCook.ToString());
-                CommandUtils.Log("CulturesToCook={0}", CulturesToCook.ToString());
+				CommandUtils.Log("CulturesToCook={0}", CulturesToCook.ToString());
 				CommandUtils.Log("EditorTargets={0}", EditorTargets.ToString());
 				CommandUtils.Log("Foreign={0}", Foreign);
 				CommandUtils.Log("IsCodeBasedProject={0}", IsCodeBasedProject.ToString());
@@ -1742,19 +1755,19 @@ namespace AutomationTool
 				CommandUtils.Log("Pak={0}", Pak);
 				CommandUtils.Log("NullRHI={0}", NullRHI);
 				CommandUtils.Log("FakeClient={0}", FakeClient);
-                CommandUtils.Log("EditorTest={0}", EditorTest);
-                CommandUtils.Log("RunAutomationTests={0}", RunAutomationTests); 
-                CommandUtils.Log("RunAutomationTest={0}", RunAutomationTest);
-                CommandUtils.Log("RunTimeoutSeconds={0}", RunTimeoutSeconds);
-                CommandUtils.Log("CrashIndex={0}", CrashIndex);
+				CommandUtils.Log("EditorTest={0}", EditorTest);
+				CommandUtils.Log("RunAutomationTests={0}", RunAutomationTests); 
+				CommandUtils.Log("RunAutomationTest={0}", RunAutomationTest);
+				CommandUtils.Log("RunTimeoutSeconds={0}", RunTimeoutSeconds);
+				CommandUtils.Log("CrashIndex={0}", CrashIndex);
 				CommandUtils.Log("ProgramTargets={0}", ProgramTargets.ToString());
-                CommandUtils.Log("ProjectBinariesFolder={0}", ProjectBinariesFolder);
+				CommandUtils.Log("ProjectBinariesFolder={0}", ProjectBinariesFolder);
 				CommandUtils.Log("ProjectBinariesPath={0}", ProjectBinariesPath);
 				CommandUtils.Log("ProjectGameExeFilename={0}", ProjectGameExeFilename);
 				CommandUtils.Log("ProjectGameExePath={0}", ProjectGameExePath);
 				CommandUtils.Log("Distribution={0}", Distribution);
-                CommandUtils.Log("OBBinAPK={0}", OBBinAPK);
-                CommandUtils.Log("Prebuilt={0}", Prebuilt);
+				CommandUtils.Log("OBBinAPK={0}", OBBinAPK);
+				CommandUtils.Log("Prebuilt={0}", Prebuilt);
 				CommandUtils.Log("Prereqs={0}", Prereqs);
 				CommandUtils.Log("NoBootstrapExe={0}", NoBootstrapExe);
 				CommandUtils.Log("RawProjectPath={0}", RawProjectPath);
@@ -1773,7 +1786,7 @@ namespace AutomationTool
 				CommandUtils.Log("Stage={0}", Stage);
 				CommandUtils.Log("bUsesSteam={0}", bUsesSteam);
 				CommandUtils.Log("bUsesSlate={0}", bUsesSlate);
-                CommandUtils.Log("bDebugBuildsActuallyUseDebugCRT={0}", bDebugBuildsActuallyUseDebugCRT);
+				CommandUtils.Log("bDebugBuildsActuallyUseDebugCRT={0}", bDebugBuildsActuallyUseDebugCRT);
 				CommandUtils.Log("Project Params **************");
 			}
 			bLogged = true;
diff --git a/Engine/Source/Programs/UnrealBuildTool/Configuration/EngineConfiguration.cs b/Engine/Source/Programs/UnrealBuildTool/Configuration/EngineConfiguration.cs
--- a/Engine/Source/Programs/UnrealBuildTool/Configuration/EngineConfiguration.cs
+++ b/Engine/Source/Programs/UnrealBuildTool/Configuration/EngineConfiguration.cs
@@ -54,6 +54,17 @@ namespace UnrealBuildTool
 		/// </summary>
 		Dictionary<string, IniSection> Sections;
 
+		/// <summary>
+		/// Constructor. Parses ini hierarchy for the specified project.  No Platform settings.
+		/// </summary>
+		/// <param name="ProjectDirectory">Project path</param>
+		/// <param name="Platform">Target platform</param>
+		/// <param name="BaseIniName">Ini name (Engine, Editor, etc)</param>
+		public ConfigCacheIni(string BaseIniName, string ProjectDirectory, string EngineDirectory = null)
+		{
+			Init(UnrealTargetPlatform.Unknown, BaseIniName, ProjectDirectory, EngineDirectory);
+		}
+
 		/// <summary>
 		/// Constructor. Parses ini hierarchy for the specified platform and project.
 		/// </summary>
@@ -61,19 +72,36 @@ namespace UnrealBuildTool
 		/// <param name="Platform">Target platform</param>
 		/// <param name="BaseIniName">Ini name (Engine, Editor, etc)</param>
 		public ConfigCacheIni(UnrealTargetPlatform Platform, string BaseIniName, string ProjectDirectory, string EngineDirectory = null)
+		{
+			Init(Platform, BaseIniName, ProjectDirectory, EngineDirectory);
+		}
+
+		private void Init(UnrealTargetPlatform Platform, string BaseIniName, string ProjectDirectory, string EngineDirectory)
 		{
 			Sections = new Dictionary<string, IniSection>(StringComparer.InvariantCultureIgnoreCase);
 			if (String.IsNullOrEmpty(EngineDirectory))
 			{
 				EngineDirectory = BuildConfiguration.RelativeEnginePath;
 			}
-			foreach (var IniFileName in EnumerateIniFileNames(ProjectDirectory, EngineDirectory, Platform, BaseIniName))
+
+			foreach (var IniFileName in EnumerateCrossPlatformIniFileNames(ProjectDirectory, EngineDirectory, BaseIniName))
 			{
 				if (File.Exists(IniFileName))
 				{
 					ParseIniFile(IniFileName);
 				}
 			}
+
+			if (Platform != UnrealTargetPlatform.Unknown)
+			{
+				foreach (var IniFileName in EnumeratePlatformSpecificIniFileNames(ProjectDirectory, EngineDirectory, Platform, BaseIniName))
+				{
+					if (File.Exists(IniFileName))
+					{
+						ParseIniFile(IniFileName);
+					}
+				}
+			}
 		}
 
 		/// <summary>
@@ -442,10 +470,8 @@ namespace UnrealBuildTool
 		/// <summary>
 		/// Returns a list of INI filenames for the given project
 		/// </summary>
-		private static IEnumerable<string> EnumerateIniFileNames(string ProjectDirectory, string EngineDirectory, UnrealTargetPlatform Platform, string BaseIniName)
+		private static IEnumerable<string> EnumerateCrossPlatformIniFileNames(string ProjectDirectory, string EngineDirectory, string BaseIniName)
 		{
-			string PlatformName = GetIniPlatformName(Platform);
-
 			// Engine/Config/Base.ini (included in every ini type, required)
 			yield return Path.Combine(EngineDirectory, "Config", "Base.ini");
 
@@ -454,13 +480,21 @@ namespace UnrealBuildTool
 
 			// Game/Config/Default* ini
 			yield return Path.Combine(ProjectDirectory, "Config", "Default" + BaseIniName + ".ini");
+		}
+
+		/// <summary>
+		/// Returns a list of INI filenames for the given project
+		/// </summary>
+		private static IEnumerable<string> EnumeratePlatformSpecificIniFileNames(string ProjectDirectory, string EngineDirectory, UnrealTargetPlatform Platform, string BaseIniName)
+		{
+			string PlatformName = GetIniPlatformName(Platform);			
 
 			// Engine/Config/Platform/Platform* ini
 			yield return Path.Combine(EngineDirectory, "Config", PlatformName, PlatformName + BaseIniName + ".ini");
 
 			// Game/Config/Platform/Platform* ini
 			yield return Path.Combine(ProjectDirectory, "Config", PlatformName, PlatformName + BaseIniName + ".ini");
-		}
+		}		
 
 		/// <summary>
 		/// Returns the platform name to use as part of platform-specific config files
