Resolved unused variable (InternationalizationPreset) compilation warning-error in BuildCookRun.

diff --git a/Engine/Source/Programs/AutomationTool/Scripts/BuildCookRun.Automation.cs b/Engine/Source/Programs/AutomationTool/Scripts/BuildCookRun.Automation.cs
--- a/Engine/Source/Programs/AutomationTool/Scripts/BuildCookRun.Automation.cs
+++ b/Engine/Source/Programs/AutomationTool/Scripts/BuildCookRun.Automation.cs
@@ -107,6 +107,10 @@ public class BuildCookRun : BuildCommand
 		}
 
         var InternationalizationPreset = ParseParamValue("i18npreset");
+        if (!String.IsNullOrEmpty(InternationalizationPreset))
+        {
+            Params.InternationalizationPreset = InternationalizationPreset;
+        }
 
         var CulturesToCook = ParseParamValue("cookcultures");
         if (!String.IsNullOrEmpty(CulturesToCook))
