Add a UAT script which builds the common tools needed to run, which Epic devs usually sync from Perforce. Indended as a shortcut for Git users, and to build binaries for promotions on Git.

diff --git a/Engine/Source/Programs/AutomationTool/Scripts/BuildCommonTools.Automation.cs b/Engine/Source/Programs/AutomationTool/Scripts/BuildCommonTools.Automation.cs
new file mode 100644
--- /dev/null
+++ b/Engine/Source/Programs/AutomationTool/Scripts/BuildCommonTools.Automation.cs
@@ -0,0 +1,88 @@
+﻿using System;
+using System.Collections.Generic;
+using System.Linq;
+using System.Text;
+using AutomationTool;
+using UnrealBuildTool;
+
+[Help("Builds common tools used by the engine which are not part of typical editor or game builds. Useful when syncing source-only on GitHub.")]
+[Help("platforms=<X>+<Y>+...", "Specifies on or more platforms to build for (defaults to the current host platform)")]
+[Help("manifest=<Path>", "Writes a manifest of all the build products to the given path")]
+public class BuildCommonTools : BuildCommand
+{
+	public override void ExecuteBuild()
+	{
+		Log("************************* BuildCommonTools");
+
+		// Get the list of platform names
+		string[] PlatformNames = ParseParamValue("platforms", BuildHostPlatform.Current.ToString()).Split('+');
+
+		// Parse the platforms
+		List<UnrealBuildTool.UnrealTargetPlatform> Platforms = new List<UnrealTargetPlatform>();
+		foreach(string PlatformName in PlatformNames)
+		{
+			UnrealBuildTool.UnrealTargetPlatform Platform;
+			if(!UnrealBuildTool.UnrealTargetPlatform.TryParse(PlatformName, out Platform))
+			{
+				throw new AutomationException("Unknown platform specified on command line - '{0}' - valid platforms are {1}", PlatformName, String.Join("/", Enum.GetNames(typeof(UnrealBuildTool.UnrealTargetPlatform))));
+			}
+			Platforms.Add(Platform);
+		}
+
+		// Get the agenda
+		UE4Build.BuildAgenda Agenda = MakeAgenda(Platforms.ToArray());
+
+		// Build everything. We don't want to touch version files for GitHub builds -- these are "programmer builds" and won't have a canonical build version
+		UE4Build Builder = new UE4Build(this);
+		Builder.Build(Agenda, InDeleteBuildProducts:true, InUpdateVersionFiles: false);
+
+		// Make sure all the build products exist
+		UE4Build.CheckBuildProducts(Builder.BuildProductFiles);
+
+		// Write the manifest if needed
+		string ManifestPath = ParseParamValue("manifest");
+		if(ManifestPath != null)
+		{
+			UnrealBuildTool.FileManifest Manifest = new UnrealBuildTool.FileManifest();
+			foreach(string BuildProductFile in Builder.BuildProductFiles)
+			{
+				Manifest.AddFileName(BuildProductFile);
+			}
+			UnrealBuildTool.Utils.WriteClass(Manifest, ManifestPath, "");
+		}
+	}
+
+	public static UE4Build.BuildAgenda MakeAgenda(UnrealBuildTool.UnrealTargetPlatform[] Platforms)
+	{
+		// Create the build agenda
+		UE4Build.BuildAgenda Agenda = new UE4Build.BuildAgenda();
+
+		// C# binaries
+		Agenda.SwarmProject = @"Engine\Source\Programs\UnrealSwarm\UnrealSwarm.sln";
+		Agenda.DotNetProjects.Add(@"Engine/Source/Programs/DotNETCommon/DotNETUtilities/DotNETUtilities.csproj");
+		Agenda.DotNetProjects.Add(@"Engine/Source/Programs/UnrealControls/UnrealControls.csproj");
+
+		// Windows binaries
+		if(Platforms.Contains(UnrealBuildTool.UnrealTargetPlatform.Win64))
+		{
+			Agenda.AddTarget("UnrealHeaderTool", UnrealBuildTool.UnrealTargetPlatform.Win64, UnrealBuildTool.UnrealTargetConfiguration.Development);
+			Agenda.AddTarget("UnrealPak", UnrealBuildTool.UnrealTargetPlatform.Win64, UnrealBuildTool.UnrealTargetConfiguration.Development);
+			Agenda.AddTarget("UnrealLightmass", UnrealBuildTool.UnrealTargetPlatform.Win64, UnrealBuildTool.UnrealTargetConfiguration.Development);
+			Agenda.AddTarget("ShaderCompileWorker", UnrealBuildTool.UnrealTargetPlatform.Win64, UnrealBuildTool.UnrealTargetConfiguration.Development);
+			Agenda.AddTarget("UnrealVersionSelector", UnrealBuildTool.UnrealTargetPlatform.Win64, UnrealBuildTool.UnrealTargetConfiguration.Shipping);
+			Agenda.AddTarget("BootstrapPackagedGame", UnrealBuildTool.UnrealTargetPlatform.Win64, UnrealBuildTool.UnrealTargetConfiguration.Shipping);
+			Agenda.AddTarget("BootstrapPackagedGame", UnrealBuildTool.UnrealTargetPlatform.Win32, UnrealBuildTool.UnrealTargetConfiguration.Shipping);
+		}
+
+		// Mac binaries
+		if(Platforms.Contains(UnrealBuildTool.UnrealTargetPlatform.Mac))
+		{
+			Agenda.AddTarget("UnrealPak", UnrealBuildTool.UnrealTargetPlatform.Mac, UnrealBuildTool.UnrealTargetConfiguration.Development, InAddArgs: "-CopyAppBundleBackToDevice");
+			Agenda.AddTarget("UnrealLightmass", UnrealBuildTool.UnrealTargetPlatform.Mac, UnrealBuildTool.UnrealTargetConfiguration.Development, InAddArgs: "-CopyAppBundleBackToDevice");
+			Agenda.AddTarget("ShaderCompileWorker", UnrealBuildTool.UnrealTargetPlatform.Mac, UnrealBuildTool.UnrealTargetConfiguration.Development, InAddArgs: "-CopyAppBundleBackToDevice");
+			Agenda.AddTarget("UE4EditorServices", UnrealBuildTool.UnrealTargetPlatform.Mac, UnrealBuildTool.UnrealTargetConfiguration.Development, InAddArgs: "-CopyAppBundleBackToDevice");
+		}
+
+		return Agenda;
+	}
+}
