Fix command line for signing executables where there's a space in the path.

diff --git a/Engine/Source/Programs/AutomationTool/CommandUtils.cs b/Engine/Source/Programs/AutomationTool/CommandUtils.cs
--- a/Engine/Source/Programs/AutomationTool/CommandUtils.cs
+++ b/Engine/Source/Programs/AutomationTool/CommandUtils.cs
@@ -2194,7 +2194,7 @@ namespace AutomationTool
 			//@TODO: Verbosity choosing
 			//  /v will spew lots of info
 			//  /q does nothing on success and minimal output on failure
-			string CodeSignArgs = String.Format("sign{0} /a /n \"{1}\" /t {2} /v {3}", SpecificStoreArg, SigningIdentity, TimestampServer, TargetFileInfo.FullName);
+			string CodeSignArgs = String.Format("sign{0} /a /n \"{1}\" /t {2} /v \"{3}\"", SpecificStoreArg, SigningIdentity, TimestampServer, TargetFileInfo.FullName);
 
 			DateTime StartTime = DateTime.Now;
 
