Fixed HTML5 Packaging failing because of missing funciton.

diff --git a/Engine/Source/Programs/AutomationTool/HTML5/HTML5Platform.Automation.cs b/Engine/Source/Programs/AutomationTool/HTML5/HTML5Platform.Automation.cs
--- a/Engine/Source/Programs/AutomationTool/HTML5/HTML5Platform.Automation.cs
+++ b/Engine/Source/Programs/AutomationTool/HTML5/HTML5Platform.Automation.cs
@@ -249,6 +249,10 @@ public class HTML5Platform : Platform
 		}
 	}
 
+	public override void GetFilesToDeployOrStage(ProjectParams Params, DeploymentContext SC)
+	{
+	}
+
 	public override void GetFilesToArchive(ProjectParams Params, DeploymentContext SC)
 	{
 		if (SC.StageTargetConfigurations.Count != 1)
