Added nicer error when staging fails to copy from launch on.

diff --git a/Engine/Source/Editor/UnrealEd/Private/Analytics/EditorAnalytics.cpp b/Engine/Source/Editor/UnrealEd/Private/Analytics/EditorAnalytics.cpp
--- a/Engine/Source/Editor/UnrealEd/Private/Analytics/EditorAnalytics.cpp
+++ b/Engine/Source/Editor/UnrealEd/Private/Analytics/EditorAnalytics.cpp
@@ -160,6 +160,8 @@ FString FEditorAnalytics::TranslateErrorCode(int32 ErrorCode)
 		return TEXT("Android build-tools directory not found");
 	case EAnalyticsErrorCodes::NoApkSuitableForArchitecture:
 		return TEXT("No APK suitable for architecture found");
+	case EAnalyticsErrorCodes::FailedToDeleteStagingDirectory :
+		return TEXT("Failed to delete staging directory.  This could be because something is currently using the staging directory (ps4/xbox/etc)");
 	}
 	return TEXT("Unknown Error");
 }
diff --git a/Engine/Source/Editor/UnrealEd/Public/EditorAnalytics.h b/Engine/Source/Editor/UnrealEd/Public/EditorAnalytics.h
--- a/Engine/Source/Editor/UnrealEd/Public/EditorAnalytics.h
+++ b/Engine/Source/Editor/UnrealEd/Public/EditorAnalytics.h
@@ -55,6 +55,7 @@ namespace EAnalyticsErrorCodes
 		NoApkSuitableForArchitecture = 55,
 		LauncherFailed = 100,
 		UATLaunchFailure = 101,
+		FailedToDeleteStagingDirectory = 102,
 	};
 };
 
diff --git a/Engine/Source/Programs/AutomationTool/Program.cs b/Engine/Source/Programs/AutomationTool/Program.cs
--- a/Engine/Source/Programs/AutomationTool/Program.cs
+++ b/Engine/Source/Programs/AutomationTool/Program.cs
@@ -56,6 +56,7 @@ namespace AutomationTool
         Error_NoApkSuitableForArchitecture = 55,
         Error_LauncherFailed = 100,
 		Error_UATLaunchFailure = 101,
+        Error_FailedToDeleteStagingDirectory = 102,
 	};
 
 	public class ErrorReporter
diff --git a/Engine/Source/Programs/AutomationTool/Scripts/CopyBuildToStagingDirectory.Automation.cs b/Engine/Source/Programs/AutomationTool/Scripts/CopyBuildToStagingDirectory.Automation.cs
--- a/Engine/Source/Programs/AutomationTool/Scripts/CopyBuildToStagingDirectory.Automation.cs
+++ b/Engine/Source/Programs/AutomationTool/Scripts/CopyBuildToStagingDirectory.Automation.cs
@@ -764,7 +764,17 @@ public partial class Project : CommandUtils
 		MaybeConvertToLowerCase(Params, SC);
 		if (SC.Stage && !Params.NoCleanStage && !Params.SkipStage)
 		{
-			DeleteDirectory(SC.StageDirectory);
+            try
+            {
+                DeleteDirectory(SC.StageDirectory);
+            }
+            catch (Exception Ex)
+            {
+                // Delete cooked data (if any) as it may be incomplete / corrupted.
+                Log("Failed to delete staging directory "+SC.StageDirectory);
+                AutomationTool.ErrorReporter.Error("Stage Failed.", (int)AutomationTool.ErrorCodes.Error_FailedToDeleteStagingDirectory);
+                throw Ex;   
+            }
 		}
 		if (ShouldCreatePak(Params, SC))
 		{
diff --git a/Engine/Source/Programs/IOS/iPhonePackager/iPhonePackager.cs b/Engine/Source/Programs/IOS/iPhonePackager/iPhonePackager.cs
--- a/Engine/Source/Programs/IOS/iPhonePackager/iPhonePackager.cs
+++ b/Engine/Source/Programs/IOS/iPhonePackager/iPhonePackager.cs
@@ -62,7 +62,8 @@ namespace iPhonePackager
         Error_AndroidBuildToolsPathNotFound = 54,
         Error_NoApkSuitableForArchitecture = 55,
 		Error_LauncherFailed = 100,
-		Error_UATLaunchFailure = 101,
+        Error_UATLaunchFailure = 101,
+        Error_FailedToDeleteStagingDirectory = 102,
 	};
 
 	public partial class Program
