Add option to disable spew for various p4 commands.

diff --git a/Engine/Source/Programs/AutomationTool/P4Utils.cs b/Engine/Source/Programs/AutomationTool/P4Utils.cs
--- a/Engine/Source/Programs/AutomationTool/P4Utils.cs
+++ b/Engine/Source/Programs/AutomationTool/P4Utils.cs
@@ -613,7 +613,7 @@ namespace AutomationTool
             return Result;
         }
         static Dictionary<string, List<ChangeRecord>> ChangesCache = new Dictionary<string, List<ChangeRecord>>();
-        public bool Changes(out List<ChangeRecord> ChangeRecords, string CommandLine, bool AllowSpew = true, bool UseCaching = false, bool LongComment = false)
+        public bool Changes(out List<ChangeRecord> ChangeRecords, string CommandLine, bool AllowSpew = true, bool UseCaching = false, bool LongComment = false, bool WithClient = false)
         {
             // If the user specified '-l' or '-L', the summary will appear on subsequent lines (no quotes) instead of the same line (surrounded by single quotes)
             bool ContainsDashL = CommandLine.StartsWith("-L ", StringComparison.InvariantCultureIgnoreCase) ||
@@ -636,7 +636,7 @@ namespace AutomationTool
                 // Change 1999345 on 2014/02/16 by buildmachine@BuildFarm_BUILD-23_buildmachine_++depot+UE4 'GUBP Node Shadow_LabelPromotabl'
 
                 string Output;
-                if (!LogP4Output(out Output, "changes " + CommandLine, null, AllowSpew))
+                if (!LogP4Output(out Output, "changes " + CommandLine, null, AllowSpew, WithClient: WithClient))
                 {
                     throw new AutomationException("P4 returned failure.");
                 }
@@ -1009,10 +1009,10 @@ namespace AutomationTool
         /// <param name="FromCL">Changelist to unshelve.</param>
         /// <param name="ToCL">Changelist where the checked out files should be added.</param>
         /// <param name="CommandLine">Commandline for the command.</param>
-        public void Shelve(int FromCL, string CommandLine = "")
+        public void Shelve(int FromCL, string CommandLine = "", bool AllowSpew = true)
         {
             CheckP4Enabled();
-            LogP4("shelve " + String.Format("-r -c {0} ", FromCL) + CommandLine);
+            LogP4("shelve " + String.Format("-r -c {0} ", FromCL) + CommandLine, AllowSpew: AllowSpew);
         }
 		/// <summary>
 		/// Invokes p4 edit command.
@@ -1091,10 +1091,10 @@ namespace AutomationTool
 		/// </summary>
 		/// <param name="CL">Changelist to check the files out.</param>
 		/// <param name="CommandLine">Commandline for the command.</param>
-		public void ReconcileNoDeletes(int CL, string CommandLine)
+		public void ReconcileNoDeletes(int CL, string CommandLine, bool AllowSpew = true)
 		{
 			CheckP4Enabled();
-			LogP4("reconcile " + String.Format("-c {0} -ea ", CL) + CommandLine);
+			LogP4("reconcile " + String.Format("-c {0} -ea ", CL) + CommandLine, AllowSpew: AllowSpew);
 		}
 
 		/// <summary>
@@ -1124,10 +1124,10 @@ namespace AutomationTool
 		/// </summary>
 		/// <param name="CL">Changelist to revert</param>
 		/// <param name="CommandLine">Commandline for the command.</param>
-		public void Revert(int CL, string CommandLine = "")
+		public void Revert(int CL, string CommandLine = "", bool AllowSpew = true)
 		{
 			CheckP4Enabled();
-			LogP4("revert " + String.Format("-c {0} ", CL) + CommandLine);
+			LogP4("revert " + String.Format("-c {0} ", CL) + CommandLine, AllowSpew: AllowSpew);
 		}
 
 		/// <summary>
@@ -1325,7 +1325,7 @@ namespace AutomationTool
 		/// <param name="Owner">Owner of the changelist.</param>
 		/// <param name="Description">Description of the changelist.</param>
 		/// <returns>Id of the created changelist.</returns>
-		public int CreateChange(string Owner = null, string Description = null)
+		public int CreateChange(string Owner = null, string Description = null, bool AllowSpew = false)
 		{
 			CheckP4Enabled();
 			var ChangeSpec = "Change: new" + "\n";
@@ -1333,8 +1333,11 @@ namespace AutomationTool
 			ChangeSpec += "Description: \n " + ((Description != null) ? Description : "(none)") + "\n";
 			string CmdOutput;
 			int CL = 0;
-			CommandUtils.Log(TraceEventType.Information, "Creating Change\n {0}\n", ChangeSpec);
-			if (LogP4Output(out CmdOutput, "change -i", Input: ChangeSpec))
+			if(AllowSpew)
+			{
+				CommandUtils.Log(TraceEventType.Information, "Creating Change\n {0}\n", ChangeSpec);
+			}
+			if (LogP4Output(out CmdOutput, "change -i", Input: ChangeSpec, AllowSpew: AllowSpew))
 			{
 				string EndStr = " created.";
 				string ChangeStr = "Change ";
@@ -1349,7 +1352,7 @@ namespace AutomationTool
 			{
 				throw new P4Exception("Failed to create Changelist. Owner: {0} Desc: {1}", Owner, Description);
 			}
-			else
+			else if(AllowSpew)
 			{
 				CommandUtils.Log(TraceEventType.Information, "Returned CL {0}\n", CL);
 			}
@@ -2027,10 +2030,13 @@ namespace AutomationTool
 		/// </summary>
 		/// <param name="ClientName">Client name</param>
 		/// <returns>True if the client exists.</returns>
-		public bool DoesClientExist(string ClientName)
+		public bool DoesClientExist(string ClientName, bool Quiet = false)
 		{
 			CheckP4Enabled();
-			CommandUtils.Log("Checking if client {0} exists", ClientName);
+			if(!Quiet)
+			{
+				CommandUtils.Log("Checking if client {0} exists", ClientName);
+			}
 
             var P4Result = P4(String.Format("-c {0} where //...", ClientName), AllowSpew: false, WithClient: false);
             return P4Result.Output.IndexOf("unknown - use 'client' command", StringComparison.InvariantCultureIgnoreCase) < 0 && P4Result.Output.IndexOf("doesn't exist", StringComparison.InvariantCultureIgnoreCase) < 0;
@@ -2041,12 +2047,15 @@ namespace AutomationTool
 		/// </summary>
 		/// <param name="ClientName">Name of the client.</param>
 		/// <returns></returns>
-		public P4ClientInfo GetClientInfo(string ClientName)
+		public P4ClientInfo GetClientInfo(string ClientName, bool Quiet = false)
 		{
 			CheckP4Enabled();
 
-			CommandUtils.Log("Getting info for client {0}", ClientName);
-			if (!DoesClientExist(ClientName))
+			if(!Quiet)
+			{
+				CommandUtils.Log("Getting info for client {0}", ClientName);
+			}
+			if (!DoesClientExist(ClientName, Quiet))
 			{
 				return null;
 			}
@@ -2213,7 +2222,7 @@ namespace AutomationTool
         /// </summary>
         /// <param name="ClientSpec">Client specification.</param>
         /// <returns></returns>
-        public P4ClientInfo CreateClient(P4ClientInfo ClientSpec)
+        public P4ClientInfo CreateClient(P4ClientInfo ClientSpec, bool AllowSpew = true)
         {
             string SpecInput = "Client: " + ClientSpec.Name + Environment.NewLine;
             SpecInput += "Owner: " + ClientSpec.Owner + Environment.NewLine;
@@ -2227,8 +2236,8 @@ namespace AutomationTool
             {
                 SpecInput += "\t" + Mapping.Key + " //" + ClientSpec.Name + Mapping.Value + Environment.NewLine;
             }
-			CommandUtils.Log(SpecInput);
-            LogP4("client -i", SpecInput, WithClient: false);
+			if(AllowSpew) CommandUtils.Log(SpecInput);
+            LogP4("client -i", SpecInput, AllowSpew: AllowSpew, WithClient: false);
             return ClientSpec;
         }
 
