Add utility function to add a target to a build agenda.

diff --git a/Engine/Source/Programs/AutomationTool/UE4Build.cs b/Engine/Source/Programs/AutomationTool/UE4Build.cs
--- a/Engine/Source/Programs/AutomationTool/UE4Build.cs
+++ b/Engine/Source/Programs/AutomationTool/UE4Build.cs
@@ -641,6 +641,32 @@ namespace AutomationTool
 
 			public bool DoRetries = IsBuildMachine;
 
+			/// <summary>
+			/// Adds a target with the specified configuration.
+			/// </summary>
+			/// <param name="TargetName">Name of the target</param>
+			/// <param name="InPlatform">Platform</param>
+			/// <param name="InConfiguration">Configuration</param>
+			/// <param name="InUprojectPath">Path to optional uproject file</param>
+			/// <param name="bForceMonolithic">Force monolithic build.</param>
+			/// <param name="bForceNonUnity">Force non-unity</param>
+			/// <param name="bForceDebugInfo">Force debug info even in development builds</param>
+			/// <param name="InAddArgs">Specifies additional arguments for UBT</param>
+			public void AddTarget(string TargetName, UnrealBuildTool.UnrealTargetPlatform InPlatform, UnrealBuildTool.UnrealTargetConfiguration InConfiguration, string InUprojectPath = null, bool bForceMonolithic = false, bool bForceNonUnity = false, bool bForceDebugInfo = false, string InAddArgs = "")
+			{
+				Targets.Add(new BuildTarget()
+				{
+					TargetName = TargetName,
+					Platform = InPlatform,
+					Config = InConfiguration,
+					UprojectPath = InUprojectPath,
+					ForceMonolithic = bForceMonolithic,
+					ForceNonUnity = bForceNonUnity,
+					ForceDebugInfo = bForceDebugInfo,
+					UBTArgs = InAddArgs,
+				});
+			}
+
 			/// <summary>
 			/// Adds multiple targets with the specified configuration.
 			/// </summary>
