Add BuildCommonTools to UAT project.

diff --git a/Engine/Source/Programs/AutomationTool/Scripts/AutomationScripts.Automation.csproj b/Engine/Source/Programs/AutomationTool/Scripts/AutomationScripts.Automation.csproj
--- a/Engine/Source/Programs/AutomationTool/Scripts/AutomationScripts.Automation.csproj
+++ b/Engine/Source/Programs/AutomationTool/Scripts/AutomationScripts.Automation.csproj
@@ -61,6 +61,7 @@
     <Compile Include="AnalyzeThirdPartyLibs.Automation.cs" />
     <Compile Include="ArchiveCommand.Automation.cs" />
     <Compile Include="BlameKeyword.Automation.cs" />
+    <Compile Include="BuildCommonTools.Automation.cs" />
     <Compile Include="BuildCookRun.Automation.cs" />
     <Compile Include="BuildDocumentation.Automation.cs" />
     <Compile Include="BuildProjectCommand.Automation.cs" />
