Allow html sdk path to be set using ini file.

   - takes preference over EMSCRIPTEN env variable
   - editor can now detect sdk on yosemite.
   - refactor and simplyfy emscripten sdk access in .cs files.

diff --git a/Engine/Config/HTML5/HTML5Engine.ini b/Engine/Config/HTML5/HTML5Engine.ini
--- a/Engine/Config/HTML5/HTML5Engine.ini
+++ b/Engine/Config/HTML5/HTML5Engine.ini
@@ -25,3 +25,6 @@ Firefox (Aurora)="/Applications/FirefoxAurora.app/Contents/MacOS/firefox-bin"
 
 [HTML5DevicesLinux]
 Firefox="/usr/bin/firefox"
+
+[HTML5SDKPaths]
+Mac="/path/to/emscripten/sdk"
diff --git a/Engine/Source/Developer/HTML5/HTML5TargetPlatform/Private/HTML5TargetPlatform.cpp b/Engine/Source/Developer/HTML5/HTML5TargetPlatform/Private/HTML5TargetPlatform.cpp
--- a/Engine/Source/Developer/HTML5/HTML5TargetPlatform/Private/HTML5TargetPlatform.cpp
+++ b/Engine/Source/Developer/HTML5/HTML5TargetPlatform/Private/HTML5TargetPlatform.cpp
@@ -82,11 +82,36 @@ ITargetDevicePtr FHTML5TargetPlatform::GetDevice( const FTargetDeviceId& DeviceI
 
 bool FHTML5TargetPlatform::IsSdkInstalled(bool bProjectHasCode, FString& OutDocumentationPath) const
 {
-	bool bHTML5SDKInstalled = false; // @todo How do we check that the iOS SDK is installed when building from Windows? Is that even possible?
-	TCHAR BaseSDKPath[256];
+	FString SectionName = "HTML5SDKPaths";
+	FConfigSection SDKPaths = HTML5EngineSettings[SectionName];
+	for (auto It : SDKPaths )
+	{
+		const FString& Platform = It.Key.ToString();
+		const FString& Path = It.Value;
+		{
+#if PLATFORM_WINDOWS
+			if ( Platform == "Windows" && IFileManager::Get().DirectoryExists(*Path)) 
+			{
+				return true; 
+			}
+#endif 
+#if PLATFORM_MAC
+			if ( Platform == "Mac" && IFileManager::Get().DirectoryExists(*Path)) 
+			{
+				return true; 
+			}
+#endif 
+		}
+	}
+
+	TCHAR BaseSDKPath[512];
 	FPlatformMisc::GetEnvironmentVariable(TEXT("EMSCRIPTEN"), BaseSDKPath, ARRAY_COUNT(BaseSDKPath));
-	bHTML5SDKInstalled = FString(BaseSDKPath).Len() > 0;
-	return bHTML5SDKInstalled;
+	if (FString(BaseSDKPath).Len() > 0 && IFileManager::Get().DirectoryExists(BaseSDKPath) )
+	{
+		return true; 	
+	}
+
+	return false; 
 }
 
 
diff --git a/Engine/Source/Programs/AutomationTool/HTML5/HTML5Platform.Automation.cs b/Engine/Source/Programs/AutomationTool/HTML5/HTML5Platform.Automation.cs
--- a/Engine/Source/Programs/AutomationTool/HTML5/HTML5Platform.Automation.cs
+++ b/Engine/Source/Programs/AutomationTool/HTML5/HTML5Platform.Automation.cs
@@ -14,35 +14,6 @@ public class HTML5Platform : Platform
 {
     static string EmscriptenSettingsPath = CombinePaths(Environment.GetFolderPath(Environment.SpecialFolder.UserProfile), "/.emscripten");
 
-    Dictionary<string, string> ReadEmscriptenSettings()
-    {
-        // Check HTML5ToolChain.cs for duplicate
-        if (!System.IO.File.Exists(EmscriptenSettingsPath))
-        {
-            return new Dictionary<string, string>();
-        }
-
-        Dictionary<string, string> Settings = new Dictionary<string, string>();
-        System.IO.StreamReader SettingFile = new System.IO.StreamReader(EmscriptenSettingsPath);
-        string EMLine = null;
-        while ((EMLine = SettingFile.ReadLine()) != null)
-        {
-            EMLine = EMLine.Split('#')[0];
-            string Pattern1 = @"(\w*)\s*=\s*['\[]?([^'\]\r\n]*)['\]]?";
-            Regex Rgx = new Regex(Pattern1, RegexOptions.IgnoreCase);
-            MatchCollection Matches = Rgx.Matches(EMLine);
-            foreach (Match Matched in Matches)
-            {
-                if (Matched.Groups.Count == 3 && Matched.Groups[2].ToString() != "")
-                {
-                    Settings[Matched.Groups[1].ToString()] = Matched.Groups[2].ToString();
-                }
-            }
-        }
-
-        return Settings;
-    }
-
 	public HTML5Platform()
 		: base(UnrealTargetPlatform.HTML5)
 	{
@@ -51,8 +22,7 @@ public class HTML5Platform : Platform
 	public override void Package(ProjectParams Params, DeploymentContext SC, int WorkingCL)
 	{
 		Log("Package {0}", Params.RawProjectPath);
-        var EmscriptenSettings = ReadEmscriptenSettings();
-
+        
 		string PackagePath = Path.Combine(Path.GetDirectoryName(Params.RawProjectPath), "Binaries", "HTML5");
 		if (!Directory.Exists(PackagePath))
 		{
@@ -60,119 +30,66 @@ public class HTML5Platform : Platform
 		}
 		string FinalDataLocation = Path.Combine(PackagePath, Params.ShortProjectName) + ".data";
 
+		    // we need to operate in the root
+            using (new PushedDirectory(Path.Combine(Params.BaseStageDirectory, "HTML5")))
+            {
+                string PythonPath = HTML5SDKInfo.PythonPath();
+                string PackagerPath = HTML5SDKInfo.EmscriptenPackager();
 
-		// we need to operate in the root
-		using (new PushedDirectory(Path.Combine(Params.BaseStageDirectory, "HTML5")))
-		{
-			string BaseSDKPath = Environment.GetEnvironmentVariable("EMSCRIPTEN");
-            string PythonPath = null;
+                string CmdLine = string.Format("{0} \"{1}\" --preload . --js-output=\"{1}.js\"", PackagerPath, FinalDataLocation);
+                RunAndLog(CmdEnv, PythonPath, CmdLine);
+            }
 
-            // Check the .emscripten file for a possible python path
-            if (EmscriptenSettings.ContainsKey("PYTHON"))
+            // copy the "Executable" to the package directory
+            string GameExe = Path.GetFileNameWithoutExtension(Params.ProjectGameExeFilename);
+            if (Params.ClientConfigsToBuild[0].ToString() != "Development")
             {
-                PythonPath = EmscriptenSettings["PYTHON"];
-                Log("Found python path {0} in emscripten file", PythonPath);
+                GameExe += "-HTML5-" + Params.ClientConfigsToBuild[0].ToString();
             }
-            // The AutoSDK defines this env var as part of its install. See setup.bat/unsetup.bat
-            // If it's missing then just assume that python lives on the path
-            if (PythonPath == null && Environment.GetEnvironmentVariable("PYTHON") != null)
+            GameExe += ".js";
+            if (Path.Combine(Path.GetDirectoryName(Params.ProjectGameExeFilename), GameExe) != Path.Combine(PackagePath, GameExe))
             {
-                PythonPath = Environment.GetEnvironmentVariable("PYTHON");
-                Log("Found python path {0} in PYTHON Environment Variable", PythonPath);
+                File.Copy(Path.Combine(Path.GetDirectoryName(Params.ProjectGameExeFilename), GameExe), Path.Combine(PackagePath, GameExe), true);
+                File.Copy(Path.Combine(Path.GetDirectoryName(Params.ProjectGameExeFilename), GameExe) + ".mem", Path.Combine(PackagePath, GameExe) + ".mem", true);
             }
+            File.SetAttributes(Path.Combine(PackagePath, GameExe), FileAttributes.Normal);
+            File.SetAttributes(Path.Combine(PackagePath, GameExe) + ".mem", FileAttributes.Normal);
 
-            // Check if the exe exists
-            if (!System.IO.File.Exists(PythonPath))
-            {
-                PythonPath = null;
-            }
+            // put the HTML file to the package directory
+            string TemplateFile = Path.Combine(CombinePaths(CmdEnv.LocalRoot, "Engine"), "Build", "HTML5", "Game.html.template");
+            string OutputFile = Path.Combine(PackagePath, (Params.ClientConfigsToBuild[0].ToString() != "Development" ? (Params.ShortProjectName + "-HTML5-" + Params.ClientConfigsToBuild[0].ToString()) : Params.ShortProjectName)) + ".html";
 
-            if (PythonPath == null)
-            {
-                Log("Either no python path can be found or it doesn't exist. Using python on PATH");
-            }
+            // find Heap Size.
+            ulong HeapSize;
+            var ConfigCache = new UnrealBuildTool.ConfigCacheIni(UnrealTargetPlatform.HTML5, "Engine", Path.GetDirectoryName(Params.RawProjectPath), CombinePaths(CmdEnv.LocalRoot, "Engine"));
 
-			// make the file_packager command line
-			if (Utils.IsRunningOnMono)
-			{
-				string PackagerPath = BaseSDKPath + "/tools/file_packager.py";
-                if (PythonPath == null)
-                {
-                    string CmdLine = string.Format("-c \" python {0} '{1}' --preload . --js-output='{1}.js' \" ", PackagerPath, FinalDataLocation);
-                    RunAndLog(CmdEnv, "/bin/bash", CmdLine);
-                }
-                else
-                {
-                    string CmdLine = string.Format("{0} '{1}' --preload . --js-output='{1}.js' ", PackagerPath, FinalDataLocation);
-                    RunAndLog(CmdEnv, PythonPath, CmdLine);
-                }
-			}
-			else
-			{
-				string PackagerPath = "\"" + BaseSDKPath + "\\tools\\file_packager.py\"";
-                if (PythonPath == null)
-                {
-                    string CmdLine = string.Format("/c python {0} \"{1}\" --preload . --js-output=\"{1}.js\"", PackagerPath, FinalDataLocation);
-                    RunAndLog(CmdEnv, CommandUtils.CombinePaths(Environment.SystemDirectory, "cmd.exe"), CmdLine);
-                }
-                else 
+            int ConfigHeapSize;
+            if (!ConfigCache.GetInt32("BuildSettings", "HeapSize" + Params.ClientConfigsToBuild[0].ToString(), out ConfigHeapSize)) // in Megs.
+            {
+                // we couldn't find a per config heap size, look for a common one.
+                if (!ConfigCache.GetInt32("BuildSettings", "HeapSize", out ConfigHeapSize))
                 {
-                    string CmdLine = string.Format("{0} \"{1}\" --preload . --js-output=\"{1}.js\"", PackagerPath, FinalDataLocation);
-                    RunAndLog(CmdEnv, PythonPath, CmdLine);
+                    ConfigHeapSize = Params.IsCodeBasedProject ? 1024 : 512;
+                    Log("Could not find Heap Size setting in .ini for Client config {0}", Params.ClientConfigsToBuild[0].ToString());
                 }
-			}
-		}
-
-		// copy the "Executable" to the package directory
-		string GameExe = Path.GetFileNameWithoutExtension(Params.ProjectGameExeFilename);
-		if (Params.ClientConfigsToBuild[0].ToString() != "Development")
-		{
-			GameExe += "-HTML5-" + Params.ClientConfigsToBuild[0].ToString();
-		}
-		GameExe += ".js";
-		if (Path.Combine(Path.GetDirectoryName(Params.ProjectGameExeFilename), GameExe) != Path.Combine(PackagePath, GameExe))
-		{
-			File.Copy(Path.Combine(Path.GetDirectoryName(Params.ProjectGameExeFilename), GameExe), Path.Combine(PackagePath, GameExe), true);
-			File.Copy(Path.Combine(Path.GetDirectoryName(Params.ProjectGameExeFilename), GameExe) + ".mem", Path.Combine(PackagePath, GameExe) + ".mem", true);
-		}
-		File.SetAttributes(Path.Combine(PackagePath, GameExe), FileAttributes.Normal);
-		File.SetAttributes(Path.Combine(PackagePath, GameExe) + ".mem", FileAttributes.Normal);
-
-		// put the HTML file to the package directory
-		string TemplateFile = Path.Combine(CombinePaths(CmdEnv.LocalRoot, "Engine"), "Build", "HTML5", "Game.html.template");
-		string OutputFile = Path.Combine(PackagePath, (Params.ClientConfigsToBuild[0].ToString() != "Development" ? (Params.ShortProjectName + "-HTML5-" + Params.ClientConfigsToBuild[0].ToString()) : Params.ShortProjectName)) + ".html";
-
-		// find Heap Size.
-		ulong HeapSize;
-		var ConfigCache = new UnrealBuildTool.ConfigCacheIni(UnrealTargetPlatform.HTML5, "Engine", Path.GetDirectoryName(Params.RawProjectPath), CombinePaths(CmdEnv.LocalRoot, "Engine"));
-
-		int ConfigHeapSize;
-		if (!ConfigCache.GetInt32("BuildSettings", "HeapSize" + Params.ClientConfigsToBuild[0].ToString(), out ConfigHeapSize)) // in Megs.
-		{
-			// we couldn't find a per config heap size, look for a common one.
-			if (!ConfigCache.GetInt32("BuildSettings", "HeapSize", out ConfigHeapSize))
-			{
-				ConfigHeapSize = Params.IsCodeBasedProject ? 1024 : 512;
-				Log("Could not find Heap Size setting in .ini for Client config {0}", Params.ClientConfigsToBuild[0].ToString());
-			}
-		}
+            }
 
-		HeapSize = (ulong)ConfigHeapSize * 1024L * 1024L; // convert to bytes.
-		Log("Setting Heap size to {0} Mb ", ConfigHeapSize);
+            HeapSize = (ulong)ConfigHeapSize * 1024L * 1024L; // convert to bytes.
+            Log("Setting Heap size to {0} Mb ", ConfigHeapSize);
 
 
-		GenerateFileFromTemplate(TemplateFile, OutputFile, Params.ShortProjectName, Params.ClientConfigsToBuild[0].ToString(), Params.StageCommandline, !Params.IsCodeBasedProject, HeapSize);
+            GenerateFileFromTemplate(TemplateFile, OutputFile, Params.ShortProjectName, Params.ClientConfigsToBuild[0].ToString(), Params.StageCommandline, !Params.IsCodeBasedProject, HeapSize);
 
-		// copy the jstorage files to the binaries directory
-		string JSDir = Path.Combine(CombinePaths(CmdEnv.LocalRoot, "Engine"), "Build", "HTML5");
-		string OutDir = PackagePath;
-		File.Copy(JSDir + "/json2.js", OutDir + "/json2.js", true);
-		File.SetAttributes(OutDir + "/json2.js", FileAttributes.Normal);
-		File.Copy(JSDir + "/jstorage.js", OutDir + "/jstorage.js", true);
-		File.SetAttributes(OutDir + "/jstorage.js", FileAttributes.Normal);
-		File.Copy(JSDir + "/moz_binarystring.js", OutDir + "/moz_binarystring.js", true);
-		File.SetAttributes(OutDir + "/moz_binarystring.js", FileAttributes.Normal);
-		PrintRunTime();
+            // copy the jstorage files to the binaries directory
+            string JSDir = Path.Combine(CombinePaths(CmdEnv.LocalRoot, "Engine"), "Build", "HTML5");
+            string OutDir = PackagePath;
+            File.Copy(JSDir + "/json2.js", OutDir + "/json2.js", true);
+            File.SetAttributes(OutDir + "/json2.js", FileAttributes.Normal);
+            File.Copy(JSDir + "/jstorage.js", OutDir + "/jstorage.js", true);
+            File.SetAttributes(OutDir + "/jstorage.js", FileAttributes.Normal);
+            File.Copy(JSDir + "/moz_binarystring.js", OutDir + "/moz_binarystring.js", true);
+            File.SetAttributes(OutDir + "/moz_binarystring.js", FileAttributes.Normal);
+            PrintRunTime();
 	}
 
 
@@ -331,36 +248,13 @@ public class HTML5Platform : Platform
 		}
 		else
 		{
-			var EmscriptenSettings = ReadEmscriptenSettings();
 			url = "http://127.0.0.1:8000/" + url;
 			// this will be killed UBT instances dies.
 			string input = String.Format(" -m SimpleHTTPServer 8000");
 
-
-			string PythonName = null;
-			// Check the .emscripten file for a possible python path
-			if (EmscriptenSettings.ContainsKey("PYTHON"))
-			{
-				PythonName = EmscriptenSettings["PYTHON"];
-				Log("Found python path {0} in emscripten file", PythonName);
-			}
-			// The AutoSDK defines this env var as part of its install. See setup.bat/unsetup.bat
-			// If it's missing then just assume that python lives on the path
-			if (PythonName == null && Environment.GetEnvironmentVariable("PYTHON") != null)
-			{
-				PythonName = Environment.GetEnvironmentVariable("PYTHON");
-				Log("Found python path {0} in PYTHON Environment Variable", PythonName);
-			}
-
-			// Check if the exe exists
-			if (!System.IO.File.Exists(PythonName))
-			{
-				Log("Either no python path can be found or it doesn't exist. Using python on PATH");
-				PythonName = Utils.IsRunningOnMono ? "python" : "python.exe";
-			}
-
-			ProcessResult Result = ProcessManager.CreateProcess(PythonName, true, "html5server.log");
-			Result.ProcessObject.StartInfo.FileName = PythonName;
+            string PythonPath = HTML5SDKInfo.PythonPath();
+            ProcessResult Result = ProcessManager.CreateProcess(PythonPath, true, "html5server.log");
+            Result.ProcessObject.StartInfo.FileName = PythonPath;
 			Result.ProcessObject.StartInfo.UseShellExecute = false;
 			Result.ProcessObject.StartInfo.RedirectStandardOutput = true;
 			Result.ProcessObject.StartInfo.RedirectStandardInput = true;
diff --git a/Engine/Source/Programs/UnrealBuildTool/HTML5/HTML5SDKInfo.cs b/Engine/Source/Programs/UnrealBuildTool/HTML5/HTML5SDKInfo.cs
new file mode 100644
--- /dev/null
+++ b/Engine/Source/Programs/UnrealBuildTool/HTML5/HTML5SDKInfo.cs
@@ -0,0 +1,169 @@
+﻿// Copyright 1998-2014 Epic Games, Inc. All Rights Reserved.
+using System;
+using System.IO; 
+using System.Text;
+using UnrealBuildTool;
+using System.Diagnostics;
+using System.Collections.Generic;
+using System.Text.RegularExpressions;
+
+namespace UnrealBuildTool
+{
+    public class HTML5SDKInfo
+    {
+        static string EmscriptenSettingsPath = Path.Combine(Environment.GetFolderPath(Environment.SpecialFolder.UserProfile), ".emscripten");
+
+        static Dictionary<string, string> ReadEmscriptenSettings()
+        {
+            // Check HTML5ToolChain.cs for duplicate
+            if (!System.IO.File.Exists(EmscriptenSettingsPath))
+            {
+                return new Dictionary<string, string>();
+            }
+
+            Dictionary<string, string> Settings = new Dictionary<string, string>();
+            System.IO.StreamReader SettingFile = new System.IO.StreamReader(EmscriptenSettingsPath);
+            string EMLine = null;
+            while ((EMLine = SettingFile.ReadLine()) != null)
+            {
+                EMLine = EMLine.Split('#')[0];
+                string Pattern1 = @"(\w*)\s*=\s*['\[]?([^'\]\r\n]*)['\]]?";
+                Regex Rgx = new Regex(Pattern1, RegexOptions.IgnoreCase);
+                MatchCollection Matches = Rgx.Matches(EMLine);
+                foreach (Match Matched in Matches)
+                {
+                    if (Matched.Groups.Count == 3 && Matched.Groups[2].ToString() != "")
+                    {
+                        Settings[Matched.Groups[1].ToString()] = Matched.Groups[2].ToString();
+                    }
+                }
+            }
+
+            return Settings;
+        }
+
+        public static string EmscriptenSDKPath()
+        {
+            // try to find SDK location in the ini file. 
+            var ConfigCache = new ConfigCacheIni(UnrealTargetPlatform.HTML5, "Engine", "");
+
+            string PlatformName = "";
+            if (!Utils.IsRunningOnMono)
+            {
+                PlatformName = "Windows";
+            }
+            else
+            {
+                PlatformName = "Mac";
+            }
+            string EMCCPath;
+            bool ok = ConfigCache.GetString("HTML5SDKPaths", PlatformName, out EMCCPath);
+
+            if (ok && System.IO.Directory.Exists(EMCCPath))
+                return EMCCPath;
+
+            // try to find SDK Location from env.
+            if (Environment.GetEnvironmentVariable("EMSCRIPTEN") != null
+                && System.IO.Directory.Exists(Environment.GetEnvironmentVariable("EMSCRIPTEN"))
+                )
+            {
+                return Environment.GetEnvironmentVariable("EMSCRIPTEN");
+            }
+
+            return ""; 
+        }
+
+        public  static string PythonPath()
+        {
+            // find Python.
+            var EmscriptenSettings = ReadEmscriptenSettings();
+
+			// check emscripten generated config file. 
+            if (EmscriptenSettings.ContainsKey("PYTHON")
+                && System.IO.File.Exists(EmscriptenSettings["PYTHON"])
+                )
+            {
+                return EmscriptenSettings["PYTHON"];
+            }
+
+            // It might be setup as a env variable. 
+            if (Environment.GetEnvironmentVariable("PYTHON") != null
+                &&
+                System.IO.File.Exists(Environment.GetEnvironmentVariable("PYTHON"))
+                )
+            {
+                return Environment.GetEnvironmentVariable("PYTHON");
+            }
+
+            // it might just be on path. 
+            ProcessStartInfo startInfo = new ProcessStartInfo();
+            Process PythonProcess = new Process();
+
+            startInfo.CreateNoWindow = true;
+            startInfo.RedirectStandardOutput = true;
+            startInfo.RedirectStandardInput = true;
+
+            startInfo.Arguments = "python";
+			startInfo.UseShellExecute = false;
+
+			if (!Utils.IsRunningOnMono) 
+			{
+				startInfo.FileName = "C:\\Windows\\System32\\where.exe";
+			}
+			else 
+			{
+				startInfo.FileName = "whereis";
+			}
+
+
+			PythonProcess.StartInfo = startInfo;
+            PythonProcess.Start();
+            PythonProcess.WaitForExit();
+
+            string Apps = PythonProcess.StandardOutput.ReadToEnd();
+            Apps = Apps.Replace('\r', '\n');
+            string[] locations = Apps.Split('\n');
+
+            return locations[0];
+        }
+
+        public static string EmscriptenPackager()
+        {
+            return Path.Combine(EmscriptenSDKPath(), "tools", "file_packager.py"); 
+        }
+
+        public static string EmscriptenCompiler()
+        {
+            return Path.Combine(EmscriptenSDKPath(), "emcc");
+        }
+
+        public static bool IsSDKInstalled()
+        {
+            try 
+            {
+                EmscriptenVersion();
+                return true; 
+            }
+            catch (Exception /*ex*/ )
+            {
+                return false; 
+            }
+        }
+
+        public static bool IsPythonInstalled()
+        {
+            if (String.IsNullOrEmpty(PythonPath()))
+                return false;
+            return true;
+        }
+
+        public static string EmscriptenVersion()
+        {
+            string BaseSDKPath = HTML5SDKInfo.EmscriptenSDKPath(); 
+            string VersionInfo = File.ReadAllText( Path.Combine(BaseSDKPath,"emscripten-version.txt"));
+            VersionInfo = VersionInfo.Trim();
+            return VersionInfo; 
+        }
+    }
+}
+ 
\ No newline at end of file
diff --git a/Engine/Source/Programs/UnrealBuildTool/HTML5/HTML5ToolChain.cs b/Engine/Source/Programs/UnrealBuildTool/HTML5/HTML5ToolChain.cs
--- a/Engine/Source/Programs/UnrealBuildTool/HTML5/HTML5ToolChain.cs
+++ b/Engine/Source/Programs/UnrealBuildTool/HTML5/HTML5ToolChain.cs
@@ -11,76 +11,17 @@ namespace UnrealBuildTool
 {
 	class HTML5ToolChain : VCToolChain
 	{
-		// cache the location of SDK tools
-		static string EMCCPath;
-		static string PythonPath;
-        static string EmscriptenSettingsPath = Path.Combine(Environment.GetFolderPath(Environment.SpecialFolder.UserProfile), "/.emscripten");
-
-        Dictionary<string, string> ReadEmscriptenSettings()
-        {
-            // Check HTML5ToolChain.cs for duplicate
-            if (!System.IO.File.Exists(EmscriptenSettingsPath))
-            {
-                return new Dictionary<string, string>();
-            }
-
-            Dictionary<string, string> Settings = new Dictionary<string, string>();
-            System.IO.StreamReader SettingFile = new System.IO.StreamReader(EmscriptenSettingsPath);
-            string EMLine = null;
-            while ((EMLine = SettingFile.ReadLine()) != null)
-            {
-                EMLine = EMLine.Split('#')[0];
-                string Pattern1 = @"(\w*)\s*=\s*['\[]?([^'\]\r\n]*)['\]]?";
-                Regex Rgx = new Regex(Pattern1, RegexOptions.IgnoreCase);
-                MatchCollection Matches = Rgx.Matches(EMLine);
-                foreach (Match Matched in Matches)
-                {
-                    if (Matched.Groups.Count == 3 && Matched.Groups[2].ToString() != "")
-                    {
-                        Settings[Matched.Groups[1].ToString()] = Matched.Groups[2].ToString();
-                    }
-                }
-            }
-
-            return Settings;
-        }
+        static string EMCCPath;
+        static string PythonPath;
 
+		// cache the location of SDK tools
 		public override void RegisterToolChain()
 		{
-			// Make sure the SDK is installed
-            // look up installed SDK. 
-			string BaseSDKPath = Environment.GetEnvironmentVariable("EMSCRIPTEN");
-			if (!String.IsNullOrEmpty(BaseSDKPath))
-			{
-				BaseSDKPath = BaseSDKPath.Replace("\"", "");
-				if (!String.IsNullOrEmpty(BaseSDKPath))
-                {
-                    var EmscriptenSettings = ReadEmscriptenSettings();
-					EMCCPath = Path.Combine(BaseSDKPath, "emcc");
-					// also figure out where python lives (if no envvar, assume it's in the path)
-                    if (EmscriptenSettings.ContainsKey("PYTHON"))
-                    {
-                        PythonPath = EmscriptenSettings["PYTHON"];
-                    }
-                    else
-                    {
-                        PythonPath = Environment.GetEnvironmentVariable("PYTHON");
-                    }
-
-                    string PythonExeName = Utils.IsRunningOnMono ? "python" : "python.exe";
+            if (HTML5SDKInfo.IsSDKInstalled() && HTML5SDKInfo.IsPythonInstalled())
+            {
+					EMCCPath = HTML5SDKInfo.EmscriptenCompiler();
+					PythonPath = HTML5SDKInfo.PythonPath(); 
 
-					if (PythonPath == null)
-					{
-                        PythonPath = PythonExeName;
-					}
-                    else
-                    {
-                        if (!PythonPath.EndsWith(PythonExeName))
-                        {
-                            PythonPath += "/" + PythonExeName;
-                        }
-                    }
-                    EMCCPath = "\"" + EMCCPath + "\"";
 					// set some environment variable we'll need
 					//Environment.SetEnvironmentVariable("EMCC_DEBUG", "cache");
 					Environment.SetEnvironmentVariable("EMCC_CORES", "8");
@@ -89,7 +30,8 @@ namespace UnrealBuildTool
 					// finally register the toolchain that is now ready to go
                     Log.TraceVerbose("        Registered for {0}", CPPTargetPlatform.HTML5.ToString());
 					UEToolChain.RegisterPlatformToolChain(CPPTargetPlatform.HTML5, this);
-				}
+
+   
 			}
 		}
 
@@ -405,7 +347,7 @@ namespace UnrealBuildTool
 				CompileAction.WorkingDirectory = Path.GetFullPath(".");
 				CompileAction.CommandPath = PythonPath;
 
-				CompileAction.CommandArguments = EMCCPath + Arguments + FileArguments + CompileEnvironment.Config.AdditionalArguments;
+				CompileAction.CommandArguments = EMCCPath + " " + Arguments + FileArguments + CompileEnvironment.Config.AdditionalArguments;
 
                 System.Console.WriteLine(CompileAction.CommandArguments); 
 				CompileAction.StatusDescription = Path.GetFileName(SourceFile.AbsolutePath);
diff --git a/Engine/Source/Programs/UnrealBuildTool/HTML5/UEBuildHTML5.cs b/Engine/Source/Programs/UnrealBuildTool/HTML5/UEBuildHTML5.cs
--- a/Engine/Source/Programs/UnrealBuildTool/HTML5/UEBuildHTML5.cs
+++ b/Engine/Source/Programs/UnrealBuildTool/HTML5/UEBuildHTML5.cs
@@ -79,33 +79,23 @@ namespace UnrealBuildTool
 			if (HasSetupAutoSDK())
 			{
 				return SDKStatus.Invalid;
-			}
-
-			string BaseSDKPath = Environment.GetEnvironmentVariable("EMSCRIPTEN");
-			if (!string.IsNullOrEmpty(BaseSDKPath))
-			{
-
-				try 
-				{ 
-					// Check for the *actual* sdk version
-					string VersionInfo = File.ReadAllText( Path.Combine(BaseSDKPath,"emscripten-version.txt")); 
-					if (VersionInfo.Contains(ExpectedSDKVersion))
-					{
-						return SDKStatus.Valid;
-					}
-					else
-					{
-						Console.WriteLine( "EMSCRIPTEN sdk found but of unexpected version, Please install version " + ExpectedSDKVersion); 
-						return SDKStatus.Invalid;
-					}
-				} 
-				catch (System.Exception) 
-				{
-					Console.WriteLine( "Please check you emscripten installation. Incorrectly set EMSCRIPTEN sdk at " + BaseSDKPath );
-					return SDKStatus.Invalid; 
-				}
-			}
-			return SDKStatus.Invalid;
+            }
+            try
+            {
+                if (HTML5SDKInfo.EmscriptenVersion().Contains(ExpectedSDKVersion))
+                {
+                    return SDKStatus.Valid;
+                }
+                else
+                {
+                    Console.WriteLine("EMSCRIPTEN sdk " + HTML5SDKInfo.EmscriptenVersion() + " found which is unsupported, Please install version " + ExpectedSDKVersion);
+                    return SDKStatus.Invalid;
+                }
+            }
+            catch (Exception /*ex*/)
+            {
+                 return SDKStatus.Invalid;
+            }
         }
 
         public override bool CanUseXGE()
diff --git a/Engine/Source/Programs/UnrealBuildTool/UnrealBuildTool.csproj b/Engine/Source/Programs/UnrealBuildTool/UnrealBuildTool.csproj
--- a/Engine/Source/Programs/UnrealBuildTool/UnrealBuildTool.csproj
+++ b/Engine/Source/Programs/UnrealBuildTool/UnrealBuildTool.csproj
@@ -168,6 +168,7 @@
         <Compile Include="HTML5\HTML5ProjectGenerator.cs" />
         <Compile Include="HTML5\HTML5ToolChain.cs" />
         <Compile Include="HTML5\UEBuildHTML5.cs" />
+        <Compile Include="HTML5\HTML5SDKInfo.cs" />
       </ItemGroup>
     </When>
   </Choose>
diff --git a/Engine/Source/Programs/UnrealBuildTool/UnrealBuildTool_Mono.csproj b/Engine/Source/Programs/UnrealBuildTool/UnrealBuildTool_Mono.csproj
--- a/Engine/Source/Programs/UnrealBuildTool/UnrealBuildTool_Mono.csproj
+++ b/Engine/Source/Programs/UnrealBuildTool/UnrealBuildTool_Mono.csproj
@@ -86,19 +86,14 @@
     <Reference Include="System.Xml" />
     <Reference Include="System.Xml.Linq" />
   </ItemGroup>
-  <Choose>
-    <When Condition="Exists('HTML5\UEBuildHTML5.cs')">
-      <ItemGroup>
-        <Compile Include="HTML5\HTML5ProjectGenerator.cs" />
-        <Compile Include="HTML5\HTML5ToolChain.cs" />
-        <Compile Include="HTML5\UEBuildHTML5.cs" />
-      </ItemGroup>
-    </When>
-  </Choose>
   <ItemGroup>
     <Compile Include="..\DotNETCommon\MetaData.cs">
       <Link>Properties\MetaData.cs</Link>
     </Compile>
+    <Compile Include="HTML5\HTML5ProjectGenerator.cs" />
+    <Compile Include="HTML5\HTML5ToolChain.cs" />
+    <Compile Include="HTML5\HTML5SDKInfo.cs" />
+    <Compile Include="HTML5\UEBuildHTML5.cs" />
     <Compile Include="Configuration\UEBuildClient.cs" />
     <Compile Include="Configuration\UEBuildDeploy.cs" />
     <Compile Include="Configuration\UEBuildBinary.cs" />
