CIS fix.

diff --git a/Engine/Source/Programs/AutomationTool/HTML5/HTML5Platform.Automation.cs b/Engine/Source/Programs/AutomationTool/HTML5/HTML5Platform.Automation.cs
--- a/Engine/Source/Programs/AutomationTool/HTML5/HTML5Platform.Automation.cs
+++ b/Engine/Source/Programs/AutomationTool/HTML5/HTML5Platform.Automation.cs
@@ -12,9 +12,7 @@ using System.Threading;
 
 public class HTML5Platform : Platform
 {
-    static string EmscriptenSettingsPath = CombinePaths(Environment.GetFolderPath(Environment.SpecialFolder.UserProfile), "/.emscripten");
-
-	public HTML5Platform()
+ 	public HTML5Platform()
 		: base(UnrealTargetPlatform.HTML5)
 	{
 	}
