Add RPCUtility into the common tools build, and fix default platform name.

diff --git a/Engine/Source/Programs/AutomationTool/Scripts/BuildCommonTools.Automation.cs b/Engine/Source/Programs/AutomationTool/Scripts/BuildCommonTools.Automation.cs
--- a/Engine/Source/Programs/AutomationTool/Scripts/BuildCommonTools.Automation.cs
+++ b/Engine/Source/Programs/AutomationTool/Scripts/BuildCommonTools.Automation.cs
@@ -15,7 +15,7 @@ public class BuildCommonTools : BuildCommand
 		Log("************************* BuildCommonTools");
 
 		// Get the list of platform names
-		string[] PlatformNames = ParseParamValue("platforms", BuildHostPlatform.Current.ToString()).Split('+');
+		string[] PlatformNames = ParseParamValue("platforms", BuildHostPlatform.Current.Platform.ToString()).Split('+');
 
 		// Parse the platforms
 		List<UnrealBuildTool.UnrealTargetPlatform> Platforms = new List<UnrealTargetPlatform>();
@@ -60,6 +60,7 @@ public class BuildCommonTools : BuildCommand
 		// C# binaries
 		Agenda.SwarmProject = @"Engine\Source\Programs\UnrealSwarm\UnrealSwarm.sln";
 		Agenda.DotNetProjects.Add(@"Engine/Source/Programs/DotNETCommon/DotNETUtilities/DotNETUtilities.csproj");
+		Agenda.DotNetProjects.Add(@"Engine/Source/Programs/RPCUtility/RPCUtility.csproj");
 		Agenda.DotNetProjects.Add(@"Engine/Source/Programs/UnrealControls/UnrealControls.csproj");
 
 		// Windows binaries
