Add P4 utility function to convert multiple file names from depot to local paths.

diff --git a/Engine/Source/Programs/AutomationTool/P4Utils.cs b/Engine/Source/Programs/AutomationTool/P4Utils.cs
--- a/Engine/Source/Programs/AutomationTool/P4Utils.cs
+++ b/Engine/Source/Programs/AutomationTool/P4Utils.cs
@@ -1892,6 +1892,74 @@ namespace AutomationTool
             return null;
         }
 
+        /// <summary>
+        /// Given a file path in the depot, returns the local disk mapping for the current view
+        /// </summary>
+        /// <param name="DepotFilepath">The full file path in depot naming form</param>
+        /// <returns>The file's first reported path on disk or null if no mapping was found</returns>
+        public string[] DepotToLocalPaths(string[] DepotPaths, bool AllowSpew = true)
+        {
+            CheckP4Enabled();
+
+			// Create a lookup from depot file to local file, and step through the input array in batches running the p4 where command on the depot files.
+			Dictionary<string, string> LocalPathsLookup = new Dictionary<string,string>(StringComparer.InvariantCultureIgnoreCase);
+			for(int InputIdx = 0; InputIdx < DepotPaths.Length; )
+			{
+				// Build the query for these files and increment the input index as we go
+				StringBuilder Command = new StringBuilder("where");
+				while(InputIdx < DepotPaths.Length && Command.Length < 512)
+				{
+					Command.Append(' ');
+					Command.Append(DepotPaths[InputIdx]);
+					InputIdx++;
+				}
+
+				// Run the command
+				string Output;
+				if (!LogP4Output(out Output, Command.ToString(), AllowSpew:AllowSpew))
+				{
+					throw new P4Exception("p4.exe {0} failed.", Command.ToString());
+				}
+
+				// Copy the results into the local paths lookup. Entries may occur more than once, and entries may be missing from the client view, or deleted in the client view.
+				string[] Lines = Output.Split(new char[]{ '\r', '\n' }, StringSplitOptions.RemoveEmptyEntries);
+				foreach(string Line in Lines)
+				{
+					string[] Tokens = Line.Split(new char[] { ' ' }, StringSplitOptions.RemoveEmptyEntries);
+					if(Tokens.Length >= 1 && Line.EndsWith("not in client view."))
+					{
+						LocalPathsLookup[Tokens[0]] = null;
+					}
+					else if(Tokens.Length == 3)
+					{
+						if(Tokens[0].StartsWith("-"))
+						{
+							LocalPathsLookup[Tokens[0].Substring(1)] = null;
+						}
+						else
+						{
+							LocalPathsLookup[Tokens[0]] = Tokens[2];
+						}
+					}
+					else
+					{
+						throw new AutomationException("Unexpected output from p4 where command: {0}", Line);
+					}
+				}
+			}
+
+			// Build the output array of local paths
+			string[] LocalPaths = new string[DepotPaths.Length];
+			for(int Idx = 0; Idx < DepotPaths.Length; Idx++)
+			{
+				if(!LocalPathsLookup.TryGetValue(DepotPaths[Idx], out LocalPaths[Idx]))
+				{
+					throw new AutomationException("Missing state of '{0}' in p4 where output.", DepotPaths[Idx]);
+				}
+			}
+			return LocalPaths;
+        }
+
 		/// <summary>
 		/// Gets file stats.
 		/// </summary>
