build script, add properties for dependent promotions

diff --git a/Engine/Source/Programs/AutomationTool/Scripts/GUBP.Automation.cs b/Engine/Source/Programs/AutomationTool/Scripts/GUBP.Automation.cs
--- a/Engine/Source/Programs/AutomationTool/Scripts/GUBP.Automation.cs
+++ b/Engine/Source/Programs/AutomationTool/Scripts/GUBP.Automation.cs
@@ -359,6 +359,7 @@ public class GUBP : BuildCommand
         public List<string> FullNamesOfPseudosependencies = new List<string>(); //these are really only used for sorting. We want the editor to fail before the monolithics. Think of it as "can't possibly be useful without".
 		public List<string> FullNamesOfDependedOn = new List<string>();
         public List<string> BuildProducts = null;
+		public List<string> DependentPromotions = new List<string>();
         public List<string> AllDependencyBuildProducts = null;
         public List<string> AllDependencies = null;
         public string AgentSharingGroup = "";
@@ -409,6 +410,14 @@ public class GUBP : BuildCommand
         {
             return false;
         }
+		public virtual bool IsPromotableAggregate()
+		{
+			return false;
+		}
+		public virtual bool IsSeparatePromotable()
+		{
+			return false;
+		}
         public virtual int AgentMemoryRequirement(GUBP bp)
         {
             return 0;
@@ -1643,7 +1652,7 @@ public class GUBP : BuildCommand
         public override bool RunInEC()
         {
             return false;
-        }
+        }		
         public override bool IsAggregate()
         {
             return true;
@@ -1711,7 +1720,7 @@ public class GUBP : BuildCommand
     public class AggregatePromotableNode : AggregateNode
     {
         protected List<UnrealTargetPlatform> HostPlatforms;
-        string PromotionLabelPrefix;
+        string PromotionLabelPrefix;		
 
         public AggregatePromotableNode(List<UnrealTargetPlatform> InHostPlatforms, string InPromotionLabelPrefix)
         {
@@ -1726,6 +1735,10 @@ public class GUBP : BuildCommand
         {
             return InPromotionLabelPrefix + "_Promotable_Aggregate";
         }
+		public override bool IsPromotableAggregate()
+		{
+			return true;
+		}
         public override string GetFullName()
         {
             return StaticGetFullName(PromotionLabelPrefix);
@@ -1734,21 +1747,20 @@ public class GUBP : BuildCommand
 
     public class GameAggregatePromotableNode : AggregatePromotableNode
     {
-        BranchInfo.BranchUProject GameProj;
+        BranchInfo.BranchUProject GameProj;		
 
-        public GameAggregatePromotableNode(GUBP bp, List<UnrealTargetPlatform> InHostPlatforms, BranchInfo.BranchUProject InGameProj)
+        public GameAggregatePromotableNode(GUBP bp, List<UnrealTargetPlatform> InHostPlatforms, BranchInfo.BranchUProject InGameProj, bool IsSeparate)
             : base(InHostPlatforms, InGameProj.GameName)
         {
             GameProj = InGameProj;
-
+			
             foreach (var HostPlatform in HostPlatforms)
             {
-                AddDependency(RootEditorNode.StaticGetFullName(HostPlatform));
+                AddDependency(RootEditorNode.StaticGetFullName(HostPlatform));				
                 if (InGameProj.GameName != bp.Branch.BaseEngineProject.GameName && GameProj.Properties.Targets.ContainsKey(TargetRules.TargetType.Editor))
                 {
                     AddDependency(EditorGameNode.StaticGetFullName(HostPlatform, GameProj));
-                }
-
+                }				
                 // add all of the platforms I use
                 {
                     var Platforms = bp.GetMonolithicPlatformsForUProject(HostPlatform, InGameProj, false);
@@ -1779,6 +1791,18 @@ public class GUBP : BuildCommand
         {
             return GameProj.GameName;
         }
+		public override bool IsSeparatePromotable()
+		{
+			bool IsSeparate = false;
+			foreach(UnrealTargetPlatform HostPlatform in HostPlatforms)
+			{
+				if(GameProj.Options(HostPlatform).bSeparateGamePromotion)
+				{
+					IsSeparate = true;
+				}
+			}
+			return IsSeparate;
+		}
     }
 
     public class SharedAggregatePromotableNode : AggregatePromotableNode
@@ -1813,7 +1837,10 @@ public class GUBP : BuildCommand
                 }
             }
         }
-
+		public override bool IsSeparatePromotable()
+		{
+			return true;
+		}
         public static string StaticGetFullName()
         {
             return AggregatePromotableNode.StaticGetFullName("Shared");
@@ -3771,6 +3798,7 @@ public class GUBP : BuildCommand
         bool bShowDetailedHistory = (bp.ParseParam("History") && GUBPNodesHistory != null) || bShowChanges;
         bool bShowDependencies = bp.ParseParam("ShowDependencies");
 		bool bShowDependednOn = bp.ParseParam("ShowDependedOn");
+		bool bShowDependentPromotions = bp.ParseParam("ShowDependentPromotions");
         bool bShowECDependencies = bp.ParseParam("ShowECDependencies");
         bool bShowHistory = !bp.ParseParam("NoHistory") && GUBPNodesHistory != null;
         bool AddEmailProps = bp.ParseParam("ShowEmails");
@@ -3886,6 +3914,13 @@ public class GUBP : BuildCommand
 					Log("            depOn> {0}", Dep);
 				}
 			}
+			if (bShowDependentPromotions)
+			{
+				foreach (var Dep in GUBPNodes[NodeToDo].DependentPromotions)
+				{
+					Log("            depPro> {0}", Dep);
+				}
+			}
         }
     }
     public void SaveGraphVisualization(List<string> Nodes)
@@ -5632,7 +5667,7 @@ public class GUBP : BuildCommand
                 }
                 if (PromotedHosts.Count > 0)
                 {
-                    AddNode(new GameAggregatePromotableNode(this, PromotedHosts, CodeProj));
+                    AddNode(new GameAggregatePromotableNode(this, PromotedHosts, CodeProj, true));
                     if (AnySeparate)
                     {
                         AddNode(new WaitForGamePromotionUserInput(this, CodeProj, false));
@@ -5645,7 +5680,7 @@ public class GUBP : BuildCommand
         }
         if (NumSharedAllHosts > 0)
         {
-            AddNode(new GameAggregatePromotableNode(this, HostPlatforms, Branch.BaseEngineProject));
+            AddNode(new GameAggregatePromotableNode(this, HostPlatforms, Branch.BaseEngineProject, false));
 
             AddNode(new SharedAggregatePromotableNode(this, HostPlatforms));
             AddNode(new WaitForSharedPromotionUserInput(this, false));
@@ -5728,6 +5763,7 @@ public class GUBP : BuildCommand
         var FullNodeListSortKey = new Dictionary<string, int>();
         var FullNodeDirectDependencies = new Dictionary<string, string>();
 		var FullNodeDependedOnBy = new Dictionary<string, string>();
+		var FullNodeDependentPromotions = new Dictionary<string, string>();
         {
             Log("******* {0} GUBP Nodes", GUBPNodes.Count);
             var SortedNodes = TopologicalSort(new HashSet<string>(GUBPNodes.Keys), LocalOnly: true, DoNotConsiderCompletion: true);
@@ -6002,6 +6038,45 @@ public class GUBP : BuildCommand
                 NodesToDo.UnionWith(Fringe);
             }
         }
+		if(!bOnlyNode)
+		{
+			foreach(var NodeToDo in NodesToDo)
+			{
+				if(GUBPNodes[NodeToDo].IsSeparatePromotable())
+				{					
+					List<string> Dependencies = new List<string>();
+					Dependencies = GetECDependencies(NodeToDo);
+					foreach(var Dep in Dependencies)
+					{
+						if(!GUBPNodes.ContainsKey(Dep))
+						{
+							throw new AutomationException("Node {0} is not in the graph.  It is a dependency of {1}.", Dep, NodeToDo);
+						}
+						if(!GUBPNodes[Dep].IsPromotableAggregate())
+						{
+							if (!GUBPNodes[Dep].DependentPromotions.Contains(NodeToDo))
+							{
+								GUBPNodes[Dep].DependentPromotions.Add(NodeToDo);
+							}
+						}
+					}
+				}
+			}
+			foreach(var NodeToDo in NodesToDo)
+			{
+				var Deps = GUBPNodes[NodeToDo].DependentPromotions;
+				string All = "";
+				foreach (var Dep in Deps)
+				{
+					if (All != "")
+					{
+						All += " ";
+					}
+					All += Dep;
+				}
+				FullNodeDependentPromotions.Add(NodeToDo, All);				
+			}
+		}
 		if (TimeIndex != 0)
 		{
 			Log("Culling based on time index");
@@ -6223,6 +6298,10 @@ public class GUBP : BuildCommand
 			{
 				ECProps.Add(string.Format("DependedOnBy/{0}={1}", NodePair.Key, NodePair.Value));
 			}
+			foreach (var NodePair in FullNodeDependentPromotions)
+			{
+				ECProps.Add(string.Format("DependentPromotions/{0}={1}", NodePair.Key, NodePair.Value));
+			}
             var ECJobProps = new List<string>();
             if (ExplicitTrigger != "")
             {
