build script, store promotables

diff --git a/Engine/Source/Programs/AutomationTool/Scripts/GUBP.Automation.cs b/Engine/Source/Programs/AutomationTool/Scripts/GUBP.Automation.cs
--- a/Engine/Source/Programs/AutomationTool/Scripts/GUBP.Automation.cs
+++ b/Engine/Source/Programs/AutomationTool/Scripts/GUBP.Automation.cs
@@ -5764,6 +5764,7 @@ public class GUBP : BuildCommand
         var FullNodeDirectDependencies = new Dictionary<string, string>();
 		var FullNodeDependedOnBy = new Dictionary<string, string>();
 		var FullNodeDependentPromotions = new Dictionary<string, string>();
+		var SeparatePromotables = new List<string>();
         {
             Log("******* {0} GUBP Nodes", GUBPNodes.Count);
             var SortedNodes = TopologicalSort(new HashSet<string>(GUBPNodes.Keys), LocalOnly: true, DoNotConsiderCompletion: true);
@@ -6043,7 +6044,8 @@ public class GUBP : BuildCommand
 			foreach(var NodeToDo in NodesToDo)
 			{
 				if(GUBPNodes[NodeToDo].IsSeparatePromotable())
-				{					
+				{
+					SeparatePromotables.Add(GUBPNodes[NodeToDo].GetFullName());
 					List<string> Dependencies = new List<string>();
 					Dependencies = GetECDependencies(NodeToDo);
 					foreach(var Dep in Dependencies)
@@ -6302,6 +6304,10 @@ public class GUBP : BuildCommand
 			{
 				ECProps.Add(string.Format("DependentPromotions/{0}={1}", NodePair.Key, NodePair.Value));
 			}
+			foreach (var Node in SeparatePromotables)
+			{
+				ECProps.Add(string.Format("PossiblePromotables/{0}={1}", Node, ""));
+			}
             var ECJobProps = new List<string>();
             if (ExplicitTrigger != "")
             {
