Add a wrapper for the P4 where function, which correctly parses the output for files containing spaces, and files which are unmapped from the repository. (Previous implementation would return incorrectly when files were unmapped.)

diff --git a/Engine/Source/Programs/AutomationTool/P4Utils.cs b/Engine/Source/Programs/AutomationTool/P4Utils.cs
--- a/Engine/Source/Programs/AutomationTool/P4Utils.cs
+++ b/Engine/Source/Programs/AutomationTool/P4Utils.cs
@@ -185,6 +185,14 @@ namespace AutomationTool
 		public bool IsValid { get { return Type != P4FileType.Unknown; } }
 	}
 
+	public class P4WhereRecord
+	{
+		public bool bUnmap;
+		public string DepotFile;
+		public string ClientFile;
+		public string Path;
+	}
+
 	public partial class CommandUtils
 	{
 		#region Environment Setup
@@ -1863,102 +1871,97 @@ namespace AutomationTool
         /// <summary>
         /// Given a file path in the depot, returns the local disk mapping for the current view
         /// </summary>
-        /// <param name="DepotFilepath">The full file path in depot naming form</param>
+		/// <param name="DepotFile">The full file path in depot naming form</param>
         /// <returns>The file's first reported path on disk or null if no mapping was found</returns>
-        public string DepotToLocalPath(string DepotFilepath)
+        public string DepotToLocalPath(string DepotFile, bool AllowSpew = true)
         {
-            CheckP4Enabled();
-            string Output;
-            string Command = "where " + DepotFilepath;
-            if (!LogP4Output(out Output, Command))
-            {
-                throw new P4Exception("p4.exe {0} failed.", Command);
-            }
+			P4WhereRecord[] Records = Where(DepotFile, AllowSpew);
+			if (Records != null)
+			{
+				foreach (P4WhereRecord Record in Records)
+				{
+					if (!Record.bUnmap)
+					{
+						return Record.Path;
+					}
+				}
+			}
+			return null;
+        }
 
-            string[] mappings = Output.Split(new char[] { '\r', '\n' }, StringSplitOptions.RemoveEmptyEntries);
-            foreach (string mapping in mappings)
-            {
-                if (mapping.EndsWith("not in client view."))
-                {
-                    return null;
-                }
-                string[] files = mapping.Split(new char[] { ' ' }, StringSplitOptions.RemoveEmptyEntries);
-                if (files.Length > 0)
-                {
-                    return files[files.Length - 1];
-                }
-            }
+		/// <summary>
+		/// Determines the mappings for a depot file in the workspace, without that file having to exist. 
+		/// NOTE: This function originally allowed multiple depot paths at once. The "file(s) not in client view" messages are written to stderr 
+		/// rather than stdout, and buffering them separately garbles the output when they're merged together.
+		/// </summary>
+		/// <param name="DepotFile">Depot path</param>
+		/// <param name="AllowSpew">Allows logging</param>
+		/// <returns>List of records describing the file's mapping. Usually just one, but may be more.</returns>
+		public P4WhereRecord[] Where(string DepotFile, bool AllowSpew = true)
+		{
+			CheckP4Enabled();
 
-            return null;
-        }
+			//  P4 where outputs missing entries 
+			string Command = String.Format("-z tag where \"{0}\"", DepotFile);
 
-        /// <summary>
-        /// Given a file path in the depot, returns the local disk mapping for the current view
-        /// </summary>
-        /// <param name="DepotFilepath">The full file path in depot naming form</param>
-        /// <returns>The file's first reported path on disk or null if no mapping was found</returns>
-        public string[] DepotToLocalPaths(string[] DepotPaths, bool AllowSpew = true)
-        {
-            CheckP4Enabled();
+			// Run the command.
+			string Output;
+			if (!LogP4Output(out Output, Command, AllowSpew: AllowSpew))
+			{
+				throw new P4Exception("p4.exe {0} failed.", Command);
+			}
+
+			// Copy the results into the local paths lookup. Entries may occur more than once, and entries may be missing from the client view, or deleted in the client view.
+			string[] Lines = Output.Split(new char[] { '\r', '\n' }, StringSplitOptions.RemoveEmptyEntries);
 
-			// Create a lookup from depot file to local file, and step through the input array in batches running the p4 where command on the depot files.
-			Dictionary<string, string> LocalPathsLookup = new Dictionary<string,string>(StringComparer.InvariantCultureIgnoreCase);
-			for(int InputIdx = 0; InputIdx < DepotPaths.Length; )
+			// Check for the file not existing
+			if(Lines.Length == 1 && Lines[0].EndsWith(" - file(s) not in client view."))
 			{
-				// Build the query for these files and increment the input index as we go
-				StringBuilder Command = new StringBuilder("where");
-				while(InputIdx < DepotPaths.Length && Command.Length < 512)
+				return null;
+			}
+
+			// Parse it into records
+			List<P4WhereRecord> Records = new List<P4WhereRecord>();
+			for (int LineIdx = 0; LineIdx < Lines.Length; )
+			{
+				P4WhereRecord Record = new P4WhereRecord();
+
+				// Parse an optional "... unmap"
+				if (Lines[LineIdx].Trim() == "... unmap")
 				{
-					Command.Append(' ');
-					Command.Append(DepotPaths[InputIdx]);
-					InputIdx++;
+					Record.bUnmap = true;
+					LineIdx++;
 				}
 
-				// Run the command
-				string Output;
-				if (!LogP4Output(out Output, Command.ToString(), AllowSpew:AllowSpew))
+				// Parse "... depotFile <depot path>"
+				const string DepotFilePrefix = "... depotFile ";
+				if (LineIdx >= Lines.Length || !Lines[LineIdx].StartsWith(DepotFilePrefix))
 				{
-					throw new P4Exception("p4.exe {0} failed.", Command.ToString());
+					throw new AutomationException("Unexpected output from p4 where: {0}", String.Join("\n", Lines.Skip(LineIdx)));
 				}
+				Record.DepotFile = Lines[LineIdx++].Substring(DepotFilePrefix.Length).Trim();
 
-				// Copy the results into the local paths lookup. Entries may occur more than once, and entries may be missing from the client view, or deleted in the client view.
-				string[] Lines = Output.Split(new char[]{ '\r', '\n' }, StringSplitOptions.RemoveEmptyEntries);
-				foreach(string Line in Lines)
+				// Parse "... clientFile <client path>"
+				const string ClientFilePrefix = "... clientFile ";
+				if (LineIdx >= Lines.Length || !Lines[LineIdx].StartsWith(ClientFilePrefix))
 				{
-					string[] Tokens = Line.Split(new char[] { ' ' }, StringSplitOptions.RemoveEmptyEntries);
-					if(Tokens.Length >= 1 && Line.EndsWith("not in client view."))
-					{
-						LocalPathsLookup[Tokens[0]] = null;
-					}
-					else if(Tokens.Length == 3)
-					{
-						if(Tokens[0].StartsWith("-"))
-						{
-							LocalPathsLookup[Tokens[0].Substring(1)] = null;
-						}
-						else
-						{
-							LocalPathsLookup[Tokens[0]] = Tokens[2];
-						}
-					}
-					else
-					{
-						throw new AutomationException("Unexpected output from p4 where command: {0}", Line);
-					}
+					throw new AutomationException("Unexpected output from p4 where: {0}", String.Join("\n", Lines.Skip(LineIdx)));
 				}
-			}
+				Record.ClientFile = Lines[LineIdx++].Substring(ClientFilePrefix.Length).Trim();
 
-			// Build the output array of local paths
-			string[] LocalPaths = new string[DepotPaths.Length];
-			for(int Idx = 0; Idx < DepotPaths.Length; Idx++)
-			{
-				if(!LocalPathsLookup.TryGetValue(DepotPaths[Idx], out LocalPaths[Idx]))
+				// Parse "... path <path to file>"
+				const string PathPrefix = "... path ";
+				if (LineIdx >= Lines.Length || !Lines[LineIdx].StartsWith(PathPrefix))
 				{
-					throw new AutomationException("Missing state of '{0}' in p4 where output.", DepotPaths[Idx]);
+					throw new AutomationException("Unexpected output from p4 where: {0}", String.Join("\n", Lines.Skip(LineIdx)));
 				}
+				Record.Path = Lines[LineIdx++].Substring(PathPrefix.Length).Trim();
+
+				// Add it to the output list
+				Records.Add(Record);
 			}
-			return LocalPaths;
-        }
+			return Records.ToArray();
+		}
 
 		/// <summary>
 		/// Gets file stats.
