fix for Soul defaulting to simulator when codesigning.  we now explicitly pass the iphoneos sdk to xcodebuild
#ios

diff --git a/Engine/Source/Programs/AutomationTool/IOS/IOSPlatform.Automation.cs b/Engine/Source/Programs/AutomationTool/IOS/IOSPlatform.Automation.cs
--- a/Engine/Source/Programs/AutomationTool/IOS/IOSPlatform.Automation.cs
+++ b/Engine/Source/Programs/AutomationTool/IOS/IOSPlatform.Automation.cs
@@ -480,6 +480,7 @@ public class IOSPlatform : Platform
 		Arguments += GameName;
 		Arguments += " - iOS'";
 		Arguments += " -configuration " + TargetConfig.ToString();
+		Arguments += " -sdk iphoneos";
 		Arguments += " CODE_SIGN_IDENTITY=" + (Distribution ? "\"iPhone Distribution\"" : "\"iPhone Developer\"");
 		ProcessResult Result = Run ("/usr/bin/env", Arguments, null, ERunOptions.Default);
 		if (bWasGenerated)
