Allow periods in P4 label names (such as //depot/UE4-Releases/4.6/Promoted-CL-XXX)

diff --git a/Engine/Source/Programs/AutomationTool/P4Utils.cs b/Engine/Source/Programs/AutomationTool/P4Utils.cs
--- a/Engine/Source/Programs/AutomationTool/P4Utils.cs
+++ b/Engine/Source/Programs/AutomationTool/P4Utils.cs
@@ -1722,7 +1722,7 @@ namespace AutomationTool
 		}
 
 		/* Pattern to parse P4 labels command output. */
-		static readonly Regex LabelsListOutputPattern = new Regex(@"^Label\s+(?<name>[\w\/-]+)\s+(?<date>\d{4}/\d{2}/\d{2} \d{2}:\d{2}:\d{2})\s+'(?<description>.+)'\s*$", RegexOptions.Compiled | RegexOptions.Multiline);
+		static readonly Regex LabelsListOutputPattern = new Regex(@"^Label\s+(?<name>[\w\/\.-]+)\s+(?<date>\d{4}/\d{2}/\d{2} \d{2}:\d{2}:\d{2})\s+'(?<description>.+)'\s*$", RegexOptions.Compiled | RegexOptions.Multiline);
 
 		/// <summary>
 		/// Gets all labels satisfying given filter.
