Correctly parse multiline descriptions in labels, and add a function to update descriptions.

diff --git a/Engine/Source/Programs/AutomationTool/P4Utils.cs b/Engine/Source/Programs/AutomationTool/P4Utils.cs
--- a/Engine/Source/Programs/AutomationTool/P4Utils.cs
+++ b/Engine/Source/Programs/AutomationTool/P4Utils.cs
@@ -1669,13 +1669,14 @@ namespace AutomationTool
 		/// </summary>
 		/// <param name="Name">Name of the label.</param>
 		/// <param name="Description">Description of the label.</param>
+		/// <param name="AllowSpew">Whether to allow log spew</param>
 		/// <returns>Returns whether the label description could be retrieved.</returns>
-		public bool LabelDescription(string Name, out string Description)
+		public bool LabelDescription(string Name, out string Description, bool AllowSpew = true)
 		{
 			CheckP4Enabled();
 			string Output;
 			Description = "";
-			if (LogP4Output(out Output, "label -o " + Name))
+			if (LogP4Output(out Output, "label -o " + Name, AllowSpew: AllowSpew))
 			{
 				string Desc = "Description:";
 				int Start = Output.LastIndexOf(Desc);
@@ -1686,7 +1687,7 @@ namespace AutomationTool
 				int End = Output.LastIndexOf("Options:");
 				if (Start > 0 && End > 0 && End > Start)
 				{
-					Description = Output.Substring(Start, End - Start);
+					Description = Output.Substring(Start, End - Start).Replace("\n\t", "\n");
 					Description = Description.Trim();
 					return true;
 				}
@@ -1694,6 +1695,45 @@ namespace AutomationTool
 			return false;
 		}
 
+		/// <summary>
+		/// Updates a label description.
+		/// </summary>
+		/// <param name="Name">Name of the label</param>
+		/// <param name="Description">Description of the label.</param>
+		/// <param name="AllowSpew">Whether to allow log spew</param>
+		public void UpdateLabelDescription(string Name, string NewDescription, bool AllowSpew = true)
+		{
+			string LabelSpec;
+			if(!LogP4Output(out LabelSpec, "label -o " + Name, AllowSpew: AllowSpew))
+			{
+				throw new P4Exception("Couldn't describe existing label '{0}', output was:\n", Name, LabelSpec);
+			}
+			List<string> Lines = new List<string>(LabelSpec.Split('\n').Select(x => x.TrimEnd()));
+
+			// Find the description text, and remove it
+			int Idx = 0;
+			for(; Idx < Lines.Count; Idx++)
+			{
+				if(Lines[Idx].StartsWith("Description:"))
+				{
+					int EndIdx = Idx + 1;
+					while(EndIdx < Lines.Count && (Lines[EndIdx].Length == 0 || Char.IsWhiteSpace(Lines[EndIdx][0]) || Lines[EndIdx].IndexOf(':') == -1))
+					{
+						EndIdx++;
+					}
+					Lines.RemoveRange(Idx, EndIdx - Idx);
+					break;
+				}
+			}
+
+			// Insert the new description text
+			Lines.Insert(Idx, "Description:  " + NewDescription.Replace("\n", "\n\t"));
+			LabelSpec = String.Join("\n", Lines);
+
+			// Update the label
+			LogP4("label -i", Input: LabelSpec, AllowSpew: AllowSpew);
+		}
+
 		/* Pattern to parse P4 changes command output. */
 		private static readonly Regex ChangesListOutputPattern = new Regex(@"^Change\s+(?<number>\d+)\s+.+$", RegexOptions.Compiled | RegexOptions.Multiline);
 
