Specify the Perforce clientspec for a branch in a config file, and allow Git publishing scripts to update it on-the-fly. Handle promotions for the master branch (which merge into a long-lived promoted branch), and fix snapshots from removing promoted dependency files.

diff --git a/Engine/Source/Programs/AutomationTool/P4Utils.cs b/Engine/Source/Programs/AutomationTool/P4Utils.cs
--- a/Engine/Source/Programs/AutomationTool/P4Utils.cs
+++ b/Engine/Source/Programs/AutomationTool/P4Utils.cs
@@ -78,6 +78,18 @@ namespace AutomationTool
         public P4SubmitOption SubmitOptions;
         public List<KeyValuePair<string, string>> View = new List<KeyValuePair<string, string>>();
 
+		public bool Matches(P4ClientInfo Other)
+		{
+			return Name == Other.Name 
+				&& RootPath == Other.RootPath 
+				&& Host == Other.Host 
+				&& Owner == Other.Owner 
+				&& LineEnd == Other.LineEnd 
+				&& Options == Other.Options 
+				&& SubmitOptions == Other.SubmitOptions
+				&& Enumerable.SequenceEqual(View, Other.View);
+		}
+
 		public override string ToString()
 		{
 			return Name;
@@ -2379,6 +2391,21 @@ namespace AutomationTool
 			return Result;
 		}
 
+		/// <summary>
+		/// Gets the contents of a particular file in the depot without syncing it
+		/// </summary>
+		/// <param name="DepotPath">Depot path to the file (with revision/range if necessary)</param>
+		/// <returns>Contents of the file</returns>
+		public string Print(string DepotPath, bool AllowSpew = true)
+		{
+			string Output;
+			if(!P4Output(out Output, "print -q " + DepotPath, AllowSpew: AllowSpew, WithClient: false))
+			{
+				throw new AutomationException("p4 print {0} failed", DepotPath);
+			}
+			return Output;
+		}
+
 		#region Utilities
 
 		private static object[] OldStyleBinaryFlags = new object[]
