Improvements to Android deploy - queries the device for the external data directory so that the runtime and push are using the same locations
Small change to BuildCookRun to query the game ini for Android package options rather than pulling it from the command line

diff --git a/Engine/Source/Editor/MainFrame/Private/Frame/MainFrameActions.cpp b/Engine/Source/Editor/MainFrame/Private/Frame/MainFrameActions.cpp
--- a/Engine/Source/Editor/MainFrame/Private/Frame/MainFrameActions.cpp
+++ b/Engine/Source/Editor/MainFrame/Private/Frame/MainFrameActions.cpp
@@ -637,11 +637,6 @@ void FMainFrameActionCallbacks::PackageProject( const FName InPlatformInfoName )
 	  if (PlatformInfo->TargetPlatformName != FName("HTML5")) 
 	  { 
 		OptionalParams += TEXT(" -pak");
-
-		if (PackagingSettings->UseOBB_InAPK)
-		{
-			OptionalParams += TEXT(" -obbinapk");
-		}
 	  }
 	}
 
diff --git a/Engine/Source/Programs/AutomationTool/Android/AndroidPlatform.Automation.cs b/Engine/Source/Programs/AutomationTool/Android/AndroidPlatform.Automation.cs
--- a/Engine/Source/Programs/AutomationTool/Android/AndroidPlatform.Automation.cs
+++ b/Engine/Source/Programs/AutomationTool/Android/AndroidPlatform.Automation.cs
@@ -89,9 +89,14 @@ public class AndroidPlatform : Platform
 	{
 		string ObbName = GetFinalObbName(ApkName);
 		string PackageName = GetPackageInfo(ApkName, false);
-		return "/mnt/sdcard/obb/" + PackageName + "/" + Path.GetFileName(ObbName);
+		return "obb/" + PackageName + "/" + Path.GetFileName(ObbName);
 	}
 
+    private static string GetStorageQueryCommand()
+    {
+        return "shell \"echo $EXTERNAL_STORAGE\"";
+    }
+
 	private static string GetFinalBatchName(string ApkName, ProjectParams Params, string Architecture, string GPUArchitecture)
 	{
 		return Path.Combine(Path.GetDirectoryName(ApkName), "Install_" + Params.ShortProjectName + "_" + Params.ClientConfigsToBuild[0].ToString() + Architecture + GPUArchitecture + ".bat");
@@ -156,13 +161,14 @@ public class AndroidPlatform : Platform
 				"set ADB=%ANDROID_HOME%\\platform-tools\\adb.exe",
 				"set DEVICE=",
 				"if not \"%1\"==\"\" set DEVICE=-s %1",
+                "for /f \"delims=\" %%A in ('adb " + GetStorageQueryCommand() +"') do @set STORAGE=%%A",
 				"%ADB% %DEVICE% uninstall " + PackageName,
 				"%ADB% %DEVICE% install " + Path.GetFileName(ApkName),
 				"@if \"%ERRORLEVEL%\" NEQ \"0\" goto Error",
-				"%ADB% %DEVICE% shell rm -r /mnt/sdcard/" + Params.ShortProjectName,
-				"%ADB% %DEVICE% shell rm -r /mnt/sdcard/UE4Game/UE4CommandLine.txt", // we need to delete the commandline in UE4Game or it will mess up loading
-				"%ADB% %DEVICE% shell rm -r /mnt/sdcard/obb/" + PackageName,
-				Params.OBBinAPK || !bHasPakFile ? "" : "%ADB% %DEVICE% push " + Path.GetFileName(LocalObbName) + " " + DeviceObbName,
+				"%ADB% %DEVICE% shell rm -r %STORAGE%/" + Params.ShortProjectName,
+				"%ADB% %DEVICE% shell rm -r %STORAGE%/UE4Game/UE4CommandLine.txt", // we need to delete the commandline in UE4Game or it will mess up loading
+				"%ADB% %DEVICE% shell rm -r %STORAGE%/obb/" + PackageName,
+				Params.OBBinAPK || !bHasPakFile ? "" : "%ADB% %DEVICE% push " + Path.GetFileName(LocalObbName) + " %STORAGE%/" + DeviceObbName,
 				Params.OBBinAPK || !bHasPakFile ? "" : "if \"%ERRORLEVEL%\" NEQ \"0\" goto Error",
 				"goto:eof",
 				":Error",
@@ -348,7 +354,6 @@ public class AndroidPlatform : Platform
         }
 
 		// now we can use the apk to get more info
-		string DeviceObbName = GetDeviceObbName(ApkName);
 		string PackageName = GetPackageInfo(ApkName, false);
 
         // try uninstalling an old app with the same identifier.
@@ -385,14 +390,19 @@ public class AndroidPlatform : Platform
 		string IntermediateCmdLineFile = CombinePaths(SC.StageDirectory, "UE4CommandLine.txt");
 		Project.WriteStageCommandline(IntermediateCmdLineFile, Params, SC);
 
+        // Setup the OBB name and add the storage path (queried from the device) to it
+        string DeviceStorageQueryCommand = GetStorageQueryCommand();
+        ProcessResult Result = Run(CmdEnv.CmdExe, AdbCommand + DeviceStorageQueryCommand, null, ERunOptions.AppMustExist);
+        String StorageLocation = Result.Output.Trim();
+        string DeviceObbName = StorageLocation + "/" + GetDeviceObbName(ApkName);
 
 		// copy files to device if we were staging
 		if (SC.Stage)
 		{
 			// cache some strings
 			string BaseCommandline = AdbCommand + "push";
-			string RemoteDir = "/mnt/sdcard/" + Params.ShortProjectName;
-			string UE4GameRemoteDir = "/mnt/sdcard/" + Params.ShortProjectName;
+			string RemoteDir = StorageLocation + "/" + Params.ShortProjectName;
+            string UE4GameRemoteDir = StorageLocation + "/" + Params.ShortProjectName;
 
 			// make sure device is at a clean state
 			Run(CmdEnv.CmdExe, AdbCommand + "shell rm -r " + RemoteDir);
@@ -518,7 +528,7 @@ public class AndroidPlatform : Platform
 		{
 			// cache some strings
 			string BaseCommandline = AdbCommand + "push";
-			string RemoteDir = "/mnt/sdcard/" + Params.ShortProjectName;
+            string RemoteDir = StorageLocation + "/" + Params.ShortProjectName;
 
 			string FinalRemoteDir = RemoteDir;
 			/*
diff --git a/Engine/Source/Programs/AutomationTool/Scripts/BuildCookRun.Automation.cs b/Engine/Source/Programs/AutomationTool/Scripts/BuildCookRun.Automation.cs
--- a/Engine/Source/Programs/AutomationTool/Scripts/BuildCookRun.Automation.cs
+++ b/Engine/Source/Programs/AutomationTool/Scripts/BuildCookRun.Automation.cs
@@ -133,6 +133,11 @@ public class BuildCookRun : BuildCommand
 			}
 		}
 
+        var configCache = new UnrealBuildTool.ConfigCacheIni("Game", Path.GetDirectoryName(ProjectPath));
+        bool obbInAPK = false;
+        configCache.GetBool("/Script/UnrealEd.ProjectPackagingSettings", "UseOBB_InAPK", out obbInAPK);
+        Params.OBBinAPK = obbInAPK;         
+
 		Params.ValidateAndLog();
 		return Params;
 	}
