Add UAT and UBT to the build products for BuildCommonTools. Should fix missing DLLs in promoted snapshots.

diff --git a/Engine/Source/Programs/AutomationTool/Scripts/BuildCommonTools.Automation.cs b/Engine/Source/Programs/AutomationTool/Scripts/BuildCommonTools.Automation.cs
--- a/Engine/Source/Programs/AutomationTool/Scripts/BuildCommonTools.Automation.cs
+++ b/Engine/Source/Programs/AutomationTool/Scripts/BuildCommonTools.Automation.cs
@@ -36,6 +36,10 @@ public class BuildCommonTools : BuildCommand
 		UE4Build Builder = new UE4Build(this);
 		Builder.Build(Agenda, InDeleteBuildProducts:true, InUpdateVersionFiles: false);
 
+		// Add UAT and UBT to the build products
+		Builder.AddUATFilesToBuildProducts();
+		Builder.AddUBTFilesToBuildProducts();
+
 		// Make sure all the build products exist
 		UE4Build.CheckBuildProducts(Builder.BuildProductFiles);
 
@@ -56,7 +60,7 @@ public class BuildCommonTools : BuildCommand
 	{
 		// Create the build agenda
 		UE4Build.BuildAgenda Agenda = new UE4Build.BuildAgenda();
-
+		
 		// C# binaries
 		Agenda.SwarmProject = @"Engine\Source\Programs\UnrealSwarm\UnrealSwarm.sln";
 		Agenda.DotNetProjects.Add(@"Engine/Source/Programs/DotNETCommon/DotNETUtilities/DotNETUtilities.csproj");
@@ -83,7 +87,7 @@ public class BuildCommonTools : BuildCommand
 			Agenda.AddTarget("ShaderCompileWorker", UnrealBuildTool.UnrealTargetPlatform.Mac, UnrealBuildTool.UnrealTargetConfiguration.Development, InAddArgs: "-CopyAppBundleBackToDevice");
 			Agenda.AddTarget("UE4EditorServices", UnrealBuildTool.UnrealTargetPlatform.Mac, UnrealBuildTool.UnrealTargetConfiguration.Development, InAddArgs: "-CopyAppBundleBackToDevice");
 		}
-
+		
 		return Agenda;
 	}
 }
