build script, add shootergame dedicated server and client packaged builds

diff --git a/Engine/Source/Programs/AutomationTool/Scripts/GUBP.Automation.cs b/Engine/Source/Programs/AutomationTool/Scripts/GUBP.Automation.cs
--- a/Engine/Source/Programs/AutomationTool/Scripts/GUBP.Automation.cs
+++ b/Engine/Source/Programs/AutomationTool/Scripts/GUBP.Automation.cs
@@ -2725,13 +2725,19 @@ public class GUBP : BuildCommand
         public override string ECAgentString()
         {
             string Result = base.ECAgentString();
-            foreach (UnrealTargetPlatform Plat in ClientTargetPlatforms)
-            {
-                if (Plat == UnrealTargetPlatform.XboxOne)
-                {
-                    Result = MergeSpaceStrings(Result, Plat.ToString());
-                }
-            }
+			if (ClientTargetPlatforms != null)
+			{
+				if (!ClientNotGame)
+				{
+					foreach (UnrealTargetPlatform Plat in ClientTargetPlatforms)
+					{
+						if (Plat == UnrealTargetPlatform.XboxOne)
+						{
+							Result = MergeSpaceStrings(Result, Plat.ToString());
+						}
+					}
+				}
+			}
             return Result;
         }      
         public override float Priority()
