Add CrashReportClient binaries to promoted builds.

diff --git a/Engine/Source/Programs/AutomationTool/Scripts/BuildCommonTools.Automation.cs b/Engine/Source/Programs/AutomationTool/Scripts/BuildCommonTools.Automation.cs
--- a/Engine/Source/Programs/AutomationTool/Scripts/BuildCommonTools.Automation.cs
+++ b/Engine/Source/Programs/AutomationTool/Scripts/BuildCommonTools.Automation.cs
@@ -70,6 +70,8 @@ public class BuildCommonTools : BuildCommand
 		// Windows binaries
 		if(Platforms.Contains(UnrealBuildTool.UnrealTargetPlatform.Win64))
 		{
+			Agenda.AddTarget("CrashReportClient", UnrealBuildTool.UnrealTargetPlatform.Win64, UnrealBuildTool.UnrealTargetConfiguration.Development);
+			Agenda.AddTarget("CrashReportClient", UnrealBuildTool.UnrealTargetPlatform.Win32, UnrealBuildTool.UnrealTargetConfiguration.Development);
 			Agenda.AddTarget("UnrealHeaderTool", UnrealBuildTool.UnrealTargetPlatform.Win64, UnrealBuildTool.UnrealTargetConfiguration.Development);
 			Agenda.AddTarget("UnrealPak", UnrealBuildTool.UnrealTargetPlatform.Win64, UnrealBuildTool.UnrealTargetConfiguration.Development);
 			Agenda.AddTarget("UnrealLightmass", UnrealBuildTool.UnrealTargetPlatform.Win64, UnrealBuildTool.UnrealTargetConfiguration.Development);
@@ -82,6 +84,7 @@ public class BuildCommonTools : BuildCommand
 		// Mac binaries
 		if(Platforms.Contains(UnrealBuildTool.UnrealTargetPlatform.Mac))
 		{
+			Agenda.AddTarget("CrashReportClient", UnrealBuildTool.UnrealTargetPlatform.Mac, UnrealBuildTool.UnrealTargetConfiguration.Development, InAddArgs: "-CopyAppBundleBackToDevice");
 			Agenda.AddTarget("UnrealPak", UnrealBuildTool.UnrealTargetPlatform.Mac, UnrealBuildTool.UnrealTargetConfiguration.Development, InAddArgs: "-CopyAppBundleBackToDevice");
 			Agenda.AddTarget("UnrealLightmass", UnrealBuildTool.UnrealTargetPlatform.Mac, UnrealBuildTool.UnrealTargetConfiguration.Development, InAddArgs: "-CopyAppBundleBackToDevice");
 			Agenda.AddTarget("ShaderCompileWorker", UnrealBuildTool.UnrealTargetPlatform.Mac, UnrealBuildTool.UnrealTargetConfiguration.Development, InAddArgs: "-CopyAppBundleBackToDevice");
