Added another missing quote in HTML5 Package commands.
Fix UE-5375

diff --git a/Engine/Source/Programs/AutomationTool/HTML5/HTML5Platform.Automation.cs b/Engine/Source/Programs/AutomationTool/HTML5/HTML5Platform.Automation.cs
--- a/Engine/Source/Programs/AutomationTool/HTML5/HTML5Platform.Automation.cs
+++ b/Engine/Source/Programs/AutomationTool/HTML5/HTML5Platform.Automation.cs
@@ -34,7 +34,7 @@ public class HTML5Platform : Platform
                 string PythonPath = HTML5SDKInfo.PythonPath();
                 string PackagerPath = HTML5SDKInfo.EmscriptenPackager();
 
-                string CmdLine = string.Format("{0} \"{1}\" --preload . --js-output=\"{1}.js\"", PackagerPath, FinalDataLocation);
+                string CmdLine = string.Format("\"{0}\" \"{1}\" --preload . --js-output=\"{1}.js\"", PackagerPath, FinalDataLocation);
                 RunAndLog(CmdEnv, PythonPath, CmdLine);
             }
 
