Rename release -> development config in SwarmInterface (so that UAT can build it), and add it to be built as part of the common tools for GitHub.

diff --git a/Engine/Source/Editor/SwarmInterface/DotNET/SwarmInterface.csproj b/Engine/Source/Editor/SwarmInterface/DotNET/SwarmInterface.csproj
--- a/Engine/Source/Editor/SwarmInterface/DotNET/SwarmInterface.csproj
+++ b/Engine/Source/Editor/SwarmInterface/DotNET/SwarmInterface.csproj
@@ -22,7 +22,7 @@
     <WarningLevel>4</WarningLevel>
     <RegisterForComInterop>false</RegisterForComInterop>
   </PropertyGroup>
-  <PropertyGroup Condition=" '$(Configuration)|$(Platform)' == 'Release|AnyCPU' ">
+  <PropertyGroup Condition=" '$(Configuration)|$(Platform)' == 'Development|AnyCPU' ">
     <DebugType>pdbonly</DebugType>
     <Optimize>True</Optimize>
     <OutputPath>..\..\..\..\Binaries\DotNET\</OutputPath>
diff --git a/Engine/Source/Programs/AutomationTool/Scripts/BuildCommonTools.Automation.cs b/Engine/Source/Programs/AutomationTool/Scripts/BuildCommonTools.Automation.cs
--- a/Engine/Source/Programs/AutomationTool/Scripts/BuildCommonTools.Automation.cs
+++ b/Engine/Source/Programs/AutomationTool/Scripts/BuildCommonTools.Automation.cs
@@ -63,6 +63,7 @@ public class BuildCommonTools : BuildCommand
 		
 		// C# binaries
 		Agenda.SwarmProject = @"Engine\Source\Programs\UnrealSwarm\UnrealSwarm.sln";
+		Agenda.DotNetProjects.Add(@"Engine/Source/Editor/SwarmInterface/DotNET/SwarmInterface.csproj");
 		Agenda.DotNetProjects.Add(@"Engine/Source/Programs/DotNETCommon/DotNETUtilities/DotNETUtilities.csproj");
 		Agenda.DotNetProjects.Add(@"Engine/Source/Programs/RPCUtility/RPCUtility.csproj");
 		Agenda.DotNetProjects.Add(@"Engine/Source/Programs/UnrealControls/UnrealControls.csproj");
