build script, mac editors should never be built on mac minis

diff --git a/Engine/Source/Programs/AutomationTool/Scripts/GUBP.Automation.cs b/Engine/Source/Programs/AutomationTool/Scripts/GUBP.Automation.cs
--- a/Engine/Source/Programs/AutomationTool/Scripts/GUBP.Automation.cs
+++ b/Engine/Source/Programs/AutomationTool/Scripts/GUBP.Automation.cs
@@ -750,7 +750,7 @@ public class GUBP : BuildCommand
         }
         public override int AgentMemoryRequirement(GUBP bp)
         {
-            if (bp.ParseParam("Launcher") || bp.TimeIndex != 0)
+            if (bp.ParseParam("Launcher") || bp.TimeIndex != 0  && HostPlatform != UnrealTargetPlatform.Mac)
             {
                 return base.AgentMemoryRequirement(bp);
             }
