Fix description in blueprint documentation changelists.

diff --git a/Engine/Source/Programs/AutomationTool/Scripts/BuildDocumentation.Automation.cs b/Engine/Source/Programs/AutomationTool/Scripts/BuildDocumentation.Automation.cs
--- a/Engine/Source/Programs/AutomationTool/Scripts/BuildDocumentation.Automation.cs
+++ b/Engine/Source/Programs/AutomationTool/Scripts/BuildDocumentation.Automation.cs
@@ -114,7 +114,7 @@ class BlueprintDocumentationNode : DocumentationNode
 	public override void DoBuild(GUBP bp)
 	{
 		ExecuteApiDocTool("-rebuildblueprint", "APIDocTool-Blueprint");
-		SubmitOutputs("Code documentation", "Engine/Documentation/Builds/BlueprintAPI-*", "Engine/Documentation/CHM/BlueprintAPI.chm");
+		SubmitOutputs("Blueprint documentation", "Engine/Documentation/Builds/BlueprintAPI-*", "Engine/Documentation/CHM/BlueprintAPI.chm");
 		base.DoBuild(bp);
 	}
 
