Write correct adb and optional device serial number passthru for install batch file (Android packaging)
#android
#ue4
#codeview Josh.Adams,Rob.Jones

diff --git a/Engine/Source/Programs/AutomationTool/Android/AndroidPlatform.Automation.cs b/Engine/Source/Programs/AutomationTool/Android/AndroidPlatform.Automation.cs
--- a/Engine/Source/Programs/AutomationTool/Android/AndroidPlatform.Automation.cs
+++ b/Engine/Source/Programs/AutomationTool/Android/AndroidPlatform.Automation.cs
@@ -184,7 +184,7 @@ public class AndroidPlatform : Platform
 				"set ADB=%ANDROID_HOME%\\platform-tools\\adb.exe",
 				"set DEVICE=",
 				"if not \"%1\"==\"\" set DEVICE=-s %1",
-                "for /f \"delims=\" %%A in ('adb " + GetStorageQueryCommand() +"') do @set STORAGE=%%A",
+                "for /f \"delims=\" %%A in ('%ADB% %DEVICE% " + GetStorageQueryCommand() +"') do @set STORAGE=%%A",
 				"%ADB% %DEVICE% uninstall " + PackageName,
 				"%ADB% %DEVICE% install " + Path.GetFileName(ApkName),
 				"@if \"%ERRORLEVEL%\" NEQ \"0\" goto Error",
