git pull request script fix for not using show-refs which breaks once we pass 500 PRs

diff --git a/Engine/Source/Programs/AutomationTool/Scripts/Tests.Automation.cs b/Engine/Source/Programs/AutomationTool/Scripts/Tests.Automation.cs
--- a/Engine/Source/Programs/AutomationTool/Scripts/Tests.Automation.cs
+++ b/Engine/Source/Programs/AutomationTool/Scripts/Tests.Automation.cs
@@ -211,11 +211,11 @@ class GitPullRequest : BuildCommand
 		RunGit("reset --hard");
 
 		// show-ref is just a double check that the PR exists
-		var Refs = RunGit("show-ref");
-		if (!Refs.Contains("refs/remotes/origin/pr/" + PRNum))
+		//var Refs = RunGit("show-ref");
+		/*if (!Refs.Contains("refs/remotes/origin/pr/" + PRNum))
 		{
 			throw new AutomationException("This is not among the refs: refs/remotes/origin/pr/{0}", PRNum);
-		}
+		}*/
 		RunGit(String.Format("fetch origin refs/pull/{0}/head:pr/{1}", PRNum, PRNum));
 		RunGit(String.Format("checkout pr/{0} --", PRNum));
 
