fix for utilizing ToLower instead of ToLowerInvariant.
#uat

diff --git a/Engine/Source/Programs/AutomationTool/Scripts/CopyBuildToStagingDirectory.Automation.cs b/Engine/Source/Programs/AutomationTool/Scripts/CopyBuildToStagingDirectory.Automation.cs
--- a/Engine/Source/Programs/AutomationTool/Scripts/CopyBuildToStagingDirectory.Automation.cs
+++ b/Engine/Source/Programs/AutomationTool/Scripts/CopyBuildToStagingDirectory.Automation.cs
@@ -201,7 +201,7 @@ public partial class Project : CommandUtils
 		string BuildPropertiesPath = CombinePaths(SC.LocalRoot, "Engine/Build");
 		if (SC.StageTargetPlatform.DeployLowerCaseFilenames(true))
 		{
-			BuildPropertiesPath = BuildPropertiesPath.ToLower();
+			BuildPropertiesPath = BuildPropertiesPath.ToLowerInvariant();
 		}
 		SC.StageFiles(StagedFileType.NonUFS, BuildPropertiesPath, "build.properties", false, null, null, true);
 
@@ -211,7 +211,7 @@ public partial class Project : CommandUtils
 		string CommandLineFile = "UE4CommandLine.txt";
 		if (SC.StageTargetPlatform.DeployLowerCaseFilenames(true))
 		{
-			CommandLineFile = CommandLineFile.ToLower();
+			CommandLineFile = CommandLineFile.ToLowerInvariant();
 		}
 		SC.StageFiles(StagedFileType.NonUFS, GetIntermediateCommandlineDir(SC), CommandLineFile, false, null, "", true, false);
 
@@ -522,7 +522,7 @@ public partial class Project : CommandUtils
 		var OutputRelativeLocation = CombinePaths(SC.RelativeProjectRootForStage, "Content/Paks/", PakName + "-" + SC.FinalCookPlatform + ".pak");
 		if (SC.StageTargetPlatform.DeployLowerCaseFilenames(true))
 		{
-			OutputRelativeLocation = OutputRelativeLocation.ToLower();
+			OutputRelativeLocation = OutputRelativeLocation.ToLowerInvariant();
 		}
 		OutputRelativeLocation = SC.StageTargetPlatform.Remap(OutputRelativeLocation);
 		var OutputLocation = CombinePaths(SC.RuntimeRootDir, OutputRelativeLocation);
@@ -544,7 +544,7 @@ public partial class Project : CommandUtils
 			var SourceOutputRelativeLocation = CombinePaths(SC.RelativeProjectRootForStage, "Content/Paks/", PakName + "-" + SC.CookPlatform + ".pak");
 			if (SC.CookSourcePlatform.DeployLowerCaseFilenames(true))
 			{
-				SourceOutputRelativeLocation = SourceOutputRelativeLocation.ToLower();
+				SourceOutputRelativeLocation = SourceOutputRelativeLocation.ToLowerInvariant();
 			}
 			SourceOutputRelativeLocation = SC.CookSourcePlatform.Remap(SourceOutputRelativeLocation);
 			var SourceOutputLocation = CombinePaths(SC.CookSourceRuntimeRootDir, SourceOutputRelativeLocation);
@@ -815,7 +815,7 @@ public partial class Project : CommandUtils
 		// @todo: Maybe there should be a new category - UFSNotForPak
 		if (SC.StageTargetPlatform.DeployLowerCaseFilenames(true))
 		{
-			IntermediateCmdLineFile = IntermediateCmdLineFile.ToLower();
+			IntermediateCmdLineFile = IntermediateCmdLineFile.ToLowerInvariant();
 		}
 		if (File.Exists(IntermediateCmdLineFile))
 		{
