Fix handling of multi-line p4 commit descriptions.

diff --git a/Engine/Source/Programs/AutomationTool/P4Utils.cs b/Engine/Source/Programs/AutomationTool/P4Utils.cs
--- a/Engine/Source/Programs/AutomationTool/P4Utils.cs
+++ b/Engine/Source/Programs/AutomationTool/P4Utils.cs
@@ -1350,7 +1350,7 @@ namespace AutomationTool
 			CheckP4Enabled();
 			var ChangeSpec = "Change: new" + "\n";
 			ChangeSpec += "Client: " + ((Owner != null) ? Owner : "") + "\n";
-			ChangeSpec += "Description: \n " + ((Description != null) ? Description : "(none)") + "\n";
+			ChangeSpec += "Description: " + ((Description != null) ? Description.Replace("\n", "\n\t") : "(none)") + "\n";
 			string CmdOutput;
 			int CL = 0;
 			if(AllowSpew)
