Add support for workspace mappings in pull request branches, and only reconcile files modified in Git. Prevents false-positives for line ending changes, and allows deletes.

diff --git a/Engine/Source/Programs/AutomationTool/P4Utils.cs b/Engine/Source/Programs/AutomationTool/P4Utils.cs
--- a/Engine/Source/Programs/AutomationTool/P4Utils.cs
+++ b/Engine/Source/Programs/AutomationTool/P4Utils.cs
@@ -1088,10 +1088,10 @@ namespace AutomationTool
 		/// </summary>
 		/// <param name="CL">Changelist to check the files out.</param>
 		/// <param name="CommandLine">Commandline for the command.</param>
-		public void Reconcile(int CL, string CommandLine)
+		public void Reconcile(int CL, string CommandLine, bool AllowSpew = true)
 		{
 			CheckP4Enabled();
-			LogP4("reconcile " + String.Format("-c {0} -ead -f ", CL) + CommandLine);
+			LogP4("reconcile " + String.Format("-c {0} -ead -f ", CL) + CommandLine, AllowSpew: AllowSpew);
 		}
 
         /// <summary>
