Add command line parameter to disable autosdks in UAT.

diff --git a/Engine/Source/Programs/AutomationTool/Automation.cs b/Engine/Source/Programs/AutomationTool/Automation.cs
--- a/Engine/Source/Programs/AutomationTool/Automation.cs
+++ b/Engine/Source/Programs/AutomationTool/Automation.cs
@@ -111,6 +111,7 @@ namespace AutomationTool
 		public static CommandLineArg NoKill = new CommandLineArg("-NoKill");
 		public static CommandLineArg Installed = new CommandLineArg("-Installed");
 		public static CommandLineArg UTF8Output = new CommandLineArg("-UTF8Output");
+		public static CommandLineArg NoAutoSDK = new CommandLineArg("-NoAutoSDK");
 
 		/// <summary>
 		/// Force initialize static members by calling this.
@@ -295,6 +296,12 @@ AutomationTool.exe [-verbose] [-compileonly] [-p4] Command0 [-Arg0 -Arg1 -Arg2 
 				return;
 			}
 
+			// Disable AutoSDKs if specified on the command line
+			if (GlobalCommandLine.NoAutoSDK)
+			{
+				UEBuildPlatform.bAllowAutoSDKSwitching = false;
+			}
+
 			// Setup environment
 			Log.TraceInformation("Setting up command environment.");
 			CommandUtils.InitCommandEnvironment();
