Allow building PS4MapFileUtil as part of promotions.

diff --git a/Engine/Source/Programs/AutomationTool/Scripts/BuildCommonTools.Automation.cs b/Engine/Source/Programs/AutomationTool/Scripts/BuildCommonTools.Automation.cs
--- a/Engine/Source/Programs/AutomationTool/Scripts/BuildCommonTools.Automation.cs
+++ b/Engine/Source/Programs/AutomationTool/Scripts/BuildCommonTools.Automation.cs
@@ -91,6 +91,12 @@ public class BuildCommonTools : BuildCommand
 			Agenda.AddTarget("ShaderCompileWorker", UnrealBuildTool.UnrealTargetPlatform.Mac, UnrealBuildTool.UnrealTargetConfiguration.Development, InAddArgs: "-CopyAppBundleBackToDevice");
 			Agenda.AddTarget("UE4EditorServices", UnrealBuildTool.UnrealTargetPlatform.Mac, UnrealBuildTool.UnrealTargetConfiguration.Development, InAddArgs: "-CopyAppBundleBackToDevice");
 		}
+
+		// PS4 binaries
+		if(Platforms.Contains(UnrealBuildTool.UnrealTargetPlatform.PS4))
+		{
+			Agenda.AddTarget("PS4MapFileUtil", UnrealBuildTool.UnrealTargetPlatform.Win64, UnrealBuildTool.UnrealTargetConfiguration.Development);
+		}
 		
 		return Agenda;
 	}
