Optionally allow incremental copy operations in UAT to skip files which exist in the destination folder

diff --git a/Engine/Source/Programs/AutomationTool/CommandUtils.cs b/Engine/Source/Programs/AutomationTool/CommandUtils.cs
--- a/Engine/Source/Programs/AutomationTool/CommandUtils.cs
+++ b/Engine/Source/Programs/AutomationTool/CommandUtils.cs
@@ -1241,17 +1241,23 @@ namespace AutomationTool
 		}
 
 		/// <summary>
-		/// Copies a file if the dest doesn't exist or the dest timestamp is different; after a copy, copies the timestamp
+		/// Copies a file if the dest doesn't exist or (optionally) if the dest timestamp is different; after a copy, copies the timestamp
 		/// </summary>
-		/// <param name="Source"></param>
-		/// <param name="Dest"></param>
+		/// <param name="Source">The full path to the source file</param>
+		/// <param name="Dest">The full path to the destination file</param>
+		/// <param name="bAllowDifferingTimestamps">If true, will always skip a file if the destination exists, even if timestamp differs; defaults to false</param>
 		/// <returns>True if the operation was successful, false otherwise.</returns>
-		public static void CopyFileIncremental(string Source, string Dest)
+		public static void CopyFileIncremental(string Source, string Dest, bool bAllowDifferingTimestamps = false)
 		{
 			Source = ConvertSeparators(PathSeparator.Default, Source);
 			Dest = ConvertSeparators(PathSeparator.Default, Dest);
 			if (InternalUtils.SafeFileExists(Dest, true))
 			{
+				if (bAllowDifferingTimestamps == true)
+				{
+					Log("CopyFileIncremental Skipping {0}, already exists", Dest);
+					return;
+				}
 				TimeSpan Diff = File.GetLastWriteTimeUtc(Dest) - File.GetLastWriteTimeUtc(Source);
 				if (Diff.TotalSeconds > -1 && Diff.TotalSeconds < 1)
 				{
