Changed UAT's IsProbablyAMacOrIOSExe() to handle binaries in app bundles outside of Binaries folder

diff --git a/Engine/Source/Programs/AutomationTool/BuildUtils.cs b/Engine/Source/Programs/AutomationTool/BuildUtils.cs
--- a/Engine/Source/Programs/AutomationTool/BuildUtils.cs
+++ b/Engine/Source/Programs/AutomationTool/BuildUtils.cs
@@ -92,7 +92,8 @@ namespace AutomationTool
                 ((
                     CommandUtils.CombinePaths(Filename).ToLower().Contains(CommandUtils.CombinePaths("Binaries", "Mac").ToLower()) ||
                     CommandUtils.CombinePaths(Filename).ToLower().Contains(CommandUtils.CombinePaths("Binaries", "IOS").ToLower()) ||
-					CommandUtils.CombinePaths(Filename).ToLower().Contains(CommandUtils.CombinePaths("Binaries", "ThirdParty").ToLower())
+					CommandUtils.CombinePaths(Filename).ToLower().Contains(CommandUtils.CombinePaths("Binaries", "ThirdParty").ToLower()) ||
+					CommandUtils.CombinePaths(Filename).ToLower().Contains(".app/Contents/MacOS".ToLower())
                 ) && (Path.GetExtension(Filename) == "" || Path.GetExtension(Filename) == "."));
         }
 
