Remove ENGINE_VERSION_HIWORD and ENGINE_VERSION_LOWORD. Nothing requires these as compile-time 16-bit values any more.

diff --git a/Engine/Source/Programs/AutomationTool/Scripts/Tests.Automation.cs b/Engine/Source/Programs/AutomationTool/Scripts/Tests.Automation.cs
--- a/Engine/Source/Programs/AutomationTool/Scripts/Tests.Automation.cs
+++ b/Engine/Source/Programs/AutomationTool/Scripts/Tests.Automation.cs
@@ -1630,8 +1630,6 @@ public class ZeroEngineVersions : BuildCommand
 			{
 				Log("Zeroing out engine versions in {0}", VersionFilename);
 				VersionH.ReplaceLine("#define ENGINE_VERSION ", "0");
-				VersionH.ReplaceLine("#define ENGINE_VERSION_HIWORD ", "0");
-				VersionH.ReplaceLine("#define ENGINE_VERSION_LOWORD ", "0");
 				VersionH.ReplaceLine("#define BRANCH_NAME ", "\"" + P4Env.BranchName + "\"");
 				VersionH.ReplaceLine("#define BUILT_FROM_CHANGELIST ", "0");
 
diff --git a/Engine/Source/Programs/AutomationTool/UE4Build.cs b/Engine/Source/Programs/AutomationTool/UE4Build.cs
--- a/Engine/Source/Programs/AutomationTool/UE4Build.cs
+++ b/Engine/Source/Programs/AutomationTool/UE4Build.cs
@@ -483,16 +483,12 @@ namespace AutomationTool
 				{
 					Log("Updating {0} with:", VerFile);
 					Log(" #define	ENGINE_VERSION  {0}", ChangelistNumber.ToString());
-					Log(" #define	ENGINE_VERSION_HIWORD  {0}", (ChangelistNumber >> 16).ToString());
-					Log(" #define	ENGINE_VERSION_LOWORD  {0}", (ChangelistNumber & 0xFFFF).ToString());					
 					Log(" #define	BRANCH_NAME  {0}", Branch);
 					Log(" #define	BUILT_FROM_CHANGELIST  {0}", ChangelistString);
 					Log(" #define   ENGINE_IS_LICENSEE_VERSION  {0}", bIsLicenseeVersion ? "1" : "0");
 
 					VersionFileUpdater VersionH = new VersionFileUpdater(VerFile);
 					VersionH.ReplaceLine("#define ENGINE_VERSION ", ChangelistNumber.ToString());
-					VersionH.ReplaceLine("#define ENGINE_VERSION_HIWORD ", (ChangelistNumber >> 16).ToString());
-					VersionH.ReplaceLine("#define ENGINE_VERSION_LOWORD ", (ChangelistNumber & 0xFFFF).ToString());					
 					VersionH.ReplaceLine("#define BRANCH_NAME ", "\"" + Branch + "\"");
 					VersionH.ReplaceLine("#define BUILT_FROM_CHANGELIST ", ChangelistString);
 					VersionH.ReplaceOrAddLine("#define ENGINE_IS_LICENSEE_VERSION ", bIsLicenseeVersion ? "1" : "0");
diff --git a/Engine/Source/Runtime/Core/Private/Mac/MacPlatformMisc.cpp b/Engine/Source/Runtime/Core/Private/Mac/MacPlatformMisc.cpp
--- a/Engine/Source/Runtime/Core/Private/Mac/MacPlatformMisc.cpp
+++ b/Engine/Source/Runtime/Core/Private/Mac/MacPlatformMisc.cpp
@@ -1723,9 +1723,9 @@ void FMacCrashContext::GenerateCrashInfoAndLaunchReporter() const
 			WriteLine(ReportFile, *GMacAppInfo.AppName);
 			
 			WriteUTF16String(ReportFile, TEXT("BuildVersion 1.0."));
-			WriteUTF16String(ReportFile, ItoTCHAR(ENGINE_VERSION_HIWORD, 10));
+			WriteUTF16String(ReportFile, ItoTCHAR(ENGINE_VERSION >> 16, 10));
 			WriteUTF16String(ReportFile, TEXT("."));
-			WriteLine(ReportFile, ItoTCHAR(ENGINE_VERSION_LOWORD, 10));
+			WriteLine(ReportFile, ItoTCHAR(ENGINE_VERSION & 0xffff, 10));
 			
 			WriteUTF16String(ReportFile, TEXT("CommandLine "));
 			WriteLine(ReportFile, *GMacAppInfo.CommandLine);
diff --git a/Engine/Source/Runtime/Launch/Resources/Version.h b/Engine/Source/Runtime/Launch/Resources/Version.h
--- a/Engine/Source/Runtime/Launch/Resources/Version.h
+++ b/Engine/Source/Runtime/Launch/Resources/Version.h
@@ -7,8 +7,6 @@
 #define ENGINE_PATCH_VERSION	0
 
 #define ENGINE_VERSION 0
-#define ENGINE_VERSION_HIWORD 0
-#define ENGINE_VERSION_LOWORD 0
 
 #define ENGINE_IS_LICENSEE_VERSION 0
 
