Handle error message from p4 print.

diff --git a/Engine/Source/Programs/AutomationTool/P4Utils.cs b/Engine/Source/Programs/AutomationTool/P4Utils.cs
--- a/Engine/Source/Programs/AutomationTool/P4Utils.cs
+++ b/Engine/Source/Programs/AutomationTool/P4Utils.cs
@@ -2403,6 +2403,10 @@ namespace AutomationTool
 			{
 				throw new AutomationException("p4 print {0} failed", DepotPath);
 			}
+			if(!Output.Trim().Contains("\n") && Output.Contains("no such file(s)"))
+			{
+				throw new AutomationException("p4 print {0} failed", DepotPath);
+			}
 			return Output;
 		}
 
