Add description param to commandline for signing binaries.

OPP-1770

diff --git a/Engine/Source/Programs/AutomationTool/CommandUtils.cs b/Engine/Source/Programs/AutomationTool/CommandUtils.cs
--- a/Engine/Source/Programs/AutomationTool/CommandUtils.cs
+++ b/Engine/Source/Programs/AutomationTool/CommandUtils.cs
@@ -2200,7 +2200,7 @@ namespace AutomationTool
 			//@TODO: Verbosity choosing
 			//  /v will spew lots of info
 			//  /q does nothing on success and minimal output on failure
-			string CodeSignArgs = String.Format("sign{0} /a /n \"{1}\" /t {2} /v \"{3}\"", SpecificStoreArg, SigningIdentity, TimestampServer, TargetFileInfo.FullName);
+			string CodeSignArgs = String.Format("sign{0} /a /n \"{1}\" /t {2} /d \"{3}\" /v \"{4}\"", SpecificStoreArg, SigningIdentity, TimestampServer, TargetFileInfo.Name, TargetFileInfo.FullName);
 
 			DateTime StartTime = DateTime.Now;
 
