Improved Windows XP support

- Windows XP is now configurable from the editor under: Project Settings | Windows | Minimum OS Version
- We will now have a UE4Game_xp.exe which will be used for XP supported content only games
- - - Added support to UBT to choose a file name when running a build, approved by Rob M.
- - - Added functionality to UAT to look for an _xp.exe when packaging.
- Minimum Windows OS, is Vista by default.
- Bootstrap Packaged Game Program is now also built with windows xp support.

Fix for JIRA issues
https://jira.ol.epicgames.net/browse/UE-2894
https://jira.ol.epicgames.net/browse/UE-3580
https://jira.ol.epicgames.net/browse/UE-4928

diff --git a/Engine/Source/Developer/Windows/WindowsTargetPlatform/Classes/WindowsTargetSettings.h b/Engine/Source/Developer/Windows/WindowsTargetPlatform/Classes/WindowsTargetSettings.h
--- a/Engine/Source/Developer/Windows/WindowsTargetPlatform/Classes/WindowsTargetSettings.h
+++ b/Engine/Source/Developer/Windows/WindowsTargetPlatform/Classes/WindowsTargetSettings.h
@@ -8,6 +8,12 @@
 
 #include "WindowsTargetSettings.generated.h"
 
+UENUM()
+enum class EMinimumSupportedOS : uint8
+{
+	MSOS_Vista = 0 UMETA(DisplayName = "Windows Vista"),
+	MSOS_XP = 1 UMETA(DisplayName = "Windows XP"),
+};
 
 /**
  * Implements the settings for the Windows target platform.
@@ -26,4 +32,10 @@ public:
 	 */
 	UPROPERTY(EditAnywhere, config, Category=Rendering)
 	TArray<FString> TargetedRHIs;
+
+	/**
+	 * Determine the minimum supported 
+	 */
+	UPROPERTY(EditAnywhere, config, Category=OSInfo)
+	TEnumAsByte<EMinimumSupportedOS> MinimumOSVersion;
 };
diff --git a/Engine/Source/Developer/Windows/WindowsTargetPlatform/Private/WindowsTargetPlatformClasses.cpp b/Engine/Source/Developer/Windows/WindowsTargetPlatform/Private/WindowsTargetPlatformClasses.cpp
--- a/Engine/Source/Developer/Windows/WindowsTargetPlatform/Private/WindowsTargetPlatformClasses.cpp
+++ b/Engine/Source/Developer/Windows/WindowsTargetPlatform/Private/WindowsTargetPlatformClasses.cpp
@@ -8,4 +8,6 @@
 
 UWindowsTargetSettings::UWindowsTargetSettings( const FObjectInitializer& ObjectInitializer )
 	: Super(ObjectInitializer)
-{ }
+{
+	MinimumOSVersion = EMinimumSupportedOS::MSOS_Vista;
+}
diff --git a/Engine/Source/Editor/MainFrame/Private/Frame/MainFrameActions.cpp b/Engine/Source/Editor/MainFrame/Private/Frame/MainFrameActions.cpp
--- a/Engine/Source/Editor/MainFrame/Private/Frame/MainFrameActions.cpp
+++ b/Engine/Source/Editor/MainFrame/Private/Frame/MainFrameActions.cpp
@@ -662,6 +662,15 @@ void FMainFrameActionCallbacks::PackageProject( const FName InPlatformInfoName )
 	{
 		OptionalParams += TEXT(" -targetplatform=Linux");
 	}
+	else if (PlatformInfo->TargetPlatformName == FName("WindowsNoEditor") && PlatformInfo->PlatformFlavor == TEXT("Win32"))
+	{
+		FString MinumumSupportedWindowsOS;
+		GConfig->GetString(TEXT("/Script/WindowsTargetPlatform.WindowsTargetSettings"), TEXT("MinimumOSVersion"), MinumumSupportedWindowsOS, GEngineIni);
+		if (MinumumSupportedWindowsOS == TEXT("MSOS_XP"))
+		{
+			OptionalParams += TEXT(" -OverrideMinimumOS=WinXP");
+		}
+	}
 
 	// Append any extra UAT flags specified for this platform flavor
 	if (!PlatformInfo->UATCommandLine.IsEmpty())
diff --git a/Engine/Source/Programs/AutomationTool/ProjectParams.cs b/Engine/Source/Programs/AutomationTool/ProjectParams.cs
--- a/Engine/Source/Programs/AutomationTool/ProjectParams.cs
+++ b/Engine/Source/Programs/AutomationTool/ProjectParams.cs
@@ -362,7 +362,8 @@ namespace AutomationTool
 			bool? Distribution = null,
             bool? OBBinAPK = null,
             bool? Prebuilt = null,
-            int? RunTimeoutSeconds = null
+            int? RunTimeoutSeconds = null,
+			string OverrideMinimumOS = null
 			)
 		{
 			//
@@ -504,6 +505,7 @@ namespace AutomationTool
 			this.DeviceUsername = ParseParamValueIfNotSpecified(Command, DeviceUsername, "deviceuser", String.Empty);
 			this.DevicePassword = ParseParamValueIfNotSpecified(Command, DevicePassword, "devicepass", String.Empty);
 			this.CrashReporter = GetParamValueIfNotSpecified(Command, CrashReporter, this.CrashReporter, "crashreporter");
+			this.OverrideMinimumOS = ParseParamValueIfNotSpecified(Command, OverrideMinimumOS, "OverrideMinimumOS", String.Empty);
 			if (ClientConfigsToBuild == null)
 			{
 				if (Command != null)
@@ -1153,6 +1155,8 @@ namespace AutomationTool
         [Help("RunTimeoutSeconds", "timeout to wait after we lunch the game")]
         public int RunTimeoutSeconds;
 
+		[Help("OverrideMinimumOS", "Determine a specific Minimum OS")]
+		public string OverrideMinimumOS;
 
 		#endregion
 
diff --git a/Engine/Source/Programs/AutomationTool/Win/WinPlatform.Automation.cs b/Engine/Source/Programs/AutomationTool/Win/WinPlatform.Automation.cs
--- a/Engine/Source/Programs/AutomationTool/Win/WinPlatform.Automation.cs
+++ b/Engine/Source/Programs/AutomationTool/Win/WinPlatform.Automation.cs
@@ -130,12 +130,19 @@ public abstract class BaseWinPlatform : Platform
 				}
 				else if (Exe.StartsWith(CombinePaths(SC.RuntimeRootDir, "Engine/Binaries", SC.PlatformDir)))
 				{
+					// Select the appropriate windows executable to stage, xp support or not.
+					string ExeFileName = Path.GetFileName(Exe);
+					if (!String.IsNullOrEmpty(Params.OverrideMinimumOS) && (Params.OverrideMinimumOS == "WinXP"))
+					{
+						ExeFileName = Path.GetFileNameWithoutExtension(ExeFileName) + "_xp" + Path.GetExtension(ExeFileName);
+					}
+
 					// keep it in the engine directory.
-					string SourceFile = CombinePaths(SC.LocalRoot, "Engine", "Binaries", SC.PlatformDir, Path.GetFileName(Exe));
+					string SourceFile = CombinePaths(SC.LocalRoot, "Engine", "Binaries", SC.PlatformDir, ExeFileName);
 					StageExecutable("exe", SC, CombinePaths(SC.LocalRoot, "Engine/Binaries", SC.PlatformDir), Path.GetFileNameWithoutExtension(SourceFile) + ".", true, null, null, false, WorkingFileType);
 					if(Exe == Exes[0] && !Params.NoBootstrapExe)
 					{
-						StageBootstrapExecutable(SC, SourceFile, CombinePaths("Engine", "Binaries", SC.PlatformDir, Path.GetFileName(Exe)), String.Format("..\\..\\..\\{0}\\{0}.uproject", SC.ShortProjectName));
+						StageBootstrapExecutable(SC, SourceFile, CombinePaths("Engine", "Binaries", SC.PlatformDir, ExeFileName), String.Format("..\\..\\..\\{0}\\{0}.uproject", SC.ShortProjectName));
 					}
 				}
 				else
diff --git a/Engine/Source/Programs/UnrealBuildTool/Configuration/UEBuildTarget.cs b/Engine/Source/Programs/UnrealBuildTool/Configuration/UEBuildTarget.cs
--- a/Engine/Source/Programs/UnrealBuildTool/Configuration/UEBuildTarget.cs
+++ b/Engine/Source/Programs/UnrealBuildTool/Configuration/UEBuildTarget.cs
@@ -231,6 +231,10 @@ namespace UnrealBuildTool
 				{
 					Platform = ParsedPlatform;
 				}
+				else if (Arguments[ArgumentIndex].ToLowerInvariant().StartsWith("-overridetargetappname="))
+				{
+					AdditionalDefinitions.Add(Arguments[ArgumentIndex]);
+				}
 				else
 				{
 					switch (Arguments[ArgumentIndex].ToUpperInvariant())
@@ -755,7 +759,17 @@ namespace UnrealBuildTool
 			List<OnlyModule> InOnlyModules,
 			bool bInEditorRecompile)
 		{
-			AppName = InAppName;
+			string CmdlineAppName = null;
+			const string OverrideTargetAppNameSwitch = "-overridetargetappname=";
+			if ((CmdlineAppName = InAdditionalDefinitions.Find(x => x.StartsWith(OverrideTargetAppNameSwitch))) != null)
+			{
+				AppName = CmdlineAppName.Substring(OverrideTargetAppNameSwitch.Length);
+			}
+			else
+			{
+				AppName = InAppName;
+			}
+
 			GameName = InGameName;
 			Platform = InPlatform;
 			Configuration = InConfiguration;
diff --git a/Engine/Source/Programs/Windows/BootstrapPackagedGame/BootstrapPackagedGame.Target.cs b/Engine/Source/Programs/Windows/BootstrapPackagedGame/BootstrapPackagedGame.Target.cs
--- a/Engine/Source/Programs/Windows/BootstrapPackagedGame/BootstrapPackagedGame.Target.cs
+++ b/Engine/Source/Programs/Windows/BootstrapPackagedGame/BootstrapPackagedGame.Target.cs
@@ -49,6 +49,11 @@ public class BootstrapPackagedGameTarget : TargetRules
 		UEBuildConfiguration.bBuildWithEditorOnlyData = false;
 		UEBuildConfiguration.bCompileAgainstEngine = false;
 		UEBuildConfiguration.bCompileAgainstCoreUObject = false;
+
+		if (Target.Platform == UnrealTargetPlatform.Win32)
+		{
+			WindowsPlatform.SupportWindowsXP = true;
+		}
 	}
 	
     public override bool GUBP_AlwaysBuildWithTools(UnrealTargetPlatform InHostPlatform, bool bBuildingRocket, out bool bInternalToolOnly, out bool SeparateNode)
