Modified the way AndroidPlatform.Automation.cs is launching adb tool so it works on Unix platforms

diff --git a/Engine/Source/Programs/AutomationTool/Android/AndroidPlatform.Automation.cs b/Engine/Source/Programs/AutomationTool/Android/AndroidPlatform.Automation.cs
--- a/Engine/Source/Programs/AutomationTool/Android/AndroidPlatform.Automation.cs
+++ b/Engine/Source/Programs/AutomationTool/Android/AndroidPlatform.Automation.cs
@@ -157,55 +157,55 @@ public class AndroidPlatform : Platform
 				string ApkName = GetFinalApkName(Params, SC.StageExecutables[0], true, bMakeSeparateApks ? Architecture : "", bMakeSeparateApks ? GPUArchitecture : "");
 				string BatchName = GetFinalBatchName(ApkName, Params, bMakeSeparateApks ? Architecture : "", bMakeSeparateApks ? GPUArchitecture : "");
 
-			if (!Params.Prebuilt)
-			{
-				string CookFlavor = SC.FinalCookPlatform.IndexOf("_") > 0 ? SC.FinalCookPlatform.Substring(SC.FinalCookPlatform.IndexOf("_")) : "";
-				string SOName = GetSONameWithoutArchitecture(Params, SC.StageExecutables[0]);
-				Deploy.PrepForUATPackageOrDeploy(Params.ShortProjectName, SC.ProjectRoot, SOName, SC.LocalRoot + "/Engine", Params.Distribution, CookFlavor);
-			}
+				if (!Params.Prebuilt)
+				{
+					string CookFlavor = SC.FinalCookPlatform.IndexOf("_") > 0 ? SC.FinalCookPlatform.Substring(SC.FinalCookPlatform.IndexOf("_")) : "";
+					string SOName = GetSONameWithoutArchitecture(Params, SC.StageExecutables[0]);
+					Deploy.PrepForUATPackageOrDeploy(Params.ShortProjectName, SC.ProjectRoot, SOName, SC.LocalRoot + "/Engine", Params.Distribution, CookFlavor);
+				}
 
-			// Create APK specific OBB in case we have a detached OBB.
-			string DeviceObbName = "";
-			string ObbName = "";
-			if (!Params.OBBinAPK)
-			{
-				DeviceObbName = GetDeviceObbName(ApkName);
-				ObbName = GetFinalObbName(ApkName);
-				CopyFile(LocalObbName, ObbName);
-			}
+				// Create APK specific OBB in case we have a detached OBB.
+				string DeviceObbName = "";
+				string ObbName = "";
+				if (!Params.OBBinAPK)
+				{
+					DeviceObbName = GetDeviceObbName(ApkName);
+					ObbName = GetFinalObbName(ApkName);
+					CopyFile(LocalObbName, ObbName);
+				}
 
-			// Write install batch file(s).
-
-			Log("Writing bat for install with {0}", Params.OBBinAPK ? "OBB in APK" : "OBB separate");
-			string PackageName = GetPackageInfo(ApkName, false);
-			// make a batch file that can be used to install the .apk and .obb files
-			string[] BatchLines = new string[] {
-				"setlocal",
-				"set ADB=%ANDROID_HOME%\\platform-tools\\adb.exe",
-				"set DEVICE=",
-				"if not \"%1\"==\"\" set DEVICE=-s %1",
-                "for /f \"delims=\" %%A in ('%ADB% %DEVICE% " + GetStorageQueryCommand() +"') do @set STORAGE=%%A",
-				"%ADB% %DEVICE% uninstall " + PackageName,
-				"%ADB% %DEVICE% install " + Path.GetFileName(ApkName),
-				"@if \"%ERRORLEVEL%\" NEQ \"0\" goto Error",
-				"%ADB% %DEVICE% shell rm -r %STORAGE%/" + Params.ShortProjectName,
-				"%ADB% %DEVICE% shell rm -r %STORAGE%/UE4Game/UE4CommandLine.txt", // we need to delete the commandline in UE4Game or it will mess up loading
-				"%ADB% %DEVICE% shell rm -r %STORAGE%/obb/" + PackageName,
-				Params.OBBinAPK ? "" : "%ADB% %DEVICE% push " + Path.GetFileName(ObbName) + " %STORAGE%/" + DeviceObbName,
-				Params.OBBinAPK ? "" : "if \"%ERRORLEVEL%\" NEQ \"0\" goto Error",
-				"goto:eof",
-				":Error",
-				"@echo.",
-				"@echo There was an error installing the game or the obb file. Look above for more info.",
-				"@echo.",
-				"@echo Things to try:",
-				"@echo Check that the device (and only the device) is listed with \"%ADB$ devices\" from a command prompt.",
-				"@echo Make sure all Developer options look normal on the device",
-				"@echo Check that the device has an SD card.",
-				"@pause"
-			};
-			File.WriteAllLines(BatchName, BatchLines);
-		}
+				// Write install batch file(s).
+
+				Log("Writing bat for install with {0}", Params.OBBinAPK ? "OBB in APK" : "OBB separate");
+				string PackageName = GetPackageInfo(ApkName, false);
+				// make a batch file that can be used to install the .apk and .obb files
+				string[] BatchLines = new string[] {
+					"setlocal",
+					"set ADB=%ANDROID_HOME%\\platform-tools\\adb.exe",
+					"set DEVICE=",
+					"if not \"%1\"==\"\" set DEVICE=-s %1",
+					"for /f \"delims=\" %%A in ('%ADB% %DEVICE% " + GetStorageQueryCommand() +"') do @set STORAGE=%%A",
+					"%ADB% %DEVICE% uninstall " + PackageName,
+					"%ADB% %DEVICE% install " + Path.GetFileName(ApkName),
+					"@if \"%ERRORLEVEL%\" NEQ \"0\" goto Error",
+					"%ADB% %DEVICE% shell rm -r %STORAGE%/" + Params.ShortProjectName,
+					"%ADB% %DEVICE% shell rm -r %STORAGE%/UE4Game/UE4CommandLine.txt", // we need to delete the commandline in UE4Game or it will mess up loading
+					"%ADB% %DEVICE% shell rm -r %STORAGE%/obb/" + PackageName,
+					Params.OBBinAPK ? "" : "%ADB% %DEVICE% push " + Path.GetFileName(ObbName) + " %STORAGE%/" + DeviceObbName,
+					Params.OBBinAPK ? "" : "if \"%ERRORLEVEL%\" NEQ \"0\" goto Error",
+					"goto:eof",
+					":Error",
+					"@echo.",
+					"@echo There was an error installing the game or the obb file. Look above for more info.",
+					"@echo.",
+					"@echo Things to try:",
+					"@echo Check that the device (and only the device) is listed with \"%ADB$ devices\" from a command prompt.",
+					"@echo Make sure all Developer options look normal on the device",
+					"@echo Check that the device has an SD card.",
+					"@pause"
+				};
+				File.WriteAllLines(BatchName, BatchLines);
+			}
 		}
 
 		PrintRunTime();
@@ -230,36 +230,36 @@ public class AndroidPlatform : Platform
 			foreach (string GPUArchitecture in GPUArchitectures)
 			{
 				string ApkName = GetFinalApkName(Params, SC.StageExecutables[0], true, bMakeSeparateApks ? Architecture : "", bMakeSeparateApks ? GPUArchitecture : "");
-			string ObbName = GetFinalObbName(ApkName);
+				string ObbName = GetFinalObbName(ApkName);
 				string BatchName = GetFinalBatchName(ApkName, Params, bMakeSeparateApks ? Architecture : "", bMakeSeparateApks ? GPUArchitecture : "");
 
-			// verify the files exist
-			if (!FileExists(ApkName))
-			{
-                string ErrorString = String.Format("ARCHIVE FAILED - {0} was not found", ApkName);
-                ErrorReporter.Error(ErrorString, (int)ErrorCodes.Error_AppNotFound);
-                throw new AutomationException(ErrorString);
-			}
-			if (!Params.OBBinAPK && !FileExists(ObbName))
-			{
-                string ErrorString = String.Format("ARCHIVE FAILED - {0} was not found", ObbName);
-                ErrorReporter.Error(ErrorString, (int)ErrorCodes.Error_ObbNotFound);
-                throw new AutomationException(ErrorString);
-            }
+				// verify the files exist
+				if (!FileExists(ApkName))
+				{
+					string ErrorString = String.Format("ARCHIVE FAILED - {0} was not found", ApkName);
+					ErrorReporter.Error(ErrorString, (int)ErrorCodes.Error_AppNotFound);
+					throw new AutomationException(ErrorString);
+				}
+				if (!Params.OBBinAPK && !FileExists(ObbName))
+				{
+					string ErrorString = String.Format("ARCHIVE FAILED - {0} was not found", ObbName);
+					ErrorReporter.Error(ErrorString, (int)ErrorCodes.Error_ObbNotFound);
+					throw new AutomationException(ErrorString);
+				}
 
-			SC.ArchiveFiles(Path.GetDirectoryName(ApkName), Path.GetFileName(ApkName));
-			if (!Params.OBBinAPK && !bAddedOBB)
-			{
-				bAddedOBB = true;
-				SC.ArchiveFiles(Path.GetDirectoryName(ObbName), Path.GetFileName(ObbName));
-			}
+				SC.ArchiveFiles(Path.GetDirectoryName(ApkName), Path.GetFileName(ApkName));
+				if (!Params.OBBinAPK && !bAddedOBB)
+				{
+					bAddedOBB = true;
+					SC.ArchiveFiles(Path.GetDirectoryName(ObbName), Path.GetFileName(ObbName));
+				}
 
-			SC.ArchiveFiles(Path.GetDirectoryName(BatchName), Path.GetFileName(BatchName));
+				SC.ArchiveFiles(Path.GetDirectoryName(BatchName), Path.GetFileName(BatchName));
+			}
 		}
 	}
-	}
 
-	private string GetAdbCommand(ProjectParams Params)
+	private string GetAdbCommandLine(ProjectParams Params, string Args)
 	{
 		string SerialNumber = Params.Device;
 		if (SerialNumber.Contains("@"))
@@ -272,14 +272,25 @@ public class AndroidPlatform : Platform
 		{
 			SerialNumber = " -s " + SerialNumber;
 		}
-		return Environment.ExpandEnvironmentVariables("/c %ANDROID_HOME%/platform-tools/adb.exe" + SerialNumber + " ");
+
+		string AdbCommand = Environment.ExpandEnvironmentVariables("%ANDROID_HOME%/platform-tools/adb" + (Utils.IsRunningOnMono ? "" : ".exe") + SerialNumber + " ");
+		return string.Format("{0} {1}{2}{3}{1}", Utils.IsRunningOnMono ? "-c" : "/c", Utils.IsRunningOnMono ? "'" : "", AdbCommand, Args);
+	}
+
+	private ProcessResult RunAdbCommand(ProjectParams Params, string Args, string Input = null, ERunOptions Options = ERunOptions.Default)
+	{
+		return Run(CmdEnv.CmdExe, GetAdbCommandLine(Params, Args), Input, Options);
+	}
+
+	private string RunAndLogAdbCommand(ProjectParams Params, string Args, out int SuccessCode)
+	{
+		return RunAndLog(CmdEnv, CmdEnv.CmdExe, GetAdbCommandLine(Params, Args), out SuccessCode);
 	}
 
 	public override void GetConnectedDevices(ProjectParams Params, out List<string> Devices)
 	{
 		Devices = new List<string>();
-		string AdbCommand = GetAdbCommand(Params);
-		ProcessResult Result = Run(CmdEnv.CmdExe, AdbCommand + "devices");
+		ProcessResult Result = RunAdbCommand(Params, "devices");
 
 		if (Result.Output.Length > 0)
 		{
@@ -312,8 +323,6 @@ public class AndroidPlatform : Platform
 						CommandUtils.LogWarning("Device attached but in bad state {0}:{1}", DeviceLine[0], DeviceLine[1]);
 					}
 				}
-
-
 			}
 		}
 	}
@@ -357,7 +366,6 @@ public class AndroidPlatform : Platform
 		string DeviceArchitecture = GetBestDeviceArchitecture(Params);
 		string GPUArchitecture = "";
 
-		string AdbCommand = GetAdbCommand(Params);
 		string ApkName = GetFinalApkName(Params, SC.StageExecutables[0], true, DeviceArchitecture, GPUArchitecture);
 
 		// make sure APK is up to date (this is fast if so)
@@ -373,13 +381,13 @@ public class AndroidPlatform : Platform
 		string PackageName = GetPackageInfo(ApkName, false);
 
 		// try uninstalling an old app with the same identifier.
-		string UninstallCommandline = AdbCommand + "uninstall " + PackageName;
-		RunAndLog(CmdEnv, CmdEnv.CmdExe, UninstallCommandline);
+		int SuccessCode = 0;
+		string UninstallCommandline = "uninstall " + PackageName;
+		RunAndLogAdbCommand(Params, UninstallCommandline, out SuccessCode);
 
 		// install the apk
-		int SuccessCode = 0;
-		string InstallCommandline = AdbCommand + "install \"" + ApkName + "\"";
-		string InstallOutput = RunAndLog(CmdEnv, CmdEnv.CmdExe, InstallCommandline, out SuccessCode);
+		string InstallCommandline = "install \"" + ApkName + "\"";
+		string InstallOutput = RunAndLogAdbCommand(Params, InstallCommandline, out SuccessCode);
 		int FailureIndex = InstallOutput.IndexOf("Failure"); 
 
 		// adb install doesn't always return an error code on failure, and instead prints "Failure", followed by an error code.
@@ -407,7 +415,7 @@ public class AndroidPlatform : Platform
 
         // Setup the OBB name and add the storage path (queried from the device) to it
         string DeviceStorageQueryCommand = GetStorageQueryCommand();
-        ProcessResult Result = Run(CmdEnv.CmdExe, AdbCommand + DeviceStorageQueryCommand, null, ERunOptions.AppMustExist);
+		ProcessResult Result = RunAdbCommand(Params, DeviceStorageQueryCommand, null, ERunOptions.AppMustExist);
         String StorageLocation = Result.Output.Trim();
         string DeviceObbName = StorageLocation + "/" + GetDeviceObbName(ApkName);
 
@@ -415,13 +423,13 @@ public class AndroidPlatform : Platform
 		if (SC.Stage)
 		{
 			// cache some strings
-			string BaseCommandline = AdbCommand + "push";
+			string BaseCommandline = "push";
 			string RemoteDir = StorageLocation + "/" + Params.ShortProjectName;
             string UE4GameRemoteDir = StorageLocation + "/" + Params.ShortProjectName;
 
 			// make sure device is at a clean state
-			Run(CmdEnv.CmdExe, AdbCommand + "shell rm -r " + RemoteDir);
-			Run(CmdEnv.CmdExe, AdbCommand + "shell rm -r " + UE4GameRemoteDir);
+			RunAdbCommand(Params, "shell rm -r " + RemoteDir);
+			RunAdbCommand(Params, "shell rm -r " + UE4GameRemoteDir);
 
 			// Copy UFS files..
 			string[] Files = Directory.GetFiles(SC.StageDirectory, "*", SearchOption.AllDirectories);
@@ -501,7 +509,7 @@ public class AndroidPlatform : Platform
 				// We run deploy commands in parallel to maximize the connection
 				// throughput.
 				DeployCommands.Add(
-					Run(CmdEnv.CmdExe, Commandline, null,
+					RunAdbCommand(Params, Commandline, null,
 						ERunOptions.Default | ERunOptions.NoWaitForExit));
 				// But we limit the parallel commands to avoid overwhelming
 				// memory resources.
@@ -524,7 +532,7 @@ public class AndroidPlatform : Platform
 			}
 
 			// delete the .obb file, since it will cause nothing we just deployed to be used
-			Run(CmdEnv.CmdExe, AdbCommand + "shell rm " + DeviceObbName);
+			RunAdbCommand(Params, "shell rm " + DeviceObbName);
 		}
 		else if (SC.Archive)
 		{
@@ -533,16 +541,16 @@ public class AndroidPlatform : Platform
 			if (File.Exists(ObbPath))
 			{
 				// cache some strings
-				string BaseCommandline = AdbCommand + "push";
+				string BaseCommandline = "push";
 
 				string Commandline = string.Format("{0} \"{1}\" \"{2}\"", BaseCommandline, ObbPath, DeviceObbName);
-				Run(CmdEnv.CmdExe, Commandline);
+				RunAdbCommand(Params, Commandline);
 			}
 		}
 		else
 		{
 			// cache some strings
-			string BaseCommandline = AdbCommand + "push";
+			string BaseCommandline = "push";
             string RemoteDir = StorageLocation + "/" + Params.ShortProjectName;
 
 			string FinalRemoteDir = RemoteDir;
@@ -556,7 +564,7 @@ public class AndroidPlatform : Platform
 
 			string RemoteFilename = IntermediateCmdLineFile.Replace(SC.StageDirectory, FinalRemoteDir).Replace("\\", "/");
 			string Commandline = string.Format("{0} \"{1}\" \"{2}\"", BaseCommandline, IntermediateCmdLineFile, RemoteFilename);
-			Run(CmdEnv.CmdExe, Commandline);
+			RunAdbCommand(Params, Commandline);
 		}
 	}
 
@@ -627,7 +635,7 @@ public class AndroidPlatform : Platform
 			throw new AutomationException("Failed to find %ANDROID_HOME%/build-tools subdirectory");
 		}
 		// we expect there to be one, so use the first one
-		return Path.Combine(Subdirs[0], "aapt.exe");
+		return Path.Combine(Subdirs[0], Utils.IsRunningOnMono ? "aapt" : "aapt.exe");
 	}
 
 	private string GetBestDeviceArchitecture(ProjectParams Params)
@@ -642,7 +650,7 @@ public class AndroidPlatform : Platform
 		string[] AppArchitectures = AndroidToolChain.GetAllArchitectures();
 
 		// ask the device
-		ProcessResult Result = Run(CmdEnv.CmdExe, GetAdbCommand(Params) + " shell getprop ro.product.cpu.abi", null, ERunOptions.AppMustExist);
+		ProcessResult Result = RunAdbCommand(Params, " shell getprop ro.product.cpu.abi", null, ERunOptions.AppMustExist);
 
 		// the output is just the architecture
 		string DeviceArch = UnrealBuildTool.Android.UEDeployAndroid.GetUE4Arch(Result.Output.Trim());
@@ -714,21 +722,20 @@ public class AndroidPlatform : Platform
             throw new AutomationException("Failed to get package name from " + ClientApp);
 		}
 
-		string AdbCommand = GetAdbCommand(Params);
 		string CommandLine = "shell am start -n " + PackageName + "/com.epicgames.ue4.GameActivity";
 
 		if (Params.Prebuilt)
 		{
 			// clear the log
-			Run(CmdEnv.CmdExe, AdbCommand + "logcat -c");
+			RunAdbCommand(Params, "logcat -c");
 		}
 
 		// Send a command to unlock the device before we try to run it
 		string UnlockCommandLine = "shell input keyevent 82";
-		Run(CmdEnv.CmdExe, AdbCommand + UnlockCommandLine, null);
+		RunAdbCommand(Params, UnlockCommandLine, null);
 
 		// start the app on device!
-		ProcessResult ClientProcess = Run(CmdEnv.CmdExe, AdbCommand + CommandLine, null, ClientRunFlags);
+		ProcessResult ClientProcess = RunAdbCommand(Params, CommandLine, null, ClientRunFlags);
 
 
 		if (Params.Prebuilt)
@@ -747,7 +754,7 @@ public class AndroidPlatform : Platform
 
 			while (true)
 			{
-				ProcessResult ProcessesResult = Run(CmdEnv.CmdExe, AdbCommand + "shell ps", null, ERunOptions.SpewIsVerbose);
+				ProcessResult ProcessesResult = RunAdbCommand(Params, "shell ps", null, ERunOptions.SpewIsVerbose);
 
 				string RunningProcessList = ProcessesResult.Output;
 				if (!RunningProcessList.Contains(PackageName))
@@ -765,10 +772,10 @@ public class AndroidPlatform : Platform
 			}
 
 			// this is just to get the ue4 log to go to the output
-			Run(CmdEnv.CmdExe, AdbCommand + "logcat -d -s UE4 -s Debug");
+			RunAdbCommand(Params, "logcat -d -s UE4 -s Debug");
 
 			// get the log we actually want to save
-			ProcessResult LogFileProcess = Run(CmdEnv.CmdExe, AdbCommand + "logcat -d", null, ERunOptions.AppMustExist);
+			ProcessResult LogFileProcess = RunAdbCommand(Params, "logcat -d", null, ERunOptions.AppMustExist);
 
 			File.WriteAllText(LogFilename, LogFileProcess.Output);
 			File.WriteAllText(ServerLogFilename, LogFileProcess.Output);
diff --git a/Engine/Source/Programs/AutomationTool/BuildEnvironment.cs b/Engine/Source/Programs/AutomationTool/BuildEnvironment.cs
--- a/Engine/Source/Programs/AutomationTool/BuildEnvironment.cs
+++ b/Engine/Source/Programs/AutomationTool/BuildEnvironment.cs
@@ -107,7 +107,7 @@ namespace AutomationTool
 			LogFolder = CommandUtils.GetEnvVar(EnvVarNames.LogFolder);
 			RobocopyExe = CommandUtils.CombinePaths(Environment.SystemDirectory, "robocopy.exe");
 			MountExe = CommandUtils.CombinePaths(Environment.SystemDirectory, "mount.exe");
-			CmdExe = CommandUtils.CombinePaths(Environment.SystemDirectory, "cmd.exe");
+			CmdExe = Utils.IsRunningOnMono ? "/bin/sh" : CommandUtils.CombinePaths(Environment.SystemDirectory, "cmd.exe");
 
 			if (String.IsNullOrEmpty(LogFolder))
 			{
