Xbox One symbol caching and loading

diff --git a/Engine/Source/Programs/AutomationTool/Scripts/BuildCommonTools.Automation.cs b/Engine/Source/Programs/AutomationTool/Scripts/BuildCommonTools.Automation.cs
--- a/Engine/Source/Programs/AutomationTool/Scripts/BuildCommonTools.Automation.cs
+++ b/Engine/Source/Programs/AutomationTool/Scripts/BuildCommonTools.Automation.cs
@@ -100,6 +100,12 @@ public class BuildCommonTools : BuildCommand
 			Agenda.AddTarget("PS4MapFileUtil", UnrealBuildTool.UnrealTargetPlatform.Win64, UnrealBuildTool.UnrealTargetConfiguration.Development);
 		}
 		
+		// Xbox One binaries
+		if(Platforms.Contains(UnrealBuildTool.UnrealTargetPlatform.XboxOne))
+		{
+			Agenda.AddTarget("XboxOnePDBFileUtil", UnrealBuildTool.UnrealTargetPlatform.Win64, UnrealBuildTool.UnrealTargetConfiguration.Development);
+		}
+		
 		return Agenda;
 	}
 }
diff --git a/Engine/Source/Programs/AutomationTool/Scripts/GUBP.Automation.cs b/Engine/Source/Programs/AutomationTool/Scripts/GUBP.Automation.cs
--- a/Engine/Source/Programs/AutomationTool/Scripts/GUBP.Automation.cs
+++ b/Engine/Source/Programs/AutomationTool/Scripts/GUBP.Automation.cs
@@ -1328,6 +1328,15 @@ public class GUBP : BuildCommand
                 }
                 AddDependency(SingleToolsNode.StaticGetFullName(HostPlatform, PS4MapFileUtil));
             }
+			if (TargetPlatform == UnrealTargetPlatform.XboxOne)
+			{
+				var XboxOnePDBFileUtil = bp.Branch.FindProgram("XboxOnePDBFileUtil");
+				if (XboxOnePDBFileUtil.Rules == null)
+				{
+					throw new AutomationException("XboxOnePDBFileUtil is not is this branch, but is required to build Xbox One monolithics");
+				}
+				AddDependency(SingleToolsNode.StaticGetFullName(HostPlatform, XboxOnePDBFileUtil));
+			}
             if (InGameProj.GameName != bp.Branch.BaseEngineProject.GameName && GameProj.Properties.Targets.ContainsKey(TargetRules.TargetType.Editor))
             {
                 AddPseudodependency(EditorGameNode.StaticGetFullName(InHostPlatform, GameProj));
