- Fixed apk selection for ES31, and removed log spam

diff --git a/Engine/Source/Programs/AutomationTool/Android/AndroidPlatform.Automation.cs b/Engine/Source/Programs/AutomationTool/Android/AndroidPlatform.Automation.cs
--- a/Engine/Source/Programs/AutomationTool/Android/AndroidPlatform.Automation.cs
+++ b/Engine/Source/Programs/AutomationTool/Android/AndroidPlatform.Automation.cs
@@ -709,17 +709,19 @@ public class AndroidPlatform : Platform
 			return "";
 		}
 
-		string[] AppGPUArchitectures = AndroidToolChain.GetAllArchitectures();
+		string[] AppGPUArchitectures = AndroidToolChain.GetAllGPUArchitectures();
 
 		// get the device extensions
-		ProcessResult ExtensionsResult = RunAdbCommand(Params, "shell dumpsys SurfaceFlinger");
+		ProcessResult ExtensionsResult = RunAdbCommand(Params, "shell dumpsys SurfaceFlinger", null, ERunOptions.AppMustExist);
 		string Extensions = ExtensionsResult.Output.Trim();
 
-		Console.WriteLine("Extensions: {0}", Extensions);
 		// look for AEP support
-		if (Extensions.Contains("GL_ANDROID_extension_pack_es31a") && Extensions.Contains("GL_EXT_color_buffer_half_float") && AppGPUArchitectures.Contains("-es31"))
+		if (Extensions.Contains("GL_ANDROID_extension_pack_es31a") && Extensions.Contains("GL_EXT_color_buffer_half_float"))
 		{
-			return "-es31";
+			if (AppGPUArchitectures.Contains("-es31"))
+			{
+				return "-es31";
+			}
 		}
 
 		return "-es2";
