Fix launchon bug causing ps4 builds to still read from the host pc

diff --git a/Engine/Source/Programs/AutomationTool/Scripts/RunProjectCommand.Automation.cs b/Engine/Source/Programs/AutomationTool/Scripts/RunProjectCommand.Automation.cs
--- a/Engine/Source/Programs/AutomationTool/Scripts/RunProjectCommand.Automation.cs
+++ b/Engine/Source/Programs/AutomationTool/Scripts/RunProjectCommand.Automation.cs
@@ -825,6 +825,11 @@ public partial class Project : CommandUtils
 		{
 			TempCmdLine += "-nullrhi ";
 		}
+		if (Params.Deploy && SC.StageTargetPlatform.PlatformType == UnrealTargetPlatform.PS4)
+		{
+			TempCmdLine += "-deployedbuild ";
+		}
+
 		TempCmdLine += "-CrashForUAT ";
 		TempCmdLine += Params.RunCommandline;
 
