Another fix for building on Mac (exception in AndroidPlatform on build farm due to missing .bash_profile)

diff --git a/Engine/Source/Programs/AutomationTool/Android/AndroidPlatform.Automation.cs b/Engine/Source/Programs/AutomationTool/Android/AndroidPlatform.Automation.cs
--- a/Engine/Source/Programs/AutomationTool/Android/AndroidPlatform.Automation.cs
+++ b/Engine/Source/Programs/AutomationTool/Android/AndroidPlatform.Automation.cs
@@ -23,13 +23,16 @@ public class AndroidPlatform : Platform
 		{
 			// Try reading env variable we need from .bash_profile
 			string BashProfilePath = Path.Combine(Environment.GetFolderPath(Environment.SpecialFolder.Personal), ".bash_profile");
-			string[] BashProfileContents = File.ReadAllLines(BashProfilePath);
-			foreach (string Line in BashProfileContents)
+			if(File.Exists(BashProfilePath))
 			{
-				if (Line.StartsWith("export ANDROID_HOME="))
+				string[] BashProfileContents = File.ReadAllLines(BashProfilePath);
+				foreach (string Line in BashProfileContents)
 				{
-					string PathVar = Line.Split('=')[1].Replace("\"", "");
-					Environment.SetEnvironmentVariable("ANDROID_HOME", PathVar);
+					if (Line.StartsWith("export ANDROID_HOME="))
+					{
+						string PathVar = Line.Split('=')[1].Replace("\"", "");
+						Environment.SetEnvironmentVariable("ANDROID_HOME", PathVar);
+					}
 				}
 			}
 		}
