Added new cook and createreleaseversion as optional params for the cooker to ufe.

diff --git a/Engine/Source/Programs/AutomationTool/ProjectParams.cs b/Engine/Source/Programs/AutomationTool/ProjectParams.cs
--- a/Engine/Source/Programs/AutomationTool/ProjectParams.cs
+++ b/Engine/Source/Programs/AutomationTool/ProjectParams.cs
@@ -196,6 +196,10 @@ namespace AutomationTool
 			this.DirectoriesToCook = InParams.DirectoriesToCook;
             this.InternationalizationPreset = InParams.InternationalizationPreset;
             this.CulturesToCook = InParams.CulturesToCook;
+            this.BasedOnReleaseVersion = InParams.BasedOnReleaseVersion;
+            this.CreateReleaseVersion = InParams.CreateReleaseVersion;
+            this.NewCook = InParams.NewCook;
+            this.AdditionalCookerOptions = InParams.AdditionalCookerOptions;
 			this.ClientCookedTargets = InParams.ClientCookedTargets;
 			this.ServerCookedTargets = InParams.ServerCookedTargets;
 			this.EditorTargets = InParams.EditorTargets;
@@ -323,6 +327,10 @@ namespace AutomationTool
             bool? UseDebugParamForEditorExe = null,
             bool? IterativeCooking = null,
 			bool? CookOnTheFly = null,
+            string AdditionalCookerOptions = null,
+            string BasedOnReleaseVersion = null,
+            string CreateReleaseVersion = null,
+            bool? NewCook = null,
 			bool? CrashReporter = null,
 			bool? DedicatedServer = null,
 			bool? Client = null,
@@ -421,6 +429,10 @@ namespace AutomationTool
 			this.Run = GetParamValueIfNotSpecified(Command, Run, this.Run, "run");
 			this.Cook = GetParamValueIfNotSpecified(Command, Cook, this.Cook, "cook");
 			this.CookFlavor = ParseParamValueIfNotSpecified(Command, CookFlavor, "cookflavor", String.Empty);
+            this.NewCook = GetParamValueIfNotSpecified(Command, NewCook, this.NewCook, "NewCook");
+            this.CreateReleaseVersion = ParseParamValueIfNotSpecified(Command, CreateReleaseVersion, "createreleaseversion", String.Empty);
+            this.BasedOnReleaseVersion = ParseParamValueIfNotSpecified(Command, BasedOnReleaseVersion, "basedonreleaseversion", String.Empty);
+            this.AdditionalCookerOptions = ParseParamValueIfNotSpecified(Command, AdditionalCookerOptions, "AdditionalCookerOptions", String.Empty);
 			this.SkipCook = GetParamValueIfNotSpecified(Command, SkipCook, this.SkipCook, "skipcook");
 			if (this.SkipCook)
 			{
@@ -912,6 +924,26 @@ namespace AutomationTool
         /// </summary>
         public string InternationalizationPreset;
 
+        /// <summary>
+        /// Cook: Create a cooked release version
+        /// </summary>
+        public string CreateReleaseVersion;
+
+        /// <summary>
+        /// Cook: Use new cooker (temporary cooker option until it's enabled permanently)
+        /// </summary>
+        public bool NewCook { private set; get; }
+
+        /// <summary>
+        /// Cook: Base this cook of a already released version of the cooked data
+        /// </summary>
+        public string BasedOnReleaseVersion;
+
+        /// <summary>
+        /// Cook: Additional cooker options to include on the cooker commandline
+        /// </summary>
+        public string AdditionalCookerOptions;
+
         /// <summary>
         /// Cook: List of cultures to cook.
         /// </summary>
@@ -1475,6 +1507,21 @@ namespace AutomationTool
             get { return !String.IsNullOrEmpty(InternationalizationPreset); }
         }
 
+        public bool HasBasedOnReleaseVersion
+        {
+            get { return !String.IsNullOrEmpty(BasedOnReleaseVersion); }
+        }
+
+        public bool HasAdditionalCookerOptions
+        {
+            get { return !String.IsNullOrEmpty(AdditionalCookerOptions); }
+        }
+
+        public bool HasCreateReleaseVersion
+        {
+            get { return !String.IsNullOrEmpty(CreateReleaseVersion); }
+        }
+
         public bool HasCulturesToCook
         {
             get { return !CommandUtils.IsNullOrEmpty(CulturesToCook); }
@@ -1761,6 +1808,9 @@ namespace AutomationTool
                 CommandUtils.Log("UseDebugParamForEditorExe={0}", UseDebugParamForEditorExe);
 				CommandUtils.Log("CookFlavor={0}", CookFlavor);
 				CommandUtils.Log("CookOnTheFly={0}", CookOnTheFly);
+                CommandUtils.Log("CreateReleaseVersion={0}", CreateReleaseVersion);
+                CommandUtils.Log("BasedOnReleaseVersion={0}", BasedOnReleaseVersion);
+                CommandUtils.Log("AdditionalCookerOptions={0}", AdditionalCookerOptions);
 				CommandUtils.Log("DedicatedServer={0}", DedicatedServer);
 				CommandUtils.Log("DirectoriesToCook={0}", DirectoriesToCook.ToString());
                 CommandUtils.Log("CulturesToCook={0}", CulturesToCook.ToString());
diff --git a/Engine/Source/Programs/AutomationTool/Scripts/CookCommand.Automation.cs b/Engine/Source/Programs/AutomationTool/Scripts/CookCommand.Automation.cs
--- a/Engine/Source/Programs/AutomationTool/Scripts/CookCommand.Automation.cs
+++ b/Engine/Source/Programs/AutomationTool/Scripts/CookCommand.Automation.cs
@@ -131,6 +131,23 @@ public partial class Project : CommandUtils
 				{
 					CommandletParams += " -iterate";
 				}
+                if (Params.NewCook)
+                {
+                    CommandletParams += " -newcook";
+                } 
+                if (Params.HasCreateReleaseVersion)
+                {
+                    CommandletParams += " -createreleaseversion=" + Params.CreateReleaseVersion;
+                }
+                if (Params.HasBasedOnReleaseVersion)
+                {
+                    CommandletParams += " -basedonreleaseversion=" + Params.BasedOnReleaseVersion;
+                }
+                if (Params.HasAdditionalCookerOptions)
+                {
+                    string FormatedAdditionalCookerParams = Params.AdditionalCookerOptions.TrimStart(new char[] {'\"', ' ' } ).TrimEnd(new char[] {'\"', ' ' } );
+                    CommandletParams += FormatedAdditionalCookerParams;
+                }
                 CookCommandlet(Params.RawProjectPath, Params.UE4Exe, Maps, Dirs, InternationalizationPreset, Cultures, CombineCommandletParams(PlatformsToCook.ToArray()), CommandletParams);
 			}
 			catch (Exception Ex)
