DLC Support part 1.
Added support for dlc to automation scripts.

diff --git a/Engine/Source/Programs/AutomationTool/ProjectParams.cs b/Engine/Source/Programs/AutomationTool/ProjectParams.cs
--- a/Engine/Source/Programs/AutomationTool/ProjectParams.cs
+++ b/Engine/Source/Programs/AutomationTool/ProjectParams.cs
@@ -198,6 +198,7 @@ namespace AutomationTool
             this.CulturesToCook = InParams.CulturesToCook;
             this.BasedOnReleaseVersion = InParams.BasedOnReleaseVersion;
             this.CreateReleaseVersion = InParams.CreateReleaseVersion;
+            this.DLCName = InParams.DLCName;
             this.NewCook = InParams.NewCook;
             this.AdditionalCookerOptions = InParams.AdditionalCookerOptions;
 			this.ClientCookedTargets = InParams.ClientCookedTargets;
@@ -330,6 +331,7 @@ namespace AutomationTool
             string AdditionalCookerOptions = null,
             string BasedOnReleaseVersion = null,
             string CreateReleaseVersion = null,
+            string DLCName = null,
             bool? NewCook = null,
 			bool? CrashReporter = null,
 			bool? DedicatedServer = null,
@@ -433,6 +435,7 @@ namespace AutomationTool
             this.CreateReleaseVersion = ParseParamValueIfNotSpecified(Command, CreateReleaseVersion, "createreleaseversion", String.Empty);
             this.BasedOnReleaseVersion = ParseParamValueIfNotSpecified(Command, BasedOnReleaseVersion, "basedonreleaseversion", String.Empty);
             this.AdditionalCookerOptions = ParseParamValueIfNotSpecified(Command, AdditionalCookerOptions, "AdditionalCookerOptions", String.Empty);
+            this.DLCName = ParseParamValueIfNotSpecified(Command, DLCName, "DLCName", String.Empty);
 			this.SkipCook = GetParamValueIfNotSpecified(Command, SkipCook, this.SkipCook, "skipcook");
 			if (this.SkipCook)
 			{
@@ -766,7 +769,16 @@ namespace AutomationTool
 		{
 			get
 			{
-				return Path.GetFullPath(String.IsNullOrEmpty(StageDirectoryParam) ? CommandUtils.CombinePaths(Path.GetDirectoryName(RawProjectPath), "Saved", "StagedBuilds") : StageDirectoryParam);
+                if( !String.IsNullOrEmpty(StageDirectoryParam ) )
+                {
+                    return Path.GetFullPath( StageDirectoryParam );
+                }
+                if ( HasDLCName )
+                {
+                     return Path.GetFullPath( CommandUtils.CombinePaths(Path.GetDirectoryName(RawProjectPath), "Plugins", DLCName, "Saved", "StagedBuilds" ) );
+                }
+                // default return the project saved\stagedbuilds directory
+                return Path.GetFullPath( CommandUtils.CombinePaths(Path.GetDirectoryName(RawProjectPath), "Saved", "StagedBuilds") );
 			}
 		}
 
@@ -939,6 +951,11 @@ namespace AutomationTool
         /// </summary>
         public string BasedOnReleaseVersion;
 
+        /// <summary>
+        /// Name of dlc to cook and package (if this paramter is supplied cooks the dlc and packages it into the dlc directory)
+        /// </summary>
+        public string DLCName;
+
         /// <summary>
         /// Cook: Additional cooker options to include on the cooker commandline
         /// </summary>
@@ -1517,6 +1534,11 @@ namespace AutomationTool
             get { return !String.IsNullOrEmpty(AdditionalCookerOptions); }
         }
 
+        public bool HasDLCName
+        {
+            get { return !String.IsNullOrEmpty(DLCName); }
+        }
+
         public bool HasCreateReleaseVersion
         {
             get { return !String.IsNullOrEmpty(CreateReleaseVersion); }
@@ -1810,6 +1832,7 @@ namespace AutomationTool
 				CommandUtils.Log("CookOnTheFly={0}", CookOnTheFly);
                 CommandUtils.Log("CreateReleaseVersion={0}", CreateReleaseVersion);
                 CommandUtils.Log("BasedOnReleaseVersion={0}", BasedOnReleaseVersion);
+                CommandUtils.Log("DLCName={0}", DLCName);
                 CommandUtils.Log("AdditionalCookerOptions={0}", AdditionalCookerOptions);
 				CommandUtils.Log("DedicatedServer={0}", DedicatedServer);
 				CommandUtils.Log("DirectoriesToCook={0}", DirectoriesToCook.ToString());
diff --git a/Engine/Source/Programs/AutomationTool/Scripts/CookCommand.Automation.cs b/Engine/Source/Programs/AutomationTool/Scripts/CookCommand.Automation.cs
--- a/Engine/Source/Programs/AutomationTool/Scripts/CookCommand.Automation.cs
+++ b/Engine/Source/Programs/AutomationTool/Scripts/CookCommand.Automation.cs
@@ -37,6 +37,10 @@ public partial class Project : CommandUtils
 
 		if (Params.CookOnTheFly && !Params.SkipServer)
 		{
+            if (Params.HasDLCName)
+            {
+                throw new AutomationException("Cook on the fly doesn't support cooking dlc");
+            }
 			if (Params.ClientTargetPlatforms.Count > 0)
 			{
 				var LogFolderOutsideOfSandbox = GetLogFolderOutsideOfSandbox();
@@ -117,24 +121,24 @@ public partial class Project : CommandUtils
 			}
 
             try
-			{
-				var CommandletParams = "-buildmachine -Unversioned -fileopenlog";
-				if (Params.UseDebugParamForEditorExe)
-				{
-					CommandletParams += " -debug";
-				}
-				if (Params.Manifests)
-				{
-					CommandletParams += " -manifests";
-				}
-				if (Params.IterativeCooking)
-				{
-					CommandletParams += " -iterate";
-				}
+            {
+                var CommandletParams = "-buildmachine -Unversioned -fileopenlog";
+                if (Params.UseDebugParamForEditorExe)
+                {
+                    CommandletParams += " -debug";
+                }
+                if (Params.Manifests)
+                {
+                    CommandletParams += " -manifests";
+                }
+                if (Params.IterativeCooking)
+                {
+                    CommandletParams += " -iterate";
+                }
                 if (Params.NewCook)
                 {
                     CommandletParams += " -newcook";
-                } 
+                }
                 if (Params.HasCreateReleaseVersion)
                 {
                     CommandletParams += " -createreleaseversion=" + Params.CreateReleaseVersion;
@@ -143,13 +147,17 @@ public partial class Project : CommandUtils
                 {
                     CommandletParams += " -basedonreleaseversion=" + Params.BasedOnReleaseVersion;
                 }
+                if (Params.HasDLCName)
+                {
+                    CommandletParams += " -dlcname=" + Params.DLCName;
+                }
                 if (Params.HasAdditionalCookerOptions)
                 {
-                    string FormatedAdditionalCookerParams = Params.AdditionalCookerOptions.TrimStart(new char[] {'\"', ' ' } ).TrimEnd(new char[] {'\"', ' ' } );
+                    string FormatedAdditionalCookerParams = Params.AdditionalCookerOptions.TrimStart(new char[] { '\"', ' ' }).TrimEnd(new char[] { '\"', ' ' });
                     CommandletParams += FormatedAdditionalCookerParams;
                 }
                 CookCommandlet(Params.RawProjectPath, Params.UE4Exe, Maps, Dirs, InternationalizationPreset, Cultures, CombineCommandletParams(PlatformsToCook.ToArray()), CommandletParams);
-			}
+            }
 			catch (Exception Ex)
 			{
 				// Delete cooked data (if any) as it may be incomplete / corrupted.
diff --git a/Engine/Source/Programs/AutomationTool/Scripts/CopyBuildToStagingDirectory.Automation.cs b/Engine/Source/Programs/AutomationTool/Scripts/CopyBuildToStagingDirectory.Automation.cs
--- a/Engine/Source/Programs/AutomationTool/Scripts/CopyBuildToStagingDirectory.Automation.cs
+++ b/Engine/Source/Programs/AutomationTool/Scripts/CopyBuildToStagingDirectory.Automation.cs
@@ -193,6 +193,34 @@ public partial class Project : CommandUtils
 		}
 		var ThisPlatform = SC.StageTargetPlatform;
 
+
+        if (Params.HasDLCName)
+        {
+            string DLCName = Params.DLCName;
+            string StagedBuildsDirectory = CombinePaths(SC.ProjectRoot, "Plugins", DLCName, "Saved", "StagedBuilds");
+
+            // Making a plugin, grab the binaries too
+            SC.StageFiles(StagedFileType.NonUFS, CombinePaths(SC.ProjectRoot, "Plugins", DLCName), "*.uplugin", true, null, null, true);
+            SC.StageFiles(StagedFileType.NonUFS, CombinePaths(SC.ProjectRoot, "Plugins", DLCName, "Binaries"), "libUE4-*.so", true, null, null, true);
+            SC.StageFiles(StagedFileType.NonUFS, CombinePaths(SC.ProjectRoot, "Plugins", DLCName, "Binaries"), "UE4-*.dll", true, null, null, true);
+            SC.StageFiles(StagedFileType.NonUFS, CombinePaths(SC.ProjectRoot, "Plugins", DLCName, "Binaries"), "libUE4Server-*.so", true, null, null, true);
+            SC.StageFiles(StagedFileType.NonUFS, CombinePaths(SC.ProjectRoot, "Plugins", DLCName, "Binaries"), "UE4Server-*.dll", true, null, null, true);
+            
+            // Put all of the cooked dir into the staged dir
+            if (SC.DedicatedServer)
+            {
+                // Dedicated server cook doesn't save shaders so no Engine dir is created
+                SC.StageFiles(StagedFileType.UFS, CombinePaths(SC.ProjectRoot, "Plugins", DLCName, "Saved", "Cooked", SC.CookPlatform), "*", true, new[] { "AssetRegistry.bin" }, "", true, !Params.UsePak(SC.StageTargetPlatform));
+            }
+            else
+            {
+                SC.StageFiles(StagedFileType.UFS, CombinePaths(SC.ProjectRoot, "Plugins", DLCName, "Saved", "Cooked", SC.CookPlatform), "*", true, new[] { "AssetRegistry.bin", CommandUtils.CombinePaths("Engine", "*") }, "", true, !Params.UsePak(SC.StageTargetPlatform));
+            }
+
+            return;
+        }
+
+
 		ThisPlatform.GetFilesToDeployOrStage(Params, SC);
 
 		// Get the build.properties file
