Fixed up compiler error.

diff --git a/Engine/Source/Programs/AutomationTool/Scripts/CopyBuildToStagingDirectory.Automation.cs b/Engine/Source/Programs/AutomationTool/Scripts/CopyBuildToStagingDirectory.Automation.cs
--- a/Engine/Source/Programs/AutomationTool/Scripts/CopyBuildToStagingDirectory.Automation.cs
+++ b/Engine/Source/Programs/AutomationTool/Scripts/CopyBuildToStagingDirectory.Automation.cs
@@ -197,7 +197,6 @@ public partial class Project : CommandUtils
         if (Params.HasDLCName)
         {
             string DLCName = Params.DLCName;
-            string StagedBuildsDirectory = CombinePaths(SC.ProjectRoot, "Plugins", DLCName, "Saved", "StagedBuilds");
 
             // Making a plugin, grab the binaries too
             SC.StageFiles(StagedFileType.NonUFS, CombinePaths(SC.ProjectRoot, "Plugins", DLCName), "*.uplugin", true, null, null, true);
