UE-6681 Cook on the Fly on iOS wait loop
fixed issue with PC side never making a process which the cross platform code used to determine if it was done or not.
#ios

diff --git a/Engine/Source/Programs/AutomationTool/IOS/IOSPlatform.Automation.cs b/Engine/Source/Programs/AutomationTool/IOS/IOSPlatform.Automation.cs
--- a/Engine/Source/Programs/AutomationTool/IOS/IOSPlatform.Automation.cs
+++ b/Engine/Source/Programs/AutomationTool/IOS/IOSPlatform.Automation.cs
@@ -972,32 +972,9 @@ public class IOSPlatform : Platform
 		}
 		else
 		{
-/*			// get the CFBundleIdentifier and modify the command line
-			int Pos = ClientCmdLine.IndexOf("-Exe=") + 6;
-			int EndPos = ClientCmdLine.IndexOf("-Targetplatform=") - 2;
-			string Exe = ClientCmdLine.Substring(Pos, EndPos - Pos);
-			
-			// check for Info.plist
-			if (string.IsNullOrEmpty(Params.StageDirectoryParam))
-			{
-				// need to crack open the ipa and read it from there - todo
-			}
-			else
-			{
-				if (File.Exists(Params.BaseStageDirectory+"/IOS/Info.plist"))
-				{
-					string Contents = File.ReadAllText(Params.BaseStageDirectory + "/IOS/Info.plist");
-					Pos = Contents.IndexOf("CFBundleIdentifier");
-					Pos = Contents.IndexOf("<string>", Pos) + 8;
-					EndPos = Contents.IndexOf("</string>", Pos);
-					string id = Contents.Substring(Pos, EndPos - Pos);
-					id = id.Substring(id.LastIndexOf(".") + 1);
-					ClientCmdLine = ClientCmdLine.Replace(Exe, id + ".stub");
-				}
-			}
-
-			return base.RunClient(ClientRunFlags, ClientApp, ClientCmdLine, Params);*/
-			return null;
+			ProcessResult Result = new ProcessResult("DummyApp", null, false, null);
+			Result.ExitCode = 0;
+			return Result;
 		}
 	}
 
diff --git a/Engine/Source/Programs/AutomationTool/ProcessUtils.cs b/Engine/Source/Programs/AutomationTool/ProcessUtils.cs
--- a/Engine/Source/Programs/AutomationTool/ProcessUtils.cs
+++ b/Engine/Source/Programs/AutomationTool/ProcessUtils.cs
@@ -228,7 +228,10 @@ namespace AutomationTool
 			Source = LogName;
 			AllowSpew = bAllowSpew;
 			this.SpewVerbosity = SpewVerbosity;
-			Proc.EnableRaisingEvents = false;
+			if (Proc != null)
+			{
+				Proc.EnableRaisingEvents = false;
+			}
 		}
 
         ~ProcessResult()
