rocket sample separation, samples can be built against specified rocket build, build individual samples

diff --git a/Engine/Source/Programs/AutomationTool/Scripts/GUBP.Automation.cs b/Engine/Source/Programs/AutomationTool/Scripts/GUBP.Automation.cs
--- a/Engine/Source/Programs/AutomationTool/Scripts/GUBP.Automation.cs
+++ b/Engine/Source/Programs/AutomationTool/Scripts/GUBP.Automation.cs
@@ -114,13 +114,17 @@ public class GUBP : BuildCommand
     public bool bFake = false;
     public static bool bNoIOSOnPC = false;
     public static bool bBuildRocket = false;
+	public static bool bPostedRocketBuild = false;
+	public static bool bBuildOnlySamples = false;
+	public static List<string> Samples = new List<string>();
+	public static string RocketBuild;	
     public static bool bForceIncrementalCompile = false;
     public static string ECProject = null;
     public string EmailHint;
     static public bool bPreflightBuild = false;
     public int PreflightShelveCL = 0;
     static public string PreflightMangleSuffix = "";
-    public GUBPBranchHacker.BranchOptions BranchOptions = null;
+    public GUBPBranchHacker.BranchOptions BranchOptions = null;	
 
     Dictionary<string, GUBPNode> GUBPNodes;
     Dictionary<string, bool> GUBPNodesCompleted;
@@ -144,11 +148,16 @@ public class GUBP : BuildCommand
 
     public string RocketUBTArgs(bool bUseRocketInsteadOfBuildRocket = false)
     {
+		string Args = "";
         if (bBuildRocket)
         {
-            return " -NoSimplygon " + (bUseRocketInsteadOfBuildRocket ? "-Rocket" : "-BuildRocket");
+            Args = " -NoSimplygon " + (bUseRocketInsteadOfBuildRocket ? "-Rocket" : "-BuildRocket");
+			if(bPostedRocketBuild && bUseRocketInsteadOfBuildRocket)
+			{
+				Args += " -PostedRocket";
+			}
         }
-        return "";
+        return Args;
     }
 
     public abstract class GUBPNodeAdder
@@ -264,7 +273,7 @@ public class GUBP : BuildCommand
         if(EmailHint == null)
         {
             EmailHint = "";
-        }
+        }			
         if (EmailHackers == null)
         {
             EmailHackers = new List<GUBPEmailHacker>();
@@ -361,7 +370,7 @@ public class GUBP : BuildCommand
         public List<string> BuildProducts = null;
 		public List<string> DependentPromotions = new List<string>();
         public List<string> AllDependencyBuildProducts = null;
-        public List<string> AllDependencies = null;
+        public List<string> AllDependencies = null;		
         public string AgentSharingGroup = "";
         public int ComputedDependentCISFrequencyQuantumShift = -1;
 
@@ -418,6 +427,14 @@ public class GUBP : BuildCommand
 		{
 			return false;
 		}
+		public virtual bool IsRocketSample()
+		{
+			return false;
+		}
+		public virtual string NodeHostPlatform()
+		{
+			return "";
+		}
         public virtual int AgentMemoryRequirement(GUBP bp)
         {
             return 0;
@@ -1265,11 +1282,11 @@ public class GUBP : BuildCommand
         public EditorGameNode(GUBP bp, UnrealTargetPlatform InHostPlatform, BranchInfo.BranchUProject InGameProj)
             : base(InHostPlatform)
         {
-            AgentSharingGroup = "Editor" + StaticGetHostPlatformSuffix(InHostPlatform);
+            AgentSharingGroup = "Editor" + (bBuildRocket ? "" : "Rocket") + StaticGetHostPlatformSuffix(InHostPlatform);
             GameProj = InGameProj;
-            AddDependency(RootEditorNode.StaticGetFullName(InHostPlatform));
-        }
-
+			AddDependency(RootEditorNode.StaticGetFullName(InHostPlatform));						
+		}
+		
         public static string StaticGetFullName(UnrealTargetPlatform InHostPlatform, BranchInfo.BranchUProject InGameProj)
         {
             return InGameProj.GameName + "_EditorGame" + StaticGetHostPlatformSuffix(InHostPlatform);
@@ -1281,7 +1298,7 @@ public class GUBP : BuildCommand
         public override string GameNameIfAnyForTempStorage()
         {
             return GameProj.GameName;
-        }
+        }		
         public override float Priority()
         {
             float Result = base.Priority();
@@ -2061,7 +2078,7 @@ public class GUBP : BuildCommand
         {
             return !bLabelPromoted;
         }
-    }    
+    }   	
 
     public class LabelPromotableNode : GUBPNode
     {
@@ -4708,8 +4725,7 @@ public class GUBP : BuildCommand
         if (ECProject == null)
         {
             ECProject = "";
-        }
-
+        }		
         HostPlatforms = new List<UnrealTargetPlatform>();
         if (!ParseParam("NoPC"))
         {
@@ -4742,15 +4758,36 @@ public class GUBP : BuildCommand
 			HostPlatforms.Add(UnrealTargetPlatform.Linux);
 		}
 
-		bBuildRocket = ParseParam("BuildRocket");
+		bBuildRocket = ParseParam("BuildRocket");		
+		string Sample = ParseParamValue("Sample");		
+		if (Sample != null)
+		{
+			Samples = new List<string>(Sample.Split('+'));
+			bBuildOnlySamples = true;
+		}
+		else
+		{
+			bBuildOnlySamples = ParseParam("OnlySamples");
+		}
         bForceIncrementalCompile = ParseParam("ForceIncrementalCompile");
         bool bNoAutomatedTesting = ParseParam("NoAutomatedTesting") || BranchOptions.bNoAutomatedTesting;		
         StoreName = ParseParamValue("Store");
         string StoreSuffix = ParseParamValue("StoreSuffix", "");
+
         if (bBuildRocket)
         {
             StoreSuffix = StoreSuffix + "-Rkt";
+			if (ParseParam("UsePostedRocket"))
+			{
+				bPostedRocketBuild = true;
+				RocketBuild = ParseParamValue("RocketBuild");
+				if (RocketBuild == null)
+				{
+					RocketBuild = FEngineVersionSupport.FromVersionFile(CommandUtils.CombinePaths(CommandUtils.CmdEnv.LocalRoot, @"Engine\Source\Runtime\Launch\Resources\Version.h")).ToString();
+				}
+			}
         }
+		
         if (bPreflightBuild)
         {
             int PreflightUID = ParseParamInt("PreflightUID", 0);
@@ -5864,7 +5901,12 @@ public class GUBP : BuildCommand
             {
                 // rocket is the shared promotable plus some other stuff and nothing else
                 bRelatedToNode = true;
-                NodeSpec = "Rocket_Aggregate";
+                NodeSpec = "";
+				if(!bBuildOnlySamples)
+				{
+					NodeSpec = "Rocket_Aggregate+";
+				}
+				NodeSpec = NodeSpec + "Sample_Aggregate";
                 if (!ParseParam("RocketValidate"))
                 {
                     //build a full promotable NodeSpec = SharedAggregatePromotableNode.StaticGetFullName() + "+" + NodeSpec;
@@ -6099,6 +6141,20 @@ public class GUBP : BuildCommand
 				FullNodeDependentPromotions.Add(NodeToDo, All);				
 			}
 		}
+		if(bBuildRocket)
+		{
+			foreach(var NodeToDo in NodesToDo)
+			{
+				if(GUBPNodes[NodeToDo].GetFullName().Contains("EditorGame"))
+				{
+					GUBPNodes[NodeToDo].AgentSharingGroup = "SampleEditors";
+					if(GUBPNodes[NodeToDo].GetFullName().Contains("OnMac"))
+					{
+						GUBPNodes[NodeToDo].AgentSharingGroup = "SampleEditors_OnMac";
+					}					
+				}
+			}
+		}
 		if (TimeIndex != 0)
 		{
 			Log("Culling based on time index");
@@ -6417,7 +6473,7 @@ public class GUBP : BuildCommand
             if (LastSticky == "" && bHaveECNodes)
             {
                 // if we don't have any sticky nodes and we have other nodes, we run a fake noop just to release the resource 
-                string Args = String.Format("{0}, subprocedure => 'GUBP_UAT_Node', parallel => '0', jobStepName => 'Noop', actualParameter => [{{actualParameterName => 'NodeName', value => 'Noop'}}, {{actualParametername => 'Sticky', value =>'1' }}], releaseMode => 'release'", BaseArgs);
+                string Args = String.Format("{0}, subprocedure => 'GUBP_UAT_Node', parallel => '0', jobStepName => 'Noop', actualParameter => [{{actualParameterName => 'NodeName', value => 'Noop'}}, {{actualParameterName => 'Sticky', value =>'1' }}], releaseMode => 'release'}});", BaseArgs);
                 StepList.Add(Args);
                 bHasNoop = true;
             }
@@ -6732,7 +6788,7 @@ public class GUBP : BuildCommand
         }
 
         var BuildProductToNodeMap = new Dictionary<string, string>();
-        foreach (var NodeToDo in OrdereredToDo)
+		foreach (var NodeToDo in OrdereredToDo)
         {
             if (GUBPNodes[NodeToDo].BuildProducts != null || GUBPNodes[NodeToDo].AllDependencyBuildProducts != null)
             {
@@ -6787,6 +6843,7 @@ public class GUBP : BuildCommand
             
             string GameNameIfAny = GUBPNodes[NodeToDo].GameNameIfAnyForTempStorage();
             string StorageRootIfAny = GUBPNodes[NodeToDo].RootIfAnyForTempStorage();
+					
             if (bFake)
             {
                 StorageRootIfAny = ""; // we don't rebase fake runs since those are entirely "records of success", which are always in the logs folder
@@ -6835,8 +6892,14 @@ public class GUBP : BuildCommand
                         GUBPNodes[NodeToDo].DoFakeBuild(this);
                     }
                     else
-                    {
-                        Log("***** Building GUBP Node {0} for {1}", NodeToDo, NodeStoreName);
+                    {                        
+						if(GUBPNodes[NodeToDo].IsRocketSample() && bPostedRocketBuild)
+						{
+
+							Log("***** Retrieving Build {0} for {1}", RocketBuild, GUBPNodes[NodeToDo].NodeHostPlatform());
+							RetrieveFromPermanentStorage(CmdEnv, "Rocket\\Automated", RocketBuild, GUBPNodes[NodeToDo].NodeHostPlatform());
+						}
+						Log("***** Building GUBP Node {0} for {1}", NodeToDo, NodeStoreName);
                         GUBPNodes[NodeToDo].DoBuild(this);
                     }
                     if (!GUBPNodes[NodeToDo].IsAggregate())
@@ -6850,11 +6913,12 @@ public class GUBP : BuildCommand
                                 try
                                 {
                                     bool WasLocal;
-                                    RetrieveFromTempStorage(CmdEnv, NodeStoreName, out WasLocal, GameNameIfAny, StorageRootIfAny);
-                                    if (!WasLocal)
-                                    {
-                                        throw new AutomationException("Retrieve was not local?");
-                                    }
+									RetrieveFromTempStorage(CmdEnv, NodeStoreName, out WasLocal, GameNameIfAny, StorageRootIfAny);
+									if (!WasLocal)
+									{
+										throw new AutomationException("Retrieve was not local?");
+									}
+																	                                    
                                 }
                                 catch(Exception Ex)
                                 {
diff --git a/Engine/Source/Programs/AutomationTool/TempStorage.cs b/Engine/Source/Programs/AutomationTool/TempStorage.cs
--- a/Engine/Source/Programs/AutomationTool/TempStorage.cs
+++ b/Engine/Source/Programs/AutomationTool/TempStorage.cs
@@ -58,7 +58,7 @@ namespace AutomationTool
                     bSilentOkToBeDifferent = bSilentOkToBeDifferent || (Name == "Engine/Binaries/DotNET/RPCUtility.pdb");
                     bSilentOkToBeDifferent = bSilentOkToBeDifferent || (Name == "Engine/Binaries/DotNET/AutomationTool.exe");
                     bSilentOkToBeDifferent = bSilentOkToBeDifferent || (Name == "Engine/Binaries/DotNET/UnrealBuildTool.exe");
-                    bSilentOkToBeDifferent = bSilentOkToBeDifferent || (Name == "Engine/Binaries/DotNET/UnrealBuildTool.exe.config");
+                    bSilentOkToBeDifferent = bSilentOkToBeDifferent || (Name == "Engine/Binaries/DotNET/UnrealBuildTool.exe.config");					
 
                     // Lets just allow all mac warnings to be slient
                     bSilentOkToBeDifferent = bSilentOkToBeDifferent || Name.Contains("Engine/Binaries/Mac");
@@ -812,7 +812,31 @@ namespace AutomationTool
             }
 
         }
-
+		public static void RetrieveFromPermanentStorage(CommandEnvironment Env, string NetworkRoot, string Build, string Platform)
+		{
+			var LocalManifest = LocalTempStorageManifestFilename(Env, Build);
+			if (!FileExists_NoExceptions(LocalManifest))
+			{
+				string BaseFolder = CombinePaths(Env.LocalRoot, "Rocket", "TempInst", Platform);
+				string Root = RootSharedTempStorageDirectory();
+				string SrcFolder = CombinePaths(Root, NetworkRoot, Build, Platform);
+				var SourceFiles = FindFiles_NoExceptions("*", true, SrcFolder);
+				foreach (var File in SourceFiles)
+				{
+					var DestFile = File.Replace(SrcFolder, "");
+					DestFile = CombinePaths(BaseFolder, DestFile);
+					CopyFile_NoExceptions(File, DestFile);
+				}
+				if(!DirectoryExists(Path.GetDirectoryName(LocalManifest)))
+				{
+					CreateDirectory(Path.GetDirectoryName(LocalManifest));
+				}
+				XmlDocument Local = new XmlDocument();
+				XmlElement RootElement = Local.CreateElement("permstorage");
+				Local.AppendChild(RootElement);
+				Local.Save(LocalManifest);
+			}
+		}
         public static List<string> RetrieveFromTempStorage(CommandEnvironment Env, string StorageBlockName, out bool WasLocal, string GameFolder = "", string BaseFolder = "")
         {
             if (String.IsNullOrEmpty(BaseFolder))
diff --git a/Engine/Source/Programs/AutomationTool/UE4Build.cs b/Engine/Source/Programs/AutomationTool/UE4Build.cs
--- a/Engine/Source/Programs/AutomationTool/UE4Build.cs
+++ b/Engine/Source/Programs/AutomationTool/UE4Build.cs
@@ -100,12 +100,13 @@ namespace AutomationTool
 			}
 		}
 
-
+		public static string BaseUBTDirectory = "";
+		
 		/// True if UBT is compiled and ready to build!
 		private bool bIsUBTReady = false;
-
+		
 		private void PrepareUBT()
-		{
+		{			
 			// Don't build UBT if we're running with pre-compiled binaries and if there's a debugger attached to this process.
 			// With the debugger attached, even though deleting the exe will work, the pdb files are still locked and the build will fail.
 			// Also, if we're running from VS then since UAT references UBT, we already have the most up-to-date version of UBT.exe
@@ -163,6 +164,17 @@ namespace AutomationTool
 			{
 				AddArgs += " -forcedebuginfo";
 			}
+			if(AddArgs.Contains("PostedRocket"))
+			{
+				if (UnrealBuildTool.BuildHostPlatform.Current.Platform == UnrealBuildTool.UnrealTargetPlatform.Win64)
+				{
+					BaseUBTDirectory = @"Rocket/TempInst/Windows";
+				}
+				if (UnrealBuildTool.BuildHostPlatform.Current.Platform == UnrealBuildTool.UnrealTargetPlatform.Mac)
+				{
+					BaseUBTDirectory = @"Rocket/TempInst/Mac";
+				}
+			}
 			PrepareUBT();
 
             string UBTManifest = GetUBTManifest(UprojectPath, AddArgs);
@@ -301,6 +313,17 @@ namespace AutomationTool
 			{
 				AddArgs += " -noxge";
 			}
+			if(AddArgs.Contains("PostedRocket"))
+			{
+				if (UnrealBuildTool.BuildHostPlatform.Current.Platform == UnrealBuildTool.UnrealTargetPlatform.Win64)
+				{
+					BaseUBTDirectory = @"Rocket/TempInst/Windows";
+				}
+				if (UnrealBuildTool.BuildHostPlatform.Current.Platform == UnrealBuildTool.UnrealTargetPlatform.Mac)
+				{
+					BaseUBTDirectory = @"Rocket/TempInst/Mac";
+				}
+			}
 
 			PrepareUBT();
 
@@ -702,7 +725,8 @@ namespace AutomationTool
 		public List<string> FindXGEFiles()
 		{
 			var Result = new List<string>();
-			Result.AddRange(FindFiles_NoExceptions("*.xge.xml", false, CombinePaths(CmdEnv.LocalRoot, @"\Engine\Intermediate\Build")));
+			var Root = CombinePaths(CmdEnv.LocalRoot, BaseUBTDirectory, @"\Engine\Intermediate\Build");			
+			Result.AddRange(FindFiles_NoExceptions("*.xge.xml", false, Root));
 			Result.Sort();
 			return Result;
 		}
@@ -1461,7 +1485,7 @@ namespace AutomationTool
 
             if ((GlobalCommandLine.Rocket || bForceRocket) && !String.IsNullOrEmpty(UProjectPath))
 			{
-				return Path.Combine(Path.GetDirectoryName(UProjectPath), "Intermediate/Build/Manifest.xml");
+				return Path.Combine(Path.GetDirectoryName(UProjectPath), "Intermediate/Build/Manifest.xml");				
 			}
 			else
 			{
@@ -1471,14 +1495,14 @@ namespace AutomationTool
 
 		public static string GetUBTExecutable()
 		{
-			return CommandUtils.CombinePaths(CmdEnv.LocalRoot, @"Engine/Binaries/DotNET/UnrealBuildTool.exe");
+			return CommandUtils.CombinePaths(CmdEnv.LocalRoot, BaseUBTDirectory, @"Engine/Binaries/DotNET/UnrealBuildTool.exe");
 		}
 
 		public string UBTExecutable
 		{
 			get
 			{
-				return GetUBTExecutable();
+				return GetUBTExecutable();							
 			}
 		}
 
diff --git a/Engine/Source/Programs/UnrealBuildTool/Configuration/EngineConfiguration.cs b/Engine/Source/Programs/UnrealBuildTool/Configuration/EngineConfiguration.cs
--- a/Engine/Source/Programs/UnrealBuildTool/Configuration/EngineConfiguration.cs
+++ b/Engine/Source/Programs/UnrealBuildTool/Configuration/EngineConfiguration.cs
@@ -309,64 +309,75 @@ namespace UnrealBuildTool
 		/// </summary>
 		public void ParseIniFile(string Filename)
 		{
-			var IniLines = File.ReadAllLines(Filename);
-			IniSection CurrentSection = null;
-
-			// Line Index for exceptions
-			var LineIndex = 1;
-			var bMultiLine = false;
-			var SingleValue = "";
-			var Key = "";
-			var LastAction = ParseAction.None;
-
-			// Parse each line
-			foreach (var Line in IniLines)
-			{
-				var TrimmedLine = Line.Trim();
-				// Multiline value support
-				bool bWasMultiLine = bMultiLine;
-				bMultiLine = TrimmedLine.EndsWith("\\");
-				if (bMultiLine)
-				{
-					TrimmedLine = TrimmedLine.Substring(0, TrimmedLine.Length - 1).TrimEnd();
-				}
-				if (!bWasMultiLine)
+			String[] IniLines = null;
+			try
+			{
+				IniLines = File.ReadAllLines(Filename);
+			}
+			catch (Exception ex)
+			{
+				Console.WriteLine("Error reading ini file: " + Filename + " Exception: " + ex.Message);
+			}
+			if (IniLines != null)
+			{
+				IniSection CurrentSection = null;
+
+				// Line Index for exceptions
+				var LineIndex = 1;
+				var bMultiLine = false;
+				var SingleValue = "";
+				var Key = "";
+				var LastAction = ParseAction.None;
+
+				// Parse each line
+				foreach (var Line in IniLines)
 				{
-					if (TrimmedLine.StartsWith("["))
+					var TrimmedLine = Line.Trim();
+					// Multiline value support
+					bool bWasMultiLine = bMultiLine;
+					bMultiLine = TrimmedLine.EndsWith("\\");
+					if (bMultiLine)
 					{
-						CurrentSection = FindOrAddSection(TrimmedLine, Filename, LineIndex);
-						LastAction = ParseAction.None;
+						TrimmedLine = TrimmedLine.Substring(0, TrimmedLine.Length - 1).TrimEnd();
 					}
-					else
+					if (!bWasMultiLine)
 					{
-						if (LastAction != ParseAction.None)
+						if (TrimmedLine.StartsWith("["))
 						{
-							throw new IniParsingException("Parsing new key/value pair when the previous one has not yet been processed ({0}, {1}) in {2}, line {3}: {4}", Key, SingleValue, Filename, LineIndex, TrimmedLine);
+							CurrentSection = FindOrAddSection(TrimmedLine, Filename, LineIndex);
+							LastAction = ParseAction.None;
 						}
-						// Check if the line is empty or a comment, also remove any +/- markers
-						LastAction = GetActionForLine(ref TrimmedLine);
-						if (LastAction != ParseAction.None)
+						else
 						{
-							if (CurrentSection == null)
+							if (LastAction != ParseAction.None)
 							{
-								throw new IniParsingException("Trying to parse key/value pair that doesn't belong to any section in {0}, line {1}: {2}", Filename, LineIndex, TrimmedLine);
+								throw new IniParsingException("Parsing new key/value pair when the previous one has not yet been processed ({0}, {1}) in {2}, line {3}: {4}", Key, SingleValue, Filename, LineIndex, TrimmedLine);
+							}
+							// Check if the line is empty or a comment, also remove any +/- markers
+							LastAction = GetActionForLine(ref TrimmedLine);
+							if (LastAction != ParseAction.None)
+							{
+								if (CurrentSection == null)
+								{
+									throw new IniParsingException("Trying to parse key/value pair that doesn't belong to any section in {0}, line {1}: {2}", Filename, LineIndex, TrimmedLine);
+								}
+								ParseKeyValuePair(TrimmedLine, Filename, LineIndex, out Key, out SingleValue);
 							}
-							ParseKeyValuePair(TrimmedLine, Filename, LineIndex, out Key, out SingleValue);
 						}
 					}
+					if (bWasMultiLine)
+					{
+						SingleValue += TrimmedLine;
+					}
+					if (!bMultiLine && LastAction != ParseAction.None)
+					{
+						ProcessKeyValuePair(CurrentSection, Key, SingleValue, LastAction);
+						LastAction = ParseAction.None;
+						SingleValue = "";
+						Key = "";
+					}
+					LineIndex++;
 				}
-				if (bWasMultiLine)
-				{
-					SingleValue += TrimmedLine;
-				}
-				if (!bMultiLine && LastAction != ParseAction.None)
-				{
-					ProcessKeyValuePair(CurrentSection, Key, SingleValue, LastAction);
-					LastAction = ParseAction.None;
-					SingleValue = "";
-					Key = "";
-				}
-				LineIndex++;
 			}
 		}
 
