build script, support nosamples

diff --git a/Engine/Source/Programs/AutomationTool/Scripts/GUBP.Automation.cs b/Engine/Source/Programs/AutomationTool/Scripts/GUBP.Automation.cs
--- a/Engine/Source/Programs/AutomationTool/Scripts/GUBP.Automation.cs
+++ b/Engine/Source/Programs/AutomationTool/Scripts/GUBP.Automation.cs
@@ -5906,7 +5906,10 @@ public class GUBP : BuildCommand
 				{
 					NodeSpec = "Rocket_Aggregate+";
 				}
-				NodeSpec = NodeSpec + "Sample_Aggregate";
+				if (!ParseParam("nosamples"))
+				{
+					NodeSpec = NodeSpec + "Sample_Aggregate";
+				}
                 if (!ParseParam("RocketValidate"))
                 {
                     //build a full promotable NodeSpec = SharedAggregatePromotableNode.StaticGetFullName() + "+" + NodeSpec;
