build script, fix regular game editor agent sharing groups

diff --git a/Engine/Source/Programs/AutomationTool/Scripts/GUBP.Automation.cs b/Engine/Source/Programs/AutomationTool/Scripts/GUBP.Automation.cs
--- a/Engine/Source/Programs/AutomationTool/Scripts/GUBP.Automation.cs
+++ b/Engine/Source/Programs/AutomationTool/Scripts/GUBP.Automation.cs
@@ -1282,7 +1282,7 @@ public class GUBP : BuildCommand
         public EditorGameNode(GUBP bp, UnrealTargetPlatform InHostPlatform, BranchInfo.BranchUProject InGameProj)
             : base(InHostPlatform)
         {
-            AgentSharingGroup = "Editor" + (bBuildRocket ? "" : "Rocket") + StaticGetHostPlatformSuffix(InHostPlatform);
+            AgentSharingGroup = "Editor"  + StaticGetHostPlatformSuffix(InHostPlatform);
             GameProj = InGameProj;
 			AddDependency(RootEditorNode.StaticGetFullName(InHostPlatform));						
 		}
