glue, don't delete main from the builders

diff --git a/Engine/Source/Programs/AutomationTool/Scripts/GUBP.Automation.cs b/Engine/Source/Programs/AutomationTool/Scripts/GUBP.Automation.cs
--- a/Engine/Source/Programs/AutomationTool/Scripts/GUBP.Automation.cs
+++ b/Engine/Source/Programs/AutomationTool/Scripts/GUBP.Automation.cs
@@ -3923,7 +3923,7 @@ public class GUBP : BuildCommand
                     Log("         Last Success: {0}", History.LastSucceeded);
                     Log("         Last Fail   : {0}", History.LastFailed);
                     Log("          Fails Since: {0}", History.FailedString);
-                    Log("     InProgress Since: {0}", History.InProgressString);
+                    Log("     InProgress Since: {0}", History.InProgressString);					
                 }
             }
             if (bShowDependencies)
@@ -3957,7 +3957,7 @@ public class GUBP : BuildCommand
 				{
 					Log("            depPro> {0}", Dep);
 				}
-			}
+			}			
         }
     }
     public void SaveGraphVisualization(List<string> Nodes)
@@ -4420,13 +4420,13 @@ public class GUBP : BuildCommand
                         History.InProgress.Add(Started);
                         History.InProgressString = GUBPNode.MergeSpaceStrings(History.InProgressString, String.Format("{0}", Started));
                     }
-                }
-                if (GUBPNodesHistory.ContainsKey(Node))
-                {
-                    GUBPNodesHistory.Remove(Node);
-                }
-                GUBPNodesHistory.Add(Node, History);
+                }                
             }
+			if (GUBPNodesHistory.ContainsKey(Node))
+			{
+				GUBPNodesHistory.Remove(Node);
+			}
+			GUBPNodesHistory.Add(Node, History);
         }
     }
 
@@ -5821,6 +5821,7 @@ public class GUBP : BuildCommand
         var FullNodeDirectDependencies = new Dictionary<string, string>();
 		var FullNodeDependedOnBy = new Dictionary<string, string>();
 		var FullNodeDependentPromotions = new Dictionary<string, string>();
+		var CurrentlyFailing = new Dictionary<string, string>();
 		var SeparatePromotables = new List<string>();
         {
             Log("******* {0} GUBP Nodes", GUBPNodes.Count);
@@ -5852,7 +5853,7 @@ public class GUBP : BuildCommand
                     Log("  {0}: {1}      {2}", Node, Note, All);
                     FullNodeList.Add(Node, Note);
                     FullNodeDirectDependencies.Add(Node, All);
-                    FullNodeListSortKey.Add(Node, Count);
+                    FullNodeListSortKey.Add(Node, Count);					
                     Count++;
                 }
                 else
@@ -6295,6 +6296,22 @@ public class GUBP : BuildCommand
                 {
                     UpdateNodeHistory(Node, CLString);
                 }
+				if (GUBPNodes[Node].RunInEC() && !GUBPNodes[Node].TriggerNode() && CLString != "")
+				{
+					try
+					{
+						var History = GUBPNodesHistory[Node];
+						if (History.LastFailed > History.LastSucceeded)
+						{
+							Log("Currently Failing Node {0} at CL {1}", Node, History.LastFailed.ToString());
+							CurrentlyFailing.Add(Node, History.LastFailed.ToString());
+						}
+					}
+					catch
+					{
+
+					}
+				}
             }
             var BuildDuration = (DateTime.UtcNow - StartTime).TotalMilliseconds;
             Log("Took {0}s to get history for {1} nodes", BuildDuration / 1000, NodesToDo.Count);
@@ -6319,8 +6336,7 @@ public class GUBP : BuildCommand
         }
 
         Log("*********** Desired And Dependent Nodes, in order.");
-        PrintNodes(this, OrdereredToDo, LocalOnly, UnfinishedTriggers);
-
+        PrintNodes(this, OrdereredToDo, LocalOnly, UnfinishedTriggers);		
         //check sorting
         {
             foreach (var NodeToDo in OrdereredToDo)
@@ -6387,6 +6403,10 @@ public class GUBP : BuildCommand
 			{
 				ECProps.Add(string.Format("PossiblePromotables/{0}={1}", Node, ""));
 			}
+			foreach (var NodePair in CurrentlyFailing)
+			{
+				ECProps.Add(string.Format("CurrentlyFailing/{0}={1}", NodePair.Key, NodePair.Value));
+			}
             var ECJobProps = new List<string>();
             if (ExplicitTrigger != "")
             {
