Implement iterative deploy for PS4
games also now sandbox their own files on /data/.

diff --git a/Engine/Source/Programs/AutomationTool/DeploymentContext.cs b/Engine/Source/Programs/AutomationTool/DeploymentContext.cs
--- a/Engine/Source/Programs/AutomationTool/DeploymentContext.cs
+++ b/Engine/Source/Programs/AutomationTool/DeploymentContext.cs
@@ -160,6 +160,20 @@ public class DeploymentContext //: ProjectParams
 	/// </summary>
 	public string ArchiveDirectory;
 
+	/// <summary>
+	/// Filename for the manifest of file changes for iterative deployment.
+	/// </summary>
+	static public readonly string UFSDeployDeltaFileName			= "Manifest_DeltaUFSFiles.txt";	
+	static public readonly string NonUFSDeployDeltaFileName			= "Manifest_DeltaNonUFSFiles.txt";
+
+	/// <summary>
+	/// Filename for the manifest of files currently deployed on a device.
+	/// </summary>
+	static public readonly string UFSDeployedManifestFileName		= "Manifest_UFSFiles.txt";
+	static public readonly string NonUFSDeployedManifestFileName	= "Manifest_NonUFSFiles.txt";
+
+	
+
 
 	/// <summary>
 	/// The client connects to dedicated server to get data
@@ -535,4 +549,14 @@ public class DeploymentContext //: ProjectParams
 
 		return FilesAdded;
 	}
+
+	public String GetUFSDeploymentDeltaPath()
+	{
+		return Path.Combine(StageDirectory, UFSDeployDeltaFileName);
+	}
+
+	public String GetNonUFSDeploymentDeltaPath()
+	{
+		return Path.Combine(StageDirectory, NonUFSDeployDeltaFileName);
+	}
 }
diff --git a/Engine/Source/Programs/AutomationTool/IOS/IOSPlatform.Automation.cs b/Engine/Source/Programs/AutomationTool/IOS/IOSPlatform.Automation.cs
--- a/Engine/Source/Programs/AutomationTool/IOS/IOSPlatform.Automation.cs
+++ b/Engine/Source/Programs/AutomationTool/IOS/IOSPlatform.Automation.cs
@@ -186,10 +186,11 @@ public class IOSPlatform : Platform
 		bool bNeedsIPA = false;
 		if (Params.IterativeDeploy)
 		{
+			String NonUFSManifestPath = SC.GetNonUFSDeploymentDeltaPath();
 			// check to determine if we need to update the IPA
-			if (File.Exists(SC.StageDirectory + "/Manifest_DeltaNonUFSFiles.txt"))
+			if (File.Exists(NonUFSManifestPath))
 			{
-				string NonUFSFiles = File.ReadAllText(SC.StageDirectory + "/Manifest_DeltaNonUFSFiles.txt");
+				string NonUFSFiles = File.ReadAllText(NonUFSManifestPath);
 				string[] Lines = NonUFSFiles.Split('\n');
 				bNeedsIPA = Lines.Length > 0 && !string.IsNullOrWhiteSpace(Lines[0]);
 			}
@@ -817,7 +818,7 @@ public class IOSPlatform : Platform
 				int EndPos = Contents.IndexOf("</string>", Pos);
 				BundleIdentifier = Contents.Substring(Pos, EndPos - Pos);
 			}
-			RunAndLog(CmdEnv, DeployServer, "Backup -file \"" + CombinePaths(Params.BaseStageDirectory, "IOS", "Manifest_UFSFiles.txt") + "\" -file \"" + CombinePaths(Params.BaseStageDirectory, "IOS", "Manifest_NonUFSFiles.txt") + "\"" + (String.IsNullOrEmpty(Params.Device) ? "" : " -device " + Params.Device.Substring(4)) + " -bundle " + BundleIdentifier);
+			RunAndLog(CmdEnv, DeployServer, "Backup -file \"" + CombinePaths(Params.BaseStageDirectory, "IOS", DeploymentContext.UFSDeployedManifestFileName) + "\" -file \"" + CombinePaths(Params.BaseStageDirectory, "IOS", DeploymentContext.NonUFSDeployedManifestFileName) + "\"" + (String.IsNullOrEmpty(Params.Device) ? "" : " -device " + Params.Device.Substring(4)) + " -bundle " + BundleIdentifier);
 
 			string[] ManifestFiles = Directory.GetFiles(CombinePaths(Params.BaseStageDirectory, "IOS"), "*_Manifest_UFS*.txt");
 			UFSManifests.AddRange(ManifestFiles);
@@ -878,9 +879,10 @@ public class IOSPlatform : Platform
 			}
 
 			// check to determine if we need to update the IPA
-			if (File.Exists(SC.StageDirectory + "/Manifest_DeltaNonUFSFiles.txt"))
+			String NonUFSManifestPath = SC.GetNonUFSDeploymentDeltaPath();
+			if (File.Exists(NonUFSManifestPath))
 			{
-				string NonUFSFiles = File.ReadAllText(SC.StageDirectory + "/Manifest_DeltaNonUFSFiles.txt");
+				string NonUFSFiles = File.ReadAllText(NonUFSManifestPath);
 				string[] Lines = NonUFSFiles.Split('\n');
 				bNeedsIPA = Lines.Length > 0 && !string.IsNullOrWhiteSpace(Lines[0]);
 			}
@@ -902,7 +904,7 @@ public class IOSPlatform : Platform
 		if (Params.IterativeDeploy)
 		{
 			// push over the changed files
-			RunAndLog(CmdEnv, DeployServer, "Deploy -manifest \"" + CombinePaths(Params.BaseStageDirectory, "IOS", "Manifest_DeltaUFSFiles.txt") + "\"" + (String.IsNullOrEmpty(Params.Device) ? "" : " -device " + Params.Device.Substring(4)) + AdditionalCommandline + " -bundle " + BundleIdentifier);
+			RunAndLog(CmdEnv, DeployServer, "Deploy -manifest \"" + CombinePaths(Params.BaseStageDirectory, "IOS", DeploymentContext.UFSDeployDeltaFileName) + "\"" + (String.IsNullOrEmpty(Params.Device) ? "" : " -device " + Params.Device.Substring(4)) + AdditionalCommandline + " -bundle " + BundleIdentifier);
 		}
 		Directory.SetCurrentDirectory (CurrentDir);
         PrintRunTime();
@@ -933,6 +935,14 @@ public class IOSPlatform : Platform
 
 	public override bool LaunchViaUFE { get { return UnrealBuildTool.BuildHostPlatform.Current.Platform != UnrealTargetPlatform.Mac; } }
 
+	public override bool UseAbsLog
+	{
+		get
+		{
+			return !LaunchViaUFE;
+		}
+	}
+
 	public override string Remap(string Dest)
 	{
 		return "cookeddata/" + Dest;
diff --git a/Engine/Source/Programs/AutomationTool/Platform.cs b/Engine/Source/Programs/AutomationTool/Platform.cs
--- a/Engine/Source/Programs/AutomationTool/Platform.cs
+++ b/Engine/Source/Programs/AutomationTool/Platform.cs
@@ -367,6 +367,14 @@ namespace AutomationTool
 			get { return false; }
 		}
 
+		/// <summary>
+		/// True if this platform can write to the abslog path that's on the host desktop.
+		/// </summary>
+		public virtual bool UseAbsLog
+		{
+			get { return true; }
+		}
+
 		/// <summary>
 		/// Remaps the given content directory to its final location
 		/// </summary>
diff --git a/Engine/Source/Programs/AutomationTool/ProjectParams.cs b/Engine/Source/Programs/AutomationTool/ProjectParams.cs
--- a/Engine/Source/Programs/AutomationTool/ProjectParams.cs
+++ b/Engine/Source/Programs/AutomationTool/ProjectParams.cs
@@ -250,6 +250,7 @@ namespace AutomationTool
 			this.Deploy = InParams.Deploy;
 			this.IterativeDeploy = InParams.IterativeDeploy;
 			this.Device = InParams.Device;
+			this.DeviceName = InParams.DeviceName;
 			this.ServerDevice = InParams.ServerDevice;
             this.NullRHI = InParams.NullRHI;
             this.FakeClient = InParams.FakeClient;
@@ -510,6 +511,17 @@ namespace AutomationTool
 			this.Deploy = GetParamValueIfNotSpecified(Command, Deploy, this.Deploy, "deploy");
 			this.IterativeDeploy = GetParamValueIfNotSpecified(Command, IterativeDeploy, this.IterativeDeploy, "iterativedeploy", "iterate");
 			this.Device = ParseParamValueIfNotSpecified(Command, Device, "device", String.Empty).Trim(new char[] { '\"' });
+
+			// strip the platform prefix the specified device.
+			if (this.Device.Contains("@"))
+			{
+				this.DeviceName = this.Device.Substring(this.Device.IndexOf("@") + 1);
+			}
+			else
+			{
+				this.DeviceName = this.Device;
+			}
+
 			this.ServerDevice = ParseParamValueIfNotSpecified(Command, ServerDevice, "serverdevice", this.Device);
 			this.NullRHI = GetParamValueIfNotSpecified(Command, NullRHI, this.NullRHI, "nullrhi");
 			this.FakeClient = GetParamValueIfNotSpecified(Command, FakeClient, this.FakeClient, "fakeclient");
@@ -1107,11 +1119,17 @@ namespace AutomationTool
 		public string AdditionalServerMapParams;
 
 		/// <summary>
-		/// Run: The target device to run the game on
+		/// Run: The target device to run the game on.  Comes in the form platform@devicename.
 		/// </summary>
 		[Help("device", "Device to run the game on")]
 		public string Device;
 
+		/// <summary>
+		/// Run: The target device to run the game on.  No platform prefix.
+		/// </summary>
+		[Help("device", "Device name without the platform prefix to run the game on")]
+		public string DeviceName;
+
 		/// <summary>
 		/// Run: the target device to run the server on
 		/// </summary>
diff --git a/Engine/Source/Programs/AutomationTool/Scripts/CopyBuildToStagingDirectory.Automation.cs b/Engine/Source/Programs/AutomationTool/Scripts/CopyBuildToStagingDirectory.Automation.cs
--- a/Engine/Source/Programs/AutomationTool/Scripts/CopyBuildToStagingDirectory.Automation.cs
+++ b/Engine/Source/Programs/AutomationTool/Scripts/CopyBuildToStagingDirectory.Automation.cs
@@ -1138,8 +1138,8 @@ public partial class Project : CommandUtils
 		// add the manifest
 		if (!DeltaManifest.Contains("NonUFS"))
 		{
-			DeltaFiles.Add("Manifest_NonUFSFiles.txt");
-			DeltaFiles.Add("Manifest_UFSFiles.txt");
+			DeltaFiles.Add(DeploymentContext.NonUFSDeployedManifestFileName);
+			DeltaFiles.Add(DeploymentContext.UFSDeployedManifestFileName);
 		}
 
 		// TODO: determine files which need to be removed
@@ -1284,12 +1284,12 @@ public partial class Project : CommandUtils
 						}
 
 						// get the staged file data
-						Dictionary<string, string> StagedUFSFiles = ReadStagedManifest(Params, SC, "Manifest_UFSFiles.txt");
-						Dictionary<string, string> StagedNonUFSFiles = ReadStagedManifest(Params, SC, "Manifest_NonUFSFiles.txt");
+						Dictionary<string, string> StagedUFSFiles = ReadStagedManifest(Params, SC, DeploymentContext.UFSDeployedManifestFileName);
+						Dictionary<string, string> StagedNonUFSFiles = ReadStagedManifest(Params, SC, DeploymentContext.NonUFSDeployedManifestFileName);
 
 						// write out the delta file data
-						WriteDeltaManifest(Params, SC, DeployedUFSFiles, StagedUFSFiles, "Manifest_DeltaUFSFiles.txt");
-						WriteDeltaManifest(Params, SC, DeployedNonUFSFiles, StagedNonUFSFiles, "Manifest_DeltaNonUFSFiles.txt");
+						WriteDeltaManifest(Params, SC, DeployedUFSFiles, StagedUFSFiles, DeploymentContext.UFSDeployDeltaFileName);
+						WriteDeltaManifest(Params, SC, DeployedNonUFSFiles, StagedNonUFSFiles, DeploymentContext.NonUFSDeployDeltaFileName);
 					}
 				}
 			}
diff --git a/Engine/Source/Programs/AutomationTool/Scripts/RunProjectCommand.Automation.cs b/Engine/Source/Programs/AutomationTool/Scripts/RunProjectCommand.Automation.cs
--- a/Engine/Source/Programs/AutomationTool/Scripts/RunProjectCommand.Automation.cs
+++ b/Engine/Source/Programs/AutomationTool/Scripts/RunProjectCommand.Automation.cs
@@ -808,7 +808,7 @@ public partial class Project : CommandUtils
 		{
 			TempCmdLine += "-execcmds=\"automation list, automation runall\" ";
 		}
-		if (!SC.StageTargetPlatform.LaunchViaUFE)
+		if (SC.StageTargetPlatform.UseAbsLog)
 		{
 			TempCmdLine += "-abslog=" + CommandUtils.MakePathSafeToUseWithCommandLine(ClientLogFile) + " ";
 		}
diff --git a/Engine/Source/Runtime/NetworkFile/Private/NetworkPlatformFile.cpp b/Engine/Source/Runtime/NetworkFile/Private/NetworkPlatformFile.cpp
--- a/Engine/Source/Runtime/NetworkFile/Private/NetworkPlatformFile.cpp
+++ b/Engine/Source/Runtime/NetworkFile/Private/NetworkPlatformFile.cpp
@@ -254,7 +254,7 @@ void FNetworkPlatformFile::InitializeAfterSetActive()
 				ConvertServerFilenameToClientFilename(ServerFile);
 
 				// Set it in the visitor file times list
-				// If there is any pathing difference (casing, or relative path, or whatever) between the server's filelist and the results
+				// If there is any pathing difference (relative path, or whatever) between the server's filelist and the results
 				// of platform directory iteration then this will Add a new entry rather than override the existing one.  This causes local file deletes
 				// and longer loads as we will never see the benefits of local device caching.
 				Visitor.FileTimes.Add(ServerFile, FDateTime::MinValue());
