Add a separate GUBP node fo generating the target platform file list, to fix getting target file list for Android on Mac in RocketBuild.

diff --git a/Engine/Source/Programs/AutomationTool/UE4Build.cs b/Engine/Source/Programs/AutomationTool/UE4Build.cs
--- a/Engine/Source/Programs/AutomationTool/UE4Build.cs
+++ b/Engine/Source/Programs/AutomationTool/UE4Build.cs
@@ -1125,18 +1125,20 @@ namespace AutomationTool
 			return UnrealBuildTool.UEBuildPlatform.GetBuildPlatform(Platform).CanUseXGE();
 		}
 
-		public string[] GetExternalFileList(BuildAgenda Agenda)
+		public string GenerateExternalFileList(BuildAgenda Agenda)
 		{
-			// Create the file list
 			string AdditionalArguments = "";
 			foreach (var Target in Agenda.Targets)
 			{
 				RunUBT(CmdEnv, UBTExecutable, Target.UprojectPath, Target.TargetName, Target.Platform.ToString(), Target.Config.ToString(), "-generateexternalfilelist" + AdditionalArguments + " " + Target.UBTArgs);
 				AdditionalArguments = " -mergeexternalfilelist";
 			}
+			return CommandUtils.CombinePaths(CmdEnv.LocalRoot, @"/Engine/Intermediate/Build/ExternalFiles.xml");
+		}
 
-			// Read it in, and return the filenames as an array of strings
-			string FileListPath = CommandUtils.CombinePaths(CmdEnv.LocalRoot, @"/Engine/Intermediate/Build/ExternalFiles.xml");
+		public string[] GetExternalFileList(BuildAgenda Agenda)
+		{
+			string FileListPath = GenerateExternalFileList(Agenda);
 			return UnrealBuildTool.Utils.ReadClass<UnrealBuildTool.ExternalFileList>(FileListPath).FileNames.ToArray();
 		}
 
