Build iOS tools when promoting a GitHub build.

diff --git a/Engine/Source/Programs/AutomationTool/Scripts/BuildCommonTools.Automation.cs b/Engine/Source/Programs/AutomationTool/Scripts/BuildCommonTools.Automation.cs
--- a/Engine/Source/Programs/AutomationTool/Scripts/BuildCommonTools.Automation.cs
+++ b/Engine/Source/Programs/AutomationTool/Scripts/BuildCommonTools.Automation.cs
@@ -31,8 +31,21 @@ public class BuildCommonTools : BuildCommand
 			Platforms.Add(Platform);
 		}
 
+		// Add all the platforms if specified
+		if(ParseParam("allplatforms"))
+		{
+			foreach(UnrealTargetPlatform Platform in Enum.GetValues(typeof(UnrealTargetPlatform)))
+			{
+				if(!Platforms.Contains(Platform))
+				{
+					Platforms.Add(Platform);
+				}
+			}
+		}
+
 		// Get the agenda
-		UE4Build.BuildAgenda Agenda = MakeAgenda(Platforms.ToArray());
+		List<string> ExtraBuildProducts = new List<string>();
+		UE4Build.BuildAgenda Agenda = MakeAgenda(Platforms.ToArray(), ExtraBuildProducts);
 
 		// Build everything. We don't want to touch version files for GitHub builds -- these are "programmer builds" and won't have a canonical build version
 		UE4Build Builder = new UE4Build(this);
@@ -42,6 +55,12 @@ public class BuildCommonTools : BuildCommand
 		Builder.AddUATFilesToBuildProducts();
 		Builder.AddUBTFilesToBuildProducts();
 
+		// Add all the extra build products
+		foreach(string ExtraBuildProduct in ExtraBuildProducts)
+		{
+			Builder.AddBuildProduct(ExtraBuildProduct);
+		}
+
 		// Make sure all the build products exist
 		UE4Build.CheckBuildProducts(Builder.BuildProductFiles);
 
@@ -58,7 +77,7 @@ public class BuildCommonTools : BuildCommand
 		}
 	}
 
-	public static UE4Build.BuildAgenda MakeAgenda(UnrealBuildTool.UnrealTargetPlatform[] Platforms)
+	public static UE4Build.BuildAgenda MakeAgenda(UnrealBuildTool.UnrealTargetPlatform[] Platforms, List<string> ExtraBuildProducts)
 	{
 		// Create the build agenda
 		UE4Build.BuildAgenda Agenda = new UE4Build.BuildAgenda();
@@ -94,6 +113,22 @@ public class BuildCommonTools : BuildCommand
 			Agenda.AddTarget("UE4EditorServices", UnrealBuildTool.UnrealTargetPlatform.Mac, UnrealBuildTool.UnrealTargetConfiguration.Development, InAddArgs: "-CopyAppBundleBackToDevice");
 		}
 
+		// iOS binaries
+		if(Platforms.Contains(UnrealBuildTool.UnrealTargetPlatform.IOS))
+		{
+			Agenda.DotNetProjects.Add(@"Engine/Source/Programs/iOS/iPhonePackager/iPhonePackager.csproj");
+			ExtraBuildProducts.Add(CommandUtils.CombinePaths(CommandUtils.CmdEnv.LocalRoot, @"Engine/Binaries/DotNET/iOS/iPhonePackager.exe"));
+
+			Agenda.DotNetProjects.Add(@"Engine/Source/Programs/iOS/DeploymentServer/DeploymentServer.csproj");
+			ExtraBuildProducts.Add(CommandUtils.CombinePaths(CommandUtils.CmdEnv.LocalRoot, @"Engine/Binaries/DotNET/iOS/DeploymentServer.exe"));
+	
+			Agenda.DotNetProjects.Add(@"Engine/Source/Programs/iOS/DeploymentInterface/DeploymentInterface.csproj");
+			ExtraBuildProducts.Add(CommandUtils.CombinePaths(CommandUtils.CmdEnv.LocalRoot, @"Engine/Binaries/DotNET/iOS/DeploymentInterface.dll"));
+			
+			Agenda.DotNetProjects.Add(@"Engine/Source/Programs/iOS/MobileDeviceInterface/MobileDeviceInterface.csproj");
+			ExtraBuildProducts.Add(CommandUtils.CombinePaths(CommandUtils.CmdEnv.LocalRoot, @"Engine/Binaries/DotNET/iOS/MobileDeviceInterface.dll"));
+		}
+
 		// PS4 binaries
 		if(Platforms.Contains(UnrealBuildTool.UnrealTargetPlatform.PS4))
 		{
