Fix variable being initialised to the wrong value. Fixes UE-6893

diff --git a/Engine/Source/Programs/AutomationTool/HTML5/HTML5Platform.Automation.cs b/Engine/Source/Programs/AutomationTool/HTML5/HTML5Platform.Automation.cs
--- a/Engine/Source/Programs/AutomationTool/HTML5/HTML5Platform.Automation.cs
+++ b/Engine/Source/Programs/AutomationTool/HTML5/HTML5Platform.Automation.cs
@@ -285,10 +285,9 @@ public class HTML5Platform : Platform
 			ClientProcess.ProcessObject.Exited += delegate(System.Object o, System.EventArgs e)
 			{
 				System.Console.WriteLine("Browser Process Ended (PID={0})", ClientProcess.ProcessObject.Id);
-				var bFoundChildProcess = true;
+				var bFoundChildProcess = false;
 				if (bBrowserWillSpawnProcess)
 				{
-					bFoundChildProcess = false;
 					// Chrome spawns a process from the tab it opens and then lets the process we spawned die, so
 					// catch that process and attach to that instead.
 					var CurrentProcesses = Process.GetProcesses();
@@ -313,7 +312,7 @@ public class HTML5Platform : Platform
 
 				if (!bFoundChildProcess)
 				{
-					System.Console.WriteLine("- Killing Webserver", ClientProcess.ProcessObject.Id);
+					System.Console.WriteLine("- Killing Webserver PID({0})", Result.ProcessObject.Id);
 					Result.ProcessObject.StandardInput.Close();
 					Result.ProcessObject.Kill();
 				}
