build script, add cross compile node on windows for linux, remove linux compiles from win64 tools node

diff --git a/Engine/Source/Programs/AutomationTool/ProjectUtils.cs b/Engine/Source/Programs/AutomationTool/ProjectUtils.cs
--- a/Engine/Source/Programs/AutomationTool/ProjectUtils.cs
+++ b/Engine/Source/Programs/AutomationTool/ProjectUtils.cs
@@ -552,13 +552,15 @@ namespace AutomationTool
                     {
                         bool bInternalToolOnly;
                         bool SeparateNode;
-                        bool Tool = ThisTarget.Rules.GUBP_AlwaysBuildWithTools(HostPlatform, false, out bInternalToolOnly, out SeparateNode);
+						bool CrossCompile;
+                        bool Tool = ThisTarget.Rules.GUBP_AlwaysBuildWithTools(HostPlatform, false, out bInternalToolOnly, out SeparateNode, out CrossCompile);
 
                         CommandUtils.Log("            TargetName                    : " + ThisTarget.TargetName);
                         CommandUtils.Log("              Build With Editor           : " + (ThisTarget.Rules.GUBP_AlwaysBuildWithBaseEditor() ? "YES" : "NO"));
                         CommandUtils.Log("              Build With Tools            : " + (Tool && !bInternalToolOnly ? "YES" : "NO"));
                         CommandUtils.Log("              Build With Internal Tools   : " + (Tool && bInternalToolOnly ? "YES" : "NO"));
                         CommandUtils.Log("              Separate Node               : " + (Tool && SeparateNode ? "YES" : "NO"));
+						CommandUtils.Log("              Cross Compile               : " + (Tool && CrossCompile ? "YES" : "NO"));
                     }
                 }
             }
diff --git a/Engine/Source/Programs/AutomationTool/Scripts/GUBP.Automation.cs b/Engine/Source/Programs/AutomationTool/Scripts/GUBP.Automation.cs
--- a/Engine/Source/Programs/AutomationTool/Scripts/GUBP.Automation.cs
+++ b/Engine/Source/Programs/AutomationTool/Scripts/GUBP.Automation.cs
@@ -1017,8 +1017,8 @@ public class GUBP : BuildCommand
             {
                 bool bInternalOnly;
                 bool SeparateNode;
-
-                if (ProgramTarget.Rules.GUBP_AlwaysBuildWithTools(HostPlatform, GUBP.bBuildRocket, out bInternalOnly, out SeparateNode) && ProgramTarget.Rules.SupportsPlatform(HostPlatform) && !bInternalOnly && !SeparateNode)
+				bool CrossCompile;
+                if (ProgramTarget.Rules.GUBP_AlwaysBuildWithTools(HostPlatform, GUBP.bBuildRocket, out bInternalOnly, out SeparateNode, out CrossCompile) && ProgramTarget.Rules.SupportsPlatform(HostPlatform) && !bInternalOnly && !SeparateNode)
                 {
                     foreach (var Plat in ProgramTarget.Rules.GUBP_ToolPlatforms(HostPlatform))
                     {
@@ -1033,6 +1033,56 @@ public class GUBP : BuildCommand
             return Agenda;
         }
     }
+	public class ToolsCrossCompileNode : CompileNode
+	{
+		public ToolsCrossCompileNode(UnrealTargetPlatform InHostPlatform)
+			: base(InHostPlatform)
+		{
+			if (!GUBP.bBuildRocket) // more errors and more performance by just starting before the root editor is done
+			{
+				AddPseudodependency(RootEditorCrossCompileLinuxNode.StaticGetFullName(HostPlatform));
+			}
+			AgentSharingGroup = "ToolsCrossCompileGroup" + StaticGetHostPlatformSuffix(HostPlatform);
+		}
+		public static string StaticGetFullName(UnrealTargetPlatform InHostPlatform)
+		{
+			return "LinuxTools" + StaticGetHostPlatformSuffix(InHostPlatform);
+		}
+		public override string GetFullName()
+		{
+			return StaticGetFullName(HostPlatform);
+		}
+		public override float Priority()
+		{
+			return base.Priority() - 1;
+		}
+		public override bool DeleteBuildProducts()
+		{
+			return true;
+		}
+		public override UE4Build.BuildAgenda GetAgenda(GUBP bp)
+		{
+			var Agenda = new UE4Build.BuildAgenda();
+
+			string AddArgs = "-nobuilduht -skipactionhistory -CopyAppBundleBackToDevice" + bp.RocketUBTArgs(); ;
+
+			foreach (var ProgramTarget in bp.Branch.BaseEngineProject.Properties.Programs)
+			{
+				bool bInternalOnly;
+				bool SeparateNode;
+				bool CrossCompile;
+				if (ProgramTarget.Rules.GUBP_AlwaysBuildWithTools(HostPlatform, GUBP.bBuildRocket, out bInternalOnly, out SeparateNode, out CrossCompile) && ProgramTarget.Rules.SupportsPlatform(HostPlatform) && !bInternalOnly && !SeparateNode && CrossCompile)
+				{
+					foreach (var Config in ProgramTarget.Rules.GUBP_ToolConfigs(HostPlatform))
+					{
+						Agenda.AddTargets(new string[] { ProgramTarget.TargetName }, UnrealTargetPlatform.Linux, Config, InAddArgs: AddArgs);
+					}					
+				}
+			}
+
+			return Agenda;
+		}
+	}
     public class SingleToolsNode : CompileNode
     {
         SingleTargetProperties ProgramTarget;
@@ -1149,8 +1199,8 @@ public class GUBP : BuildCommand
             {
                 bool bInternalOnly;
                 bool SeparateNode;
-
-                if (ProgramTarget.Rules.GUBP_AlwaysBuildWithTools(HostPlatform, GUBP.bBuildRocket, out bInternalOnly, out SeparateNode) && ProgramTarget.Rules.SupportsPlatform(HostPlatform) && bInternalOnly && !SeparateNode)
+				bool CrossCompile;
+                if (ProgramTarget.Rules.GUBP_AlwaysBuildWithTools(HostPlatform, GUBP.bBuildRocket, out bInternalOnly, out SeparateNode, out CrossCompile) && ProgramTarget.Rules.SupportsPlatform(HostPlatform) && bInternalOnly && !SeparateNode)
                 {
                     foreach (var Plat in ProgramTarget.Rules.GUBP_ToolPlatforms(HostPlatform))
                     {
@@ -5111,21 +5161,26 @@ public class GUBP : BuildCommand
             }
             AddNode(new ToolsNode(HostPlatform));            
 			AddNode(new InternalToolsNode(HostPlatform));
+			if (HostPlatform == UnrealTargetPlatform.Win64)
+			{
+				AddNode(new ToolsCrossCompileNode(HostPlatform));
+			}
             foreach (var ProgramTarget in Branch.BaseEngineProject.Properties.Programs)
             {
                 bool bInternalOnly;
                 bool SeparateNode;
+				bool CrossCompile;
 
-                if (ProgramTarget.Rules.GUBP_AlwaysBuildWithTools(HostPlatform, GUBP.bBuildRocket, out bInternalOnly, out SeparateNode) && ProgramTarget.Rules.SupportsPlatform(HostPlatform) && SeparateNode)
+                if (ProgramTarget.Rules.GUBP_AlwaysBuildWithTools(HostPlatform, GUBP.bBuildRocket, out bInternalOnly, out SeparateNode, out CrossCompile) && ProgramTarget.Rules.SupportsPlatform(HostPlatform) && SeparateNode)
                 {
                     if (bInternalOnly)
                     {
-                        AddNode(new SingleInternalToolsNode(HostPlatform, ProgramTarget));
+                        AddNode(new SingleInternalToolsNode(HostPlatform, ProgramTarget));						
                     }
                     else
                     {
                         AddNode(new SingleToolsNode(HostPlatform, ProgramTarget));
-                    }
+                    }				
                 }
 				if (ProgramTarget.Rules.GUBP_IncludeNonUnityToolTest())
 				{
@@ -5138,8 +5193,9 @@ public class GUBP : BuildCommand
 				{
 					bool bInternalNodeOnly;
 					bool SeparateNode;
+					bool CrossCompile;
 
-					if(ProgramTarget.Rules.GUBP_AlwaysBuildWithTools(HostPlatform, GUBP.bBuildRocket, out bInternalNodeOnly, out SeparateNode) && ProgramTarget.Rules.SupportsPlatform(HostPlatform) && SeparateNode)
+					if(ProgramTarget.Rules.GUBP_AlwaysBuildWithTools(HostPlatform, GUBP.bBuildRocket, out bInternalNodeOnly, out SeparateNode, out CrossCompile) && ProgramTarget.Rules.SupportsPlatform(HostPlatform) && SeparateNode)
 					{
 						if(bInternalNodeOnly)
 						{
diff --git a/Engine/Source/Programs/BuildPatchTool/BuildPatchTool.Target.cs b/Engine/Source/Programs/BuildPatchTool/BuildPatchTool.Target.cs
--- a/Engine/Source/Programs/BuildPatchTool/BuildPatchTool.Target.cs
+++ b/Engine/Source/Programs/BuildPatchTool/BuildPatchTool.Target.cs
@@ -50,10 +50,11 @@ public class BuildPatchToolTarget : TargetRules
 		OutLinkEnvironmentConfiguration.bBuildAdditionalConsoleApplication = false;
 		OutLinkEnvironmentConfiguration.bHasExports = false;
 	}
-    public override bool GUBP_AlwaysBuildWithTools(UnrealTargetPlatform InHostPlatform, bool bBuildingRocket, out bool bInternalToolOnly, out bool SeparateNode)
+    public override bool GUBP_AlwaysBuildWithTools(UnrealTargetPlatform InHostPlatform, bool bBuildingRocket, out bool bInternalToolOnly, out bool SeparateNode, out bool CrossCompile)
     {
         bInternalToolOnly = true;
         SeparateNode = true;
+		CrossCompile = false;
         return true;
     }
 }
diff --git a/Engine/Source/Programs/CrashReporter/CrashReportClient/CrashReportClient.Target.cs b/Engine/Source/Programs/CrashReporter/CrashReportClient/CrashReportClient.Target.cs
--- a/Engine/Source/Programs/CrashReporter/CrashReportClient/CrashReportClient.Target.cs
+++ b/Engine/Source/Programs/CrashReporter/CrashReportClient/CrashReportClient.Target.cs
@@ -97,17 +97,18 @@ public class CrashReportClientTarget : TargetRules
 
 		OutCPPEnvironmentConfiguration.Definitions.Add( "USE_CHECKS_IN_SHIPPING=1" );
 	}
-    public override bool GUBP_AlwaysBuildWithTools(UnrealTargetPlatform InHostPlatform, bool bBuildingRocket, out bool bInternalToolOnly, out bool SeparateNode)
+    public override bool GUBP_AlwaysBuildWithTools(UnrealTargetPlatform InHostPlatform, bool bBuildingRocket, out bool bInternalToolOnly, out bool SeparateNode, out bool CrossCompile)
 	{
 		bInternalToolOnly = false;
 		SeparateNode = false;
+		CrossCompile = true;
 		return true;
 	}
 	public override List<UnrealTargetPlatform> GUBP_ToolPlatforms(UnrealTargetPlatform InHostPlatform)
 	{
 		if (InHostPlatform == UnrealTargetPlatform.Win64)
 		{
-			return new List<UnrealTargetPlatform> { UnrealTargetPlatform.Win64, UnrealTargetPlatform.Win32, UnrealTargetPlatform.Linux };
+			return new List<UnrealTargetPlatform> { UnrealTargetPlatform.Win64, UnrealTargetPlatform.Win32};
 		}
 		return base.GUBP_ToolPlatforms(InHostPlatform);
 	}
diff --git a/Engine/Source/Programs/CrashReporter/MinidumpDiagnostics/MinidumpDiagnostics.Target.cs b/Engine/Source/Programs/CrashReporter/MinidumpDiagnostics/MinidumpDiagnostics.Target.cs
--- a/Engine/Source/Programs/CrashReporter/MinidumpDiagnostics/MinidumpDiagnostics.Target.cs
+++ b/Engine/Source/Programs/CrashReporter/MinidumpDiagnostics/MinidumpDiagnostics.Target.cs
@@ -60,10 +60,11 @@ public class MinidumpDiagnosticsTarget : TargetRules
 
 		OutCPPEnvironmentConfiguration.Definitions.Add("MINIDUMPDIAGNOSTICS=1");
 	}
-    public override bool GUBP_AlwaysBuildWithTools(UnrealTargetPlatform InHostPlatform, bool bBuildingRocket, out bool bInternalToolOnly, out bool SeparateNode)
+    public override bool GUBP_AlwaysBuildWithTools(UnrealTargetPlatform InHostPlatform, bool bBuildingRocket, out bool bInternalToolOnly, out bool SeparateNode, out bool CrossCompile)
     {
         bInternalToolOnly = true;
         SeparateNode = false;
+		CrossCompile = false;
         return true;
     }
 }
diff --git a/Engine/Source/Programs/Mac/UE4EditorServices/UE4EditorServices.Target.cs b/Engine/Source/Programs/Mac/UE4EditorServices/UE4EditorServices.Target.cs
--- a/Engine/Source/Programs/Mac/UE4EditorServices/UE4EditorServices.Target.cs
+++ b/Engine/Source/Programs/Mac/UE4EditorServices/UE4EditorServices.Target.cs
@@ -49,10 +49,11 @@ public class UE4EditorServicesTarget : TargetRules
 		UEBuildConfiguration.bCompileAgainstEngine = false;
 		UEBuildConfiguration.bCompileAgainstCoreUObject = false;
 	}
-    public override bool GUBP_AlwaysBuildWithTools(UnrealTargetPlatform InHostPlatform, bool bBuildingRocket, out bool bInternalToolOnly, out bool SeparateNode)
+    public override bool GUBP_AlwaysBuildWithTools(UnrealTargetPlatform InHostPlatform, bool bBuildingRocket, out bool bInternalToolOnly, out bool SeparateNode, out bool CrossCompile)
 	{
 		bInternalToolOnly = false;
 		SeparateNode = false;
+		CrossCompile = false;
 		return true;
 	}
 }
diff --git a/Engine/Source/Programs/UnrealBuildTool/System/RulesCompiler.cs b/Engine/Source/Programs/UnrealBuildTool/System/RulesCompiler.cs
--- a/Engine/Source/Programs/UnrealBuildTool/System/RulesCompiler.cs
+++ b/Engine/Source/Programs/UnrealBuildTool/System/RulesCompiler.cs
@@ -617,7 +617,7 @@ namespace UnrealBuildTool
         public virtual bool GUBP_AlwaysBuildWithTools(UnrealTargetPlatform InHostPlatform, out bool bInternalToolOnly, out bool SeparateNode)
         {
             bInternalToolOnly = false;
-            SeparateNode = false;
+            SeparateNode = false;			
             return false;
         }
         /// <summary>
@@ -625,10 +625,11 @@ namespace UnrealBuildTool
         /// <param name="SeparateNode">If this is set to true, the program will get its own node</param>
         /// </summary>
         /// <returns>true if this target should always be built with the base editor.</returns>        
-        public virtual bool GUBP_AlwaysBuildWithTools(UnrealTargetPlatform InHostPlatform, bool bBuildingRocket, out bool bInternalToolOnly, out bool SeparateNode)
+        public virtual bool GUBP_AlwaysBuildWithTools(UnrealTargetPlatform InHostPlatform, bool bBuildingRocket, out bool bInternalToolOnly, out bool SeparateNode, out bool CrossCompile)
         {
             bInternalToolOnly = false;
             SeparateNode = false;
+			CrossCompile = false;
             return false;
         }
         /// <summary>
diff --git a/Engine/Source/Programs/UnrealFileServer/UnrealFileServer.Target.cs b/Engine/Source/Programs/UnrealFileServer/UnrealFileServer.Target.cs
--- a/Engine/Source/Programs/UnrealFileServer/UnrealFileServer.Target.cs
+++ b/Engine/Source/Programs/UnrealFileServer/UnrealFileServer.Target.cs
@@ -57,10 +57,11 @@ public class UnrealFileServerTarget : TargetRules
 		// UnrealHeaderTool is a console application, not a Windows app (sets entry point to main(), instead of WinMain())
 		OutLinkEnvironmentConfiguration.bIsBuildingConsoleApplication = true;
 	}
-    public override bool GUBP_AlwaysBuildWithTools(UnrealTargetPlatform InHostPlatform, bool bBuildingRocket, out bool bInternalToolOnly, out bool SeparateNode)
+    public override bool GUBP_AlwaysBuildWithTools(UnrealTargetPlatform InHostPlatform, bool bBuildingRocket, out bool bInternalToolOnly, out bool SeparateNode, out bool CrossCompile)
     {
         bInternalToolOnly = false;
         SeparateNode = false;
+		CrossCompile = false;
         return true;
     }
 }
diff --git a/Engine/Source/Programs/UnrealLightmass/UnrealLightmass.Target.cs b/Engine/Source/Programs/UnrealLightmass/UnrealLightmass.Target.cs
--- a/Engine/Source/Programs/UnrealLightmass/UnrealLightmass.Target.cs
+++ b/Engine/Source/Programs/UnrealLightmass/UnrealLightmass.Target.cs
@@ -75,20 +75,11 @@ public class UnrealLightmassTarget : TargetRules
 		// Disable logging, lightmass will create its own unique logging file
 		OutCPPEnvironmentConfiguration.Definitions.Add("ALLOW_LOG_FILE=0");
 	}
-    public override bool GUBP_AlwaysBuildWithTools(UnrealTargetPlatform InHostPlatform, bool bBuildingRocket, out bool bInternalToolOnly, out bool SeparateNode)
+    public override bool GUBP_AlwaysBuildWithTools(UnrealTargetPlatform InHostPlatform, bool bBuildingRocket, out bool bInternalToolOnly, out bool SeparateNode, out bool CrossCompile)
     {
         bInternalToolOnly = false;
         SeparateNode = false;
+		CrossCompile = true;
         return true;
     }
-	public override List<UnrealTargetPlatform> GUBP_ToolPlatforms(UnrealTargetPlatform InHostPlatform)
-	{
-		List<UnrealTargetPlatform> PlatformList = new List<UnrealTargetPlatform>();
-		PlatformList.Add(InHostPlatform);
-		if(InHostPlatform == UnrealTargetPlatform.Win64)
-		{
-			PlatformList.Add(UnrealTargetPlatform.Linux);			
-		}
-		return PlatformList;
-	}
 }
diff --git a/Engine/Source/Programs/UnrealPak/UnrealPak.Target.cs b/Engine/Source/Programs/UnrealPak/UnrealPak.Target.cs
--- a/Engine/Source/Programs/UnrealPak/UnrealPak.Target.cs
+++ b/Engine/Source/Programs/UnrealPak/UnrealPak.Target.cs
@@ -51,10 +51,11 @@ public class UnrealPakTarget : TargetRules
 		// UnrealHeaderTool is a console application, not a Windows app (sets entry point to main(), instead of WinMain())
 		OutLinkEnvironmentConfiguration.bIsBuildingConsoleApplication = true;
 	}
-    public override bool GUBP_AlwaysBuildWithTools(UnrealTargetPlatform InHostPlatform, bool bBuildingRocket, out bool bInternalToolOnly, out bool SeparateNode)
+    public override bool GUBP_AlwaysBuildWithTools(UnrealTargetPlatform InHostPlatform, bool bBuildingRocket, out bool bInternalToolOnly, out bool SeparateNode, out bool CrossCompile)
     {
         bInternalToolOnly = false;
         SeparateNode = false;
+		CrossCompile = false;
         return true;
     }
 }
diff --git a/Engine/Source/Programs/UnrealSync/UnrealSync.Target.cs b/Engine/Source/Programs/UnrealSync/UnrealSync.Target.cs
--- a/Engine/Source/Programs/UnrealSync/UnrealSync.Target.cs
+++ b/Engine/Source/Programs/UnrealSync/UnrealSync.Target.cs
@@ -56,8 +56,9 @@ public class UnrealSyncTarget : TargetRules
 		OutLinkEnvironmentConfiguration.bIsBuildingConsoleApplication = false;
 	}
 
-    public override bool GUBP_AlwaysBuildWithTools(UnrealTargetPlatform InHostPlatform, bool bBuildingRocket, out bool bInternalToolOnly, out bool SeparateNode)
+    public override bool GUBP_AlwaysBuildWithTools(UnrealTargetPlatform InHostPlatform, bool bBuildingRocket, out bool bInternalToolOnly, out bool SeparateNode, out bool CrossCompile)
 	{
+		CrossCompile = false;
 		if (InHostPlatform == UnrealTargetPlatform.Win32 || InHostPlatform == UnrealTargetPlatform.Win64)
 		{
 			bInternalToolOnly = true;
diff --git a/Engine/Source/Programs/UnrealVersionSelector/UnrealVersionSelector.Target.cs b/Engine/Source/Programs/UnrealVersionSelector/UnrealVersionSelector.Target.cs
--- a/Engine/Source/Programs/UnrealVersionSelector/UnrealVersionSelector.Target.cs
+++ b/Engine/Source/Programs/UnrealVersionSelector/UnrealVersionSelector.Target.cs
@@ -52,8 +52,9 @@ public class UnrealVersionSelectorTarget : TargetRules
 		// UnrealHeaderTool is a console application, not a Windows app (sets entry point to main(), instead of WinMain())
 		//OutLinkEnvironmentConfiguration.bIsBuildingConsoleApplication = true;
 	}
-    public override bool GUBP_AlwaysBuildWithTools(UnrealTargetPlatform InHostPlatform, bool bBuildingRocket, out bool bInternalToolOnly, out bool SeparateNode)
+    public override bool GUBP_AlwaysBuildWithTools(UnrealTargetPlatform InHostPlatform, bool bBuildingRocket, out bool bInternalToolOnly, out bool SeparateNode, out bool CrossCompile)
 	{
+		CrossCompile = false;
 		if (InHostPlatform == UnrealTargetPlatform.Win32 || InHostPlatform == UnrealTargetPlatform.Win64)
 		{
 			bInternalToolOnly = true;
diff --git a/Engine/Source/Programs/Windows/BootstrapPackagedGame/BootstrapPackagedGame.Target.cs b/Engine/Source/Programs/Windows/BootstrapPackagedGame/BootstrapPackagedGame.Target.cs
--- a/Engine/Source/Programs/Windows/BootstrapPackagedGame/BootstrapPackagedGame.Target.cs
+++ b/Engine/Source/Programs/Windows/BootstrapPackagedGame/BootstrapPackagedGame.Target.cs
@@ -56,10 +56,11 @@ public class BootstrapPackagedGameTarget : TargetRules
 		}
 	}
 	
-    public override bool GUBP_AlwaysBuildWithTools(UnrealTargetPlatform InHostPlatform, bool bBuildingRocket, out bool bInternalToolOnly, out bool SeparateNode)
+    public override bool GUBP_AlwaysBuildWithTools(UnrealTargetPlatform InHostPlatform, bool bBuildingRocket, out bool bInternalToolOnly, out bool SeparateNode, out bool CrossCompile)
 	{
 		bInternalToolOnly = false;
 		SeparateNode = false;
+		CrossCompile = false;
 		return (InHostPlatform == UnrealTargetPlatform.Win64);
 	}
 
