build script, fix for code projects not building - removing hostinstall dependency in cases of only building samples - rely on permanent storage retrieval instead #codereview ben.marsh

diff --git a/Engine/Source/Programs/AutomationTool/Scripts/GUBP.Automation.cs b/Engine/Source/Programs/AutomationTool/Scripts/GUBP.Automation.cs
--- a/Engine/Source/Programs/AutomationTool/Scripts/GUBP.Automation.cs
+++ b/Engine/Source/Programs/AutomationTool/Scripts/GUBP.Automation.cs
@@ -6200,21 +6200,7 @@ public class GUBP : BuildCommand
 				}
 				FullNodeDependentPromotions.Add(NodeToDo, All);				
 			}
-		}
-		if(bBuildRocket)
-		{
-			foreach(var NodeToDo in NodesToDo)
-			{
-				if(GUBPNodes[NodeToDo].GetFullName().Contains("EditorGame"))
-				{
-					GUBPNodes[NodeToDo].AgentSharingGroup = "SampleEditors";
-					if(GUBPNodes[NodeToDo].GetFullName().Contains("OnMac"))
-					{
-						GUBPNodes[NodeToDo].AgentSharingGroup = "SampleEditors_OnMac";
-					}					
-				}
-			}
-		}
+		}		
 		if (TimeIndex != 0)
 		{
 			Log("Culling based on time index");
