Changed SafeCopyFile time checking to account for exFAT timestamp resolution of 2 seconds.

diff --git a/Engine/Source/Programs/AutomationTool/Utils.cs b/Engine/Source/Programs/AutomationTool/Utils.cs
--- a/Engine/Source/Programs/AutomationTool/Utils.cs
+++ b/Engine/Source/Programs/AutomationTool/Utils.cs
@@ -295,51 +295,53 @@ namespace AutomationTool
 			do
 			{
 				Result = true;
-                bool Retry = true;
+				bool Retry = true;
 				try
 				{
 					File.Copy(SourceName, TargetName, overwrite: true);
-                    Retry = !File.Exists(TargetName);
-                    if (!Retry)
-                    {
-                        FileInfo SourceInfo = new FileInfo(SourceName);
-                        FileInfo TargetInfo = new FileInfo(TargetName);
-                        if (SourceInfo.Length != TargetInfo.Length)
-                        {
-                            Log.WriteLine(TraceEventType.Warning, "Size mismatch {0} = {1} to {2} = {3}", SourceName, SourceInfo.Length, TargetName, TargetInfo.Length);
-                            Retry = true;
-                        }
-                        if (!((SourceInfo.LastWriteTimeUtc - TargetInfo.LastWriteTimeUtc).TotalSeconds < 1 && (SourceInfo.LastWriteTimeUtc - TargetInfo.LastWriteTimeUtc).TotalSeconds > -1))
-                        {
-                            Log.WriteLine(TraceEventType.Warning, "Date mismatch {0} = {1} to {2} = {3}", SourceName, SourceInfo.LastWriteTimeUtc, TargetName, TargetInfo.LastWriteTimeUtc);
-                            Retry = true;
-                        }
-                    }
-                }
+					Retry = !File.Exists(TargetName);
+					if (!Retry)
+					{
+						FileInfo SourceInfo = new FileInfo(SourceName);
+						FileInfo TargetInfo = new FileInfo(TargetName);
+						if (SourceInfo.Length != TargetInfo.Length)
+						{
+							Log.WriteLine(TraceEventType.Warning, "Size mismatch {0} = {1} to {2} = {3}", SourceName, SourceInfo.Length, TargetName, TargetInfo.Length);
+							Retry = true;
+						}
+						// Timestamps should be no more than 2 seconds out - assuming this as exFAT filesystems store timestamps at 2 second intervals:
+						// http://ntfs.com/exfat-time-stamp.htm
+						if (!((SourceInfo.LastWriteTimeUtc - TargetInfo.LastWriteTimeUtc).TotalSeconds < 2 && (SourceInfo.LastWriteTimeUtc - TargetInfo.LastWriteTimeUtc).TotalSeconds > -2))
+						{
+							Log.WriteLine(TraceEventType.Warning, "Date mismatch {0} = {1} to {2} = {3}", SourceName, SourceInfo.LastWriteTimeUtc, TargetName, TargetInfo.LastWriteTimeUtc);
+							Retry = true;
+						}
+					}
+				}
 				catch (Exception Ex)
 				{
-                    Log.WriteLine(System.Diagnostics.TraceEventType.Warning, "SafeCopyFile Exception was {0}", LogUtils.FormatException(Ex));
-                    Retry = true;
+					Log.WriteLine(System.Diagnostics.TraceEventType.Warning, "SafeCopyFile Exception was {0}", LogUtils.FormatException(Ex));
+					Retry = true;
 				}
 
-                if (Retry)
-                {
-                    if (Attempts + 1 < MaxAttempts)
-                    {
-                        Log.WriteLine(TraceEventType.Warning, "Failed to copy {0} to {1}, deleting, waiting 10s and retrying.", SourceName, TargetName);
-                        if (File.Exists(TargetName))
-                        {
-                            SafeDeleteFile(TargetName);
-                        }
-                        Thread.Sleep(10000);
-                    }
-                    else
-                    {
-                        Log.WriteLine(TraceEventType.Warning, "Failed to copy {0} to {1}", SourceName, TargetName);
-                    }
-                    Result = false;
-                }
-            }
+				if (Retry)
+				{
+					if (Attempts + 1 < MaxAttempts)
+					{
+						Log.WriteLine(TraceEventType.Warning, "Failed to copy {0} to {1}, deleting, waiting 10s and retrying.", SourceName, TargetName);
+						if (File.Exists(TargetName))
+						{
+							SafeDeleteFile(TargetName);
+						}
+						Thread.Sleep(10000);
+					}
+					else
+					{
+						Log.WriteLine(TraceEventType.Warning, "Failed to copy {0} to {1}", SourceName, TargetName);
+					}
+					Result = false;
+				}
+			}
 			while (Result == false && ++Attempts < MaxAttempts);
 
 			return Result;
