sets a target platform if HostPlatform is null

diff --git a/Engine/Source/Programs/AutomationTool/ProjectParams.cs b/Engine/Source/Programs/AutomationTool/ProjectParams.cs
--- a/Engine/Source/Programs/AutomationTool/ProjectParams.cs
+++ b/Engine/Source/Programs/AutomationTool/ProjectParams.cs
@@ -419,7 +419,26 @@ namespace AutomationTool
             {
                 this.ClientDependentPlatformMap = ClientDependentPlatformMap;
             }
-			this.ClientTargetPlatforms = SetupTargetPlatforms(ref this.ClientDependentPlatformMap, Command, ClientTargetPlatforms, new ParamList<UnrealTargetPlatform>() {HostPlatform.Current.HostEditorPlatform}, true, "TargetPlatform", "Platform");
+
+			UnrealTargetPlatform TargetPlatform;
+			try
+			{
+				TargetPlatform = HostPlatform.Current.HostEditorPlatform;
+			}
+			catch
+			{
+				if (this.ClientDependentPlatformMap != null)
+				{
+					// if a platform map is provided, assume key of first element is the target platform
+					TargetPlatform = this.ClientDependentPlatformMap.Keys.First();
+				}
+				else
+				{
+					// should never get here
+					throw new Exception("A host platform is not defined.");
+				}
+			}
+			this.ClientTargetPlatforms = SetupTargetPlatforms(ref this.ClientDependentPlatformMap, Command, ClientTargetPlatforms, new ParamList<UnrealTargetPlatform>() { TargetPlatform }, true, "TargetPlatform", "Platform");
 
             // Parse command line params for server platforms "-ServerTargetPlatform=Win64+Mac", "-ServerPlatform=Win64+Mac". "-Win64" etc is not allowed here
             if (ServerDependentPlatformMap != null)
