creates ListMobileDevices automation script to assist automated mobile deployments

diff --git a/Engine/Source/Programs/AutomationTool/Scripts/AutomationScripts.Automation.csproj b/Engine/Source/Programs/AutomationTool/Scripts/AutomationScripts.Automation.csproj
--- a/Engine/Source/Programs/AutomationTool/Scripts/AutomationScripts.Automation.csproj
+++ b/Engine/Source/Programs/AutomationTool/Scripts/AutomationScripts.Automation.csproj
@@ -75,6 +75,9 @@
     <Compile Include="GUBP.Automation.cs" />
     <Compile Include="IPhonePackager.Automation.cs" />
     <Compile Include="LauncherLocalization.Automation.cs" />
+    <Compile Include="ListMobileDevices.Automation.cs">
+      <SubType>Code</SubType>
+    </Compile>
     <Compile Include="Localisation.Automation.cs" />
     <Compile Include="MegaXGE.Automation.cs" />
     <Compile Include="PackageCommand.Automation.cs" />
diff --git a/Engine/Source/Programs/AutomationTool/Scripts/ListMobileDevices.Automation.cs b/Engine/Source/Programs/AutomationTool/Scripts/ListMobileDevices.Automation.cs
new file mode 100644
--- /dev/null
+++ b/Engine/Source/Programs/AutomationTool/Scripts/ListMobileDevices.Automation.cs
@@ -0,0 +1,51 @@
+﻿using System.Collections.Generic;
+using AutomationTool;
+using UnrealBuildTool;
+
+namespace AutomationScripts.Automation
+{
+	public class ListMobileDevices : BuildCommand
+	{
+		public override void ExecuteBuild()
+		{
+			Log("======= ListMobileDevices - Start =======");
+
+			var GlobalParams = new ProjectParams(
+				Command: this,
+				RawProjectPath: @"D:\UE-Main\UE4\Samples\Games\TappyChicken\TappyChicken.uproject"
+				);
+
+			if (ParseParam("android"))
+			{
+				GetConnectedDevices(GlobalParams, Platform.Platforms[UnrealTargetPlatform.Android]);
+			}
+
+			if (ParseParam("ios"))
+			{
+				throw new AutomationException("iOS is not yet implemented.");
+			}
+
+			Log("======= ListMobileDevices - Done ========");
+		}
+
+		private static void GetConnectedDevices(ProjectParams Params, Platform TargetPlatform)
+		{
+			var PlatformName = TargetPlatform.PlatformType.ToString();
+			List<string> ConnectedDevices;
+			TargetPlatform.GetConnectedDevices(Params, out ConnectedDevices);
+
+			try
+			{
+				foreach (var DeviceName in ConnectedDevices)
+				{
+					Log("Device:{0}:{1}", PlatformName, DeviceName);
+				}
+			}
+			catch
+			{
+				throw new AutomationException("No {0} devices", PlatformName);
+			}		
+		}
+	}
+}
+
