build script, remove currently failing property as unnecessary

diff --git a/Engine/Source/Programs/AutomationTool/Scripts/GUBP.Automation.cs b/Engine/Source/Programs/AutomationTool/Scripts/GUBP.Automation.cs
--- a/Engine/Source/Programs/AutomationTool/Scripts/GUBP.Automation.cs
+++ b/Engine/Source/Programs/AutomationTool/Scripts/GUBP.Automation.cs
@@ -5876,8 +5876,7 @@ public class GUBP : BuildCommand
         var FullNodeListSortKey = new Dictionary<string, int>();
         var FullNodeDirectDependencies = new Dictionary<string, string>();
 		var FullNodeDependedOnBy = new Dictionary<string, string>();
-		var FullNodeDependentPromotions = new Dictionary<string, string>();
-		var CurrentlyFailing = new Dictionary<string, string>();
+		var FullNodeDependentPromotions = new Dictionary<string, string>();		
 		var SeparatePromotables = new List<string>();
         {
             Log("******* {0} GUBP Nodes", GUBPNodes.Count);
@@ -6337,23 +6336,7 @@ public class GUBP : BuildCommand
                 if (!NodeIsAlreadyComplete(Node, LocalOnly))
                 {
                     UpdateNodeHistory(Node, CLString);
-                }
-				if (GUBPNodes[Node].RunInEC() && !GUBPNodes[Node].TriggerNode() && CLString != "")
-				{
-					try
-					{
-						var History = GUBPNodesHistory[Node];
-						if (History.LastFailed > History.LastSucceeded)
-						{
-							Log("Currently Failing Node {0} at CL {1}", Node, History.LastFailed.ToString());
-							CurrentlyFailing.Add(Node, History.LastFailed.ToString());
-						}
-					}
-					catch
-					{
-
-					}
-				}
+                }				
             }
             var BuildDuration = (DateTime.UtcNow - StartTime).TotalMilliseconds;
             Log("Took {0}s to get history for {1} nodes", BuildDuration / 1000, NodesToDo.Count);
@@ -6444,11 +6427,7 @@ public class GUBP : BuildCommand
 			foreach (var Node in SeparatePromotables)
 			{
 				ECProps.Add(string.Format("PossiblePromotables/{0}={1}", Node, ""));
-			}
-			foreach (var NodePair in CurrentlyFailing)
-			{
-				ECProps.Add(string.Format("CurrentlyFailing/{0}={1}", NodePair.Key, NodePair.Value));
-			}
+			}			
             var ECJobProps = new List<string>();
             if (ExplicitTrigger != "")
             {
