Fix editor invocation in UAT on Linux (UE-6315).

diff --git a/Engine/Source/Programs/AutomationTool/LinuxHostPlatform.cs b/Engine/Source/Programs/AutomationTool/LinuxHostPlatform.cs
--- a/Engine/Source/Programs/AutomationTool/LinuxHostPlatform.cs
+++ b/Engine/Source/Programs/AutomationTool/LinuxHostPlatform.cs
@@ -34,6 +34,19 @@ namespace AutomationTool
 
 		public override string GetUE4ExePath(string UE4Exe)
 		{
+			int CmdExeIndex = UE4Exe.IndexOf("-Cmd.exe");
+			if (CmdExeIndex != -1)
+			{
+				UE4Exe = UE4Exe.Substring (0, CmdExeIndex);
+			}
+			else
+			{
+				CmdExeIndex = UE4Exe.IndexOf (".exe");
+				if (CmdExeIndex != -1)
+				{
+					UE4Exe = UE4Exe.Substring (0, CmdExeIndex);
+				}
+			}
 			return CommandUtils.CombinePaths(CommandUtils.CmdEnv.LocalRoot, RelativeBinariesFolder, UE4Exe);
 		}
 
