fix for failure message when closing iOS game from Launch On
UE-5146
#ios

diff --git a/Engine/Source/Programs/AutomationTool/IOS/IOSPlatform.Automation.cs b/Engine/Source/Programs/AutomationTool/IOS/IOSPlatform.Automation.cs
--- a/Engine/Source/Programs/AutomationTool/IOS/IOSPlatform.Automation.cs
+++ b/Engine/Source/Programs/AutomationTool/IOS/IOSPlatform.Automation.cs
@@ -956,7 +956,7 @@ public class IOSPlatform : Platform
 	{
 		if (UnrealBuildTool.BuildHostPlatform.Current.Platform == UnrealTargetPlatform.Mac)
 		{
-			string AppDirectory = string.Format("{0}/Payload/{1}.app",
+			/*			string AppDirectory = string.Format("{0}/Payload/{1}.app",
 				Path.GetDirectoryName(Params.ProjectGameExeFilename), 
 				Path.GetFileNameWithoutExtension(Params.ProjectGameExeFilename));
 			string GameName = Path.GetFileNameWithoutExtension (ClientApp);
@@ -965,8 +965,8 @@ public class IOSPlatform : Platform
 				GameName = GameName.Substring (0, GameName.IndexOf ("-IOS-"));
 			}
 			string GameApp = AppDirectory + "/" + GameName;
-			bool bWasGenerated = false;
-			string XcodeProj = EnsureXcodeProjectExists (Params.RawProjectPath, CmdEnv.LocalRoot, Params.ShortProjectName, GetDirectoryName(Params.RawProjectPath), Params.IsCodeBasedProject, out bWasGenerated);
+			bWasGenerated = false;
+			XcodeProj = EnsureXcodeProjectExists (Params.RawProjectPath, CmdEnv.LocalRoot, Params.ShortProjectName, GetDirectoryName(Params.RawProjectPath), Params.IsCodeBasedProject, out bWasGenerated);
 			string Arguments = "UBT_NO_POST_DEPLOY=true /usr/bin/xcrun xcodebuild test -project \"" + XcodeProj + "\"";
 			Arguments += " -scheme '";
 			Arguments += GameName;
@@ -976,7 +976,21 @@ public class IOSPlatform : Platform
 			Arguments += " TEST_HOST=\"";
 			Arguments += GameApp;
 			Arguments += "\" BUNDLE_LOADER=\"";
-			Arguments += GameApp + "\"";
+			Arguments += GameApp + "\"";*/
+			string BundleIdentifier = "";
+			if (File.Exists(Params.BaseStageDirectory + "/IOS/Info.plist"))
+			{
+				string Contents = File.ReadAllText(Params.BaseStageDirectory + "/IOS/Info.plist");
+				int Pos = Contents.IndexOf("CFBundleIdentifier");
+				Pos = Contents.IndexOf("<string>", Pos) + 8;
+				int EndPos = Contents.IndexOf("</string>", Pos);
+				BundleIdentifier = Contents.Substring(Pos, EndPos - Pos);
+			}
+			string Arguments = "/usr/bin/instruments";
+			Arguments += " -w '" + Params.Device.Substring (4) + "'";
+			Arguments += " -t 'Activity Monitor'";
+			Arguments += " -D \"" + Params.BaseStageDirectory + "/IOS/launch.trace\"";
+			Arguments += " '" + BundleIdentifier + "'";
 			ProcessResult ClientProcess = Run ("/usr/bin/env", Arguments, null, ClientRunFlags | ERunOptions.NoWaitForExit);
 			return ClientProcess;
 		}
@@ -988,6 +1002,16 @@ public class IOSPlatform : Platform
 		}
 	}
 
+	public override void PostRunClient(ProcessResult Result, ProjectParams Params)
+	{
+		if (UnrealBuildTool.BuildHostPlatform.Current.Platform == UnrealTargetPlatform.Mac)
+		{
+			Console.WriteLine (Directory.GetCurrentDirectory ());
+			Console.WriteLine ("Deleting " + Params.BaseStageDirectory + "/IOS/launch.trace");
+			Directory.Delete (Params.BaseStageDirectory + "/IOS/launch.trace", true);
+		}
+	}
+
 	public override bool StageMovies
 	{
 		get { return false; }
diff --git a/Engine/Source/Programs/AutomationTool/Platform.cs b/Engine/Source/Programs/AutomationTool/Platform.cs
--- a/Engine/Source/Programs/AutomationTool/Platform.cs
+++ b/Engine/Source/Programs/AutomationTool/Platform.cs
@@ -165,6 +165,15 @@ namespace AutomationTool
 			return ClientProcess;
 		}
 
+		/// <summary>
+		/// Allow platform specific clean-up or detection after client has run
+		/// </summary>
+		/// <param name="ClientRunFlags"></param>
+		public virtual void PostRunClient(ProcessResult Result, ProjectParams Params)
+		{
+			// do nothing in the default case
+		}
+
 		/// <summary>
 		/// Get the platform-specific name for the executable (with out the file extension)
 		/// </summary>
diff --git a/Engine/Source/Programs/AutomationTool/Scripts/RunProjectCommand.Automation.cs b/Engine/Source/Programs/AutomationTool/Scripts/RunProjectCommand.Automation.cs
--- a/Engine/Source/Programs/AutomationTool/Scripts/RunProjectCommand.Automation.cs
+++ b/Engine/Source/Programs/AutomationTool/Scripts/RunProjectCommand.Automation.cs
@@ -327,6 +327,8 @@ public partial class Project : CommandUtils
 				}
 				while (ClientProcess.HasExited == false);
 
+				SC.StageTargetPlatform.PostRunClient(ClientProcess, Params);
+
 				if (ClientProcess.ExitCode != 0)
 				{
 					throw new AutomationException("Client exited with error code: " + ClientProcess.ExitCode);
