Made the Mac install script for Android deployment runnable from Finder

diff --git a/Engine/Source/Programs/AutomationTool/Android/AndroidPlatform.Automation.cs b/Engine/Source/Programs/AutomationTool/Android/AndroidPlatform.Automation.cs
--- a/Engine/Source/Programs/AutomationTool/Android/AndroidPlatform.Automation.cs
+++ b/Engine/Source/Programs/AutomationTool/Android/AndroidPlatform.Automation.cs
@@ -125,7 +125,7 @@ public class AndroidPlatform : Platform
 
 	private static string GetFinalBatchName(string ApkName, ProjectParams Params, string Architecture, string GPUArchitecture)
 	{
-		return Path.Combine(Path.GetDirectoryName(ApkName), "Install_" + Params.ShortProjectName + "_" + Params.ClientConfigsToBuild[0].ToString() + Architecture + GPUArchitecture + (Utils.IsRunningOnMono ? ".sh" : ".bat"));
+		return Path.Combine(Path.GetDirectoryName(ApkName), "Install_" + Params.ShortProjectName + "_" + Params.ClientConfigsToBuild[0].ToString() + Architecture + GPUArchitecture + (Utils.IsRunningOnMono ? ".command" : ".bat"));
 	}
 
 	public override void Package(ProjectParams Params, DeploymentContext SC, int WorkingCL)
@@ -209,6 +209,7 @@ public class AndroidPlatform : Platform
 					Log("Writing shell script for install with {0}", Params.OBBinAPK ? "OBB in APK" : "OBB separate");
 					BatchLines = new string[] {
 						"#!/bin/sh",
+						"cd \"`dirname \"$0\"`\"",
 						"ADB=$ANDROID_HOME/platform-tools/adb",
 						"DEVICE=",
 						"if [ \"$1\" != \"\" ]; then DEVICE=\"-s $1\"; fi",
@@ -221,16 +222,17 @@ public class AndroidPlatform : Platform
 						"\tSTORAGE=$(echo \"`$ADB $DEVICE shell 'echo $EXTERNAL_STORAGE'`\" | cat -v | tr -d '^M')",
 						"\t$ADB $DEVICE push " + Path.GetFileName(ObbName) + " $STORAGE/" + DeviceObbName,
 						"\tif [ $? -eq 0 ]; then",
+						"\t\techo Installation successful",
 						"\t\texit 0",
 						"\tfi",
 						"fi",
 						"echo",
-						"echo \"There was an error installing the game or the obb file. Look above for more info.\"",
+						"echo There was an error installing the game or the obb file. Look above for more info.",
 						"echo",
-						"echo \"Things to try:\"",
-						"echo \"Check that the device (and only the device) is listed with \\\"$ADB devices\\\" from a command prompt.\"",
-						"echo \"Make sure all Developer options look normal on the device\"",
-						"echo \"Check that the device has an SD card.\"",
+						"echo Things to try:",
+						"echo Check that the device (and only the device) is listed with \"$ADB devices\" from a command prompt.",
+						"echo Make sure all Developer options look normal on the device",
+						"echo Check that the device has an SD card.",
 						"exit 1"
 					};
 				}
@@ -264,6 +266,11 @@ public class AndroidPlatform : Platform
 					};
 				}
 				File.WriteAllLines(BatchName, BatchLines);
+
+				if (Utils.IsRunningOnMono)
+				{
+					CommandUtils.FixUnixFilePermissions(BatchName);
+				}
 			}
 		}
 
