build script, don't cross compile linux if nolinux is specified

diff --git a/Engine/Source/Programs/AutomationTool/Scripts/GUBP.Automation.cs b/Engine/Source/Programs/AutomationTool/Scripts/GUBP.Automation.cs
--- a/Engine/Source/Programs/AutomationTool/Scripts/GUBP.Automation.cs
+++ b/Engine/Source/Programs/AutomationTool/Scripts/GUBP.Automation.cs
@@ -4798,6 +4798,7 @@ public class GUBP : BuildCommand
         }
 
 		bool WithLinux = !BranchOptions.PlatformsToRemove.Contains(UnrealTargetPlatform.Linux);
+		bool WithoutLinux = ParseParam("NoLinux");
 		// @TODO: exclude temporarily unless running on a Linux machine to prevent spurious GUBP failures
 		if (UnrealBuildTool.BuildHostPlatform.Current.Platform != UnrealTargetPlatform.Linux || ParseParam("NoLinux"))
 		{
@@ -6218,6 +6219,23 @@ public class GUBP : BuildCommand
 			}
 			NodesToDo = NewNodesToDo;
 		}
+		//Remove Plat if specified
+		if(WithoutLinux)
+		{
+			var NewNodesToDo = new HashSet<string>();
+			foreach(var NodeToDo in NodesToDo)
+			{
+				if(!GUBPNodes[NodeToDo].GetFullName().Contains("Linux"))
+				{					
+					NewNodesToDo.Add(NodeToDo);
+				}
+				else
+				{
+					Log(" Rejecting {0} because -NoLinux was requested", NodeToDo);
+				}
+			}
+			NodesToDo = NewNodesToDo;
+		}
 		//find things that depend on our nodes and setup commander dictionary
 		if (!bOnlyNode)
 		{			
