Add PS4DevKitUtil to be built as part of GitHub promotions.

diff --git a/Engine/Source/Programs/AutomationTool/Scripts/BuildCommonTools.Automation.cs b/Engine/Source/Programs/AutomationTool/Scripts/BuildCommonTools.Automation.cs
--- a/Engine/Source/Programs/AutomationTool/Scripts/BuildCommonTools.Automation.cs
+++ b/Engine/Source/Programs/AutomationTool/Scripts/BuildCommonTools.Automation.cs
@@ -133,6 +133,9 @@ public class BuildCommonTools : BuildCommand
 		if(Platforms.Contains(UnrealBuildTool.UnrealTargetPlatform.PS4))
 		{
 			Agenda.AddTarget("PS4MapFileUtil", UnrealBuildTool.UnrealTargetPlatform.Win64, UnrealBuildTool.UnrealTargetConfiguration.Development);
+
+			Agenda.DotNetProjects.Add(CommandUtils.CombinePaths(CommandUtils.CmdEnv.LocalRoot, @"Engine/Source/Programs/PS4/PS4DevKitUtil/PS4DevKitUtil.csproj"));
+			ExtraBuildProducts.Add(CommandUtils.CombinePaths(CommandUtils.CmdEnv.LocalRoot, @"Engine/Binaries/DotNET/PS4/PS4DevKitUtil.exe"));
 		}
 		
 		// Xbox One binaries
