Add UnrealCEFSubProcess to build agenda for common tools.

diff --git a/Engine/Source/Programs/AutomationTool/Scripts/BuildCommonTools.Automation.cs b/Engine/Source/Programs/AutomationTool/Scripts/BuildCommonTools.Automation.cs
--- a/Engine/Source/Programs/AutomationTool/Scripts/BuildCommonTools.Automation.cs
+++ b/Engine/Source/Programs/AutomationTool/Scripts/BuildCommonTools.Automation.cs
@@ -101,6 +101,7 @@ public class BuildCommonTools : BuildCommand
 			Agenda.AddTarget("UnrealVersionSelector", UnrealBuildTool.UnrealTargetPlatform.Win64, UnrealBuildTool.UnrealTargetConfiguration.Shipping);
 			Agenda.AddTarget("BootstrapPackagedGame", UnrealBuildTool.UnrealTargetPlatform.Win64, UnrealBuildTool.UnrealTargetConfiguration.Shipping);
 			Agenda.AddTarget("BootstrapPackagedGame", UnrealBuildTool.UnrealTargetPlatform.Win32, UnrealBuildTool.UnrealTargetConfiguration.Shipping);
+			Agenda.AddTarget("UnrealCEFSubProcess", UnrealBuildTool.UnrealTargetPlatform.Win64, UnrealBuildTool.UnrealTargetConfiguration.Development);
 		}
 
 		// Mac binaries
@@ -111,6 +112,7 @@ public class BuildCommonTools : BuildCommand
 			Agenda.AddTarget("UnrealLightmass", UnrealBuildTool.UnrealTargetPlatform.Mac, UnrealBuildTool.UnrealTargetConfiguration.Development, InAddArgs: "-CopyAppBundleBackToDevice");
 			Agenda.AddTarget("ShaderCompileWorker", UnrealBuildTool.UnrealTargetPlatform.Mac, UnrealBuildTool.UnrealTargetConfiguration.Development, InAddArgs: "-CopyAppBundleBackToDevice");
 			Agenda.AddTarget("UE4EditorServices", UnrealBuildTool.UnrealTargetPlatform.Mac, UnrealBuildTool.UnrealTargetConfiguration.Development, InAddArgs: "-CopyAppBundleBackToDevice");
+			Agenda.AddTarget("UnrealCEFSubProcess", UnrealBuildTool.UnrealTargetPlatform.Mac, UnrealBuildTool.UnrealTargetConfiguration.Development, InAddArgs: "-CopyAppBundleBackToDevice");
 		}
 
 		// iOS binaries
