build script, add vc++ express test node

diff --git a/Engine/Source/Programs/AutomationTool/Scripts/GUBP.Automation.cs b/Engine/Source/Programs/AutomationTool/Scripts/GUBP.Automation.cs
--- a/Engine/Source/Programs/AutomationTool/Scripts/GUBP.Automation.cs
+++ b/Engine/Source/Programs/AutomationTool/Scripts/GUBP.Automation.cs
@@ -3261,7 +3261,58 @@ public class GUBP : BuildCommand
             SaveRecordOfSuccessAndAddToBuildProducts();
         }
     }
+	
+	public class VSExpressTestNode : TestNode
+	{
+		public VSExpressTestNode(GUBP bp)
+			: base(UnrealTargetPlatform.Win64)
+		{
+			AddDependency(ToolsForCompileNode.StaticGetFullName(UnrealTargetPlatform.Win64));
+			AddDependency(RootEditorNode.StaticGetFullName(UnrealTargetPlatform.Win64));
+		}
+		public static string StaticGetFullName()
+		{
+			return "VSExpressTestCompile";
+		}
+		public override string GetFullName()
+		{
+			return StaticGetFullName();
+		}
+		public override int CISFrequencyQuantumShift(GUBP bp)
+		{
+			return base.CISFrequencyQuantumShift(bp) + 3;
+		}
+		public override string ECAgentString()
+		{
+			return "VCTestAgent";
+		}
+		public override void DoTest(GUBP bp)
+		{
+			var Build = new UE4Build(bp);
+			var Agenda = new UE4Build.BuildAgenda();
+
+			string AddArgs = "-nobuilduht" + bp.RocketUBTArgs();
+			if (bp.bOrthogonalizeEditorPlatforms)
+			{
+				AddArgs += " -skipnonhostplatforms";
+			}
+
+			Agenda.AddTargets(
+				new string[] { bp.Branch.BaseEngineProject.Properties.Targets[TargetRules.TargetType.Editor].TargetName },
+				HostPlatform, UnrealTargetConfiguration.Development, InAddArgs: AddArgs);
+			foreach (var ProgramTarget in bp.Branch.BaseEngineProject.Properties.Programs)
+			{
+				if (ProgramTarget.Rules.GUBP_AlwaysBuildWithBaseEditor() && ProgramTarget.Rules.SupportsPlatform(HostPlatform))
+				{
+					Agenda.AddTargets(new string[] { ProgramTarget.TargetName }, HostPlatform, UnrealTargetConfiguration.Development, InAddArgs: AddArgs);
+				}
+			}
+			Build.Build(Agenda, InDeleteBuildProducts: true, InUpdateVersionFiles: false);
 
+			UE4Build.CheckBuildProducts(Build.BuildProductFiles);
+			SaveRecordOfSuccessAndAddToBuildProducts();
+		}
+	}
     public class UATTestNode : TestNode
     {
         string TestName;
@@ -5814,6 +5865,7 @@ public class GUBP : BuildCommand
 			{
 				//AddNode(new IOSOnPCTestNode(this)); - Disable IOSOnPCTest until a1011 crash is fixed
 			}
+			AddNode(new VSExpressTestNode(this));
 			AddNode(new RootEditorCrossCompileLinuxNode(UnrealTargetPlatform.Win64));
             if (!bPreflightBuild)
             {
