Call adb in AndroidPlatform automation script directly instead of using cmd.exe or sh to fix issues with Run returning empty output from these commands

diff --git a/Engine/Source/Programs/AutomationTool/Android/AndroidPlatform.Automation.cs b/Engine/Source/Programs/AutomationTool/Android/AndroidPlatform.Automation.cs
--- a/Engine/Source/Programs/AutomationTool/Android/AndroidPlatform.Automation.cs
+++ b/Engine/Source/Programs/AutomationTool/Android/AndroidPlatform.Automation.cs
@@ -336,21 +336,22 @@ public class AndroidPlatform : Platform
 
 		if (SerialNumber != "")
 		{
-			SerialNumber = " -s " + SerialNumber;
+			SerialNumber = "-s " + SerialNumber;
 		}
 
-		string AdbCommand = Environment.ExpandEnvironmentVariables("%ANDROID_HOME%/platform-tools/adb" + (Utils.IsRunningOnMono ? "" : ".exe") + SerialNumber + " ");
-		return string.Format("{0} {1}{2}{3}{1}", Utils.IsRunningOnMono ? "-c" : "/c", Utils.IsRunningOnMono ? "'" : "", AdbCommand, Args);
+		return string.Format("{0} {1}", SerialNumber, Args);
 	}
 
 	private ProcessResult RunAdbCommand(ProjectParams Params, string Args, string Input = null, ERunOptions Options = ERunOptions.Default)
 	{
-		return Run(CmdEnv.CmdExe, GetAdbCommandLine(Params, Args), Input, Options);
+		string AdbCommand = Environment.ExpandEnvironmentVariables("%ANDROID_HOME%/platform-tools/adb" + (Utils.IsRunningOnMono ? "" : ".exe"));
+		return Run(AdbCommand, GetAdbCommandLine(Params, Args), Input, Options);
 	}
 
 	private string RunAndLogAdbCommand(ProjectParams Params, string Args, out int SuccessCode)
 	{
-		return RunAndLog(CmdEnv, CmdEnv.CmdExe, GetAdbCommandLine(Params, Args), out SuccessCode);
+		string AdbCommand = Environment.ExpandEnvironmentVariables("%ANDROID_HOME%/platform-tools/adb" + (Utils.IsRunningOnMono ? "" : ".exe"));
+		return RunAndLog(CmdEnv, AdbCommand, GetAdbCommandLine(Params, Args), out SuccessCode);
 	}
 
 	public override void GetConnectedDevices(ProjectParams Params, out List<string> Devices)
