build script, check main storage if game node doesn't exist in game storage

diff --git a/Engine/Source/Programs/AutomationTool/Scripts/GUBP.Automation.cs b/Engine/Source/Programs/AutomationTool/Scripts/GUBP.Automation.cs
--- a/Engine/Source/Programs/AutomationTool/Scripts/GUBP.Automation.cs
+++ b/Engine/Source/Programs/AutomationTool/Scripts/GUBP.Automation.cs
@@ -3778,6 +3778,10 @@ public class GUBP : BuildCommand
         else
         {
             Result = TempStorageExists(CmdEnv, NodeStoreName, GameNameIfAny, bQuiet: true);
+			if(GameNameIfAny != "" && Result == false)
+			{
+				Result = TempStorageExists(CmdEnv, NodeStoreName, "", bQuiet: true);
+			}
         }
         if (Result)
         {
@@ -6981,7 +6985,21 @@ public class GUBP : BuildCommand
                 {
                     Log("***** Retrieving GUBP Node {0} from {1}", GUBPNodes[NodeToDo].GetFullName(), NodeStoreName);
                     bool WasLocal;
-                    GUBPNodes[NodeToDo].BuildProducts = RetrieveFromTempStorage(CmdEnv, NodeStoreName, out WasLocal, GameNameIfAny, StorageRootIfAny);
+					try
+					{
+						GUBPNodes[NodeToDo].BuildProducts = RetrieveFromTempStorage(CmdEnv, NodeStoreName, out WasLocal, GameNameIfAny, StorageRootIfAny);
+					}
+					catch
+					{
+						if(GameNameIfAny != "")
+						{
+							GUBPNodes[NodeToDo].BuildProducts = RetrieveFromTempStorage(CmdEnv, NodeStoreName, out WasLocal, "", StorageRootIfAny);
+						}
+						else
+						{
+							throw new AutomationException("Build Products cannot be found for node {0}", NodeToDo);
+						}
+					}
                     if (!WasLocal)
                     {
                         GUBPNodes[NodeToDo].PostLoadFromSharedTempStorage(this);
