Missed file in previous checkin (cl 2407688).

diff --git a/Engine/Source/Programs/AutomationTool/ProjectParams.cs b/Engine/Source/Programs/AutomationTool/ProjectParams.cs
--- a/Engine/Source/Programs/AutomationTool/ProjectParams.cs
+++ b/Engine/Source/Programs/AutomationTool/ProjectParams.cs
@@ -225,6 +225,7 @@ namespace AutomationTool
 			this.SkipPak = InParams.SkipPak;
 			this.NoXGE = InParams.NoXGE;
 			this.CookOnTheFly = InParams.CookOnTheFly;
+            this.CookOnTheFlyStreaming = InParams.CookOnTheFlyStreaming;
 			this.FileServer = InParams.FileServer;
 			this.DedicatedServer = InParams.DedicatedServer;
 			this.Client = InParams.Client;
@@ -328,7 +329,8 @@ namespace AutomationTool
             bool? Compressed = null,
             bool? UseDebugParamForEditorExe = null,
             bool? IterativeCooking = null,
-			bool? CookOnTheFly = null,
+            bool? CookOnTheFly = null,
+            bool? CookOnTheFlyStreaming = null,
             string AdditionalCookerOptions = null,
             string BasedOnReleaseVersion = null,
             string CreateReleaseVersion = null,
@@ -472,6 +474,11 @@ namespace AutomationTool
 			}
 			this.NoXGE = GetParamValueIfNotSpecified(Command, NoXGE, this.NoXGE, "noxge");
 			this.CookOnTheFly = GetParamValueIfNotSpecified(Command, CookOnTheFly, this.CookOnTheFly, "cookonthefly");
+            if (this.CookOnTheFly && this.SkipCook)
+            {
+                this.Cook = false;
+            }
+            this.CookOnTheFlyStreaming = GetParamValueIfNotSpecified(Command, CookOnTheFlyStreaming, this.CookOnTheFlyStreaming, "cookontheflystreaming");
             this.Compressed = GetParamValueIfNotSpecified(Command, Compressed, this.Compressed, "compressed");
             this.UseDebugParamForEditorExe = GetParamValueIfNotSpecified(Command, UseDebugParamForEditorExe, this.UseDebugParamForEditorExe, "UseDebugParamForEditorExe");
             this.IterativeCooking = GetParamValueIfNotSpecified(Command, IterativeCooking, this.IterativeCooking, "iterativecooking", "iterate");
@@ -1095,6 +1102,12 @@ namespace AutomationTool
 		[Help("cookonthefly", "run the client with cooked data provided by cook on the fly server")]
 		public bool CookOnTheFly { private set; get; }
 
+        /// <summary>
+        /// Run: The client should run in streaming mode when connecting to cook on the fly server
+        /// </summary>
+        [Help("Cookontheflystreaming", "run the client in streaming cook on the fly mode (don't cache files locally instead force reget from server each file load)")]
+        public bool CookOnTheFlyStreaming { private set; get; }
+
 		/// <summary>
 		/// Run: The client runs with cooked data provided by UnrealFileServer, command line: -fileserver
 		/// </summary>
@@ -1867,6 +1880,7 @@ namespace AutomationTool
                 CommandUtils.Log("UseDebugParamForEditorExe={0}", UseDebugParamForEditorExe);
 				CommandUtils.Log("CookFlavor={0}", CookFlavor);
 				CommandUtils.Log("CookOnTheFly={0}", CookOnTheFly);
+                CommandUtils.Log("CookOnTheFlyStreaming={0}", CookOnTheFlyStreaming);
                 CommandUtils.Log("CreateReleaseVersion={0}", CreateReleaseVersion);
                 CommandUtils.Log("BasedOnReleaseVersion={0}", BasedOnReleaseVersion);
                 CommandUtils.Log("DLCName={0}", DLCName);
