Just use a simple text file for the manifest of files created by BuildCommonTools; it's less sensitive to versioning.

diff --git a/Engine/Source/Programs/AutomationTool/Scripts/BuildCommonTools.Automation.cs b/Engine/Source/Programs/AutomationTool/Scripts/BuildCommonTools.Automation.cs
--- a/Engine/Source/Programs/AutomationTool/Scripts/BuildCommonTools.Automation.cs
+++ b/Engine/Source/Programs/AutomationTool/Scripts/BuildCommonTools.Automation.cs
@@ -1,6 +1,7 @@
 ﻿// Copyright 1998-2015 Epic Games, Inc. All Rights Reserved.
 
 using System;
+using System.IO;
 using System.Collections.Generic;
 using System.Linq;
 using System.Text;
@@ -68,12 +69,12 @@ public class BuildCommonTools : BuildCommand
 		string ManifestPath = ParseParamValue("manifest");
 		if(ManifestPath != null)
 		{
-			BuildManifest Manifest = new BuildManifest();
+			SortedSet<string> Files = new SortedSet<string>();
 			foreach(string BuildProductFile in Builder.BuildProductFiles)
 			{
-				Manifest.AddBuildProduct(BuildProductFile);
+				Files.Add(BuildProductFile);
 			}
-			UnrealBuildTool.Utils.WriteClass(Manifest, ManifestPath, "");
+			File.WriteAllLines(ManifestPath, Files.ToArray());
 		}
 	}
 
