Missed files from CEF3 Win32 checkin

diff --git a/Engine/Source/Programs/AutomationTool/Mac/MacPlatform.Automation.cs b/Engine/Source/Programs/AutomationTool/Mac/MacPlatform.Automation.cs
--- a/Engine/Source/Programs/AutomationTool/Mac/MacPlatform.Automation.cs
+++ b/Engine/Source/Programs/AutomationTool/Mac/MacPlatform.Automation.cs
@@ -69,6 +69,11 @@ public class MacPlatform : Platform
 
 		// Copy the splash screen, Mac specific
 		SC.StageFiles(StagedFileType.NonUFS, CombinePaths(SC.ProjectRoot, "Content/Splash"), "Splash.bmp", false, null, null, true);
+
+		// CEF3 files
+		SC.StageFiles(StagedFileType.NonUFS, CombinePaths(SC.LocalRoot, "Binaries/ThirdParty/CEF3/Mac/"), "*", true, null, null, true);
+		string UnrealCEFSubProcessPath = CombinePaths("Engine/Binaries", SC.PlatformDir, "UnrealCEFSubProcess.app");
+		StageAppBundle(SC, CombinePaths(SC.LocalRoot, "Engine/Binaries", SC.PlatformDir, "UnrealCEFSubProcess.app"), UnrealCEFSubProcessPath);
 	}
 
 	public override void Package(ProjectParams Params, DeploymentContext SC, int WorkingCL)
diff --git a/Engine/Source/Programs/AutomationTool/Win/WinPlatform.Automation.cs b/Engine/Source/Programs/AutomationTool/Win/WinPlatform.Automation.cs
--- a/Engine/Source/Programs/AutomationTool/Win/WinPlatform.Automation.cs
+++ b/Engine/Source/Programs/AutomationTool/Win/WinPlatform.Automation.cs
@@ -88,6 +88,10 @@ public abstract class BaseWinPlatform : Platform
 		// Copy the splash screen, windows specific
 		SC.StageFiles(StagedFileType.NonUFS, CombinePaths(SC.ProjectRoot, "Content/Splash"), "Splash.bmp", false, null, null, true);
 
+		// CEF3 files
+		SC.StageFiles(StagedFileType.NonUFS, CombinePaths(SC.LocalRoot, "Engine/Binaries/ThirdParty/CEF3" SC.PlatformDir), "*", true, null, null, true);
+		StageExecutable("exe", SC, CommandUtils.CombinePaths(SC.LocalRoot, "Engine/Binaries", SC.PlatformDir), "UnrealCEFSubProcess.", false, null, null, true);
+
         if (Params.StageNonMonolithic)
         {
             StageExecutable("dll", SC, CommandUtils.CombinePaths(SC.LocalRoot, "Engine/Binaries/ThirdParty/ICU/icu4c-53_1", SC.PlatformDir, "VS2013"), "*.");
diff --git a/Engine/Source/Runtime/WebBrowser/WebBrowser.Build.cs b/Engine/Source/Runtime/WebBrowser/WebBrowser.Build.cs
--- a/Engine/Source/Runtime/WebBrowser/WebBrowser.Build.cs
+++ b/Engine/Source/Runtime/WebBrowser/WebBrowser.Build.cs
@@ -22,6 +22,7 @@ public class WebBrowser : ModuleRules
 		);
 
 		if (Target.Platform == UnrealTargetPlatform.Win64
+		|| Target.Platform == UnrealTargetPlatform.Win32
         ||  Target.Platform == UnrealTargetPlatform.Mac)
 		{
 			AddThirdPartyPrivateStaticDependencies(Target,
