Missed comma

diff --git a/Engine/Source/Programs/AutomationTool/Win/WinPlatform.Automation.cs b/Engine/Source/Programs/AutomationTool/Win/WinPlatform.Automation.cs
--- a/Engine/Source/Programs/AutomationTool/Win/WinPlatform.Automation.cs
+++ b/Engine/Source/Programs/AutomationTool/Win/WinPlatform.Automation.cs
@@ -89,7 +89,7 @@ public abstract class BaseWinPlatform : Platform
 		SC.StageFiles(StagedFileType.NonUFS, CombinePaths(SC.ProjectRoot, "Content/Splash"), "Splash.bmp", false, null, null, true);
 
 		// CEF3 files
-		SC.StageFiles(StagedFileType.NonUFS, CombinePaths(SC.LocalRoot, "Engine/Binaries/ThirdParty/CEF3" SC.PlatformDir), "*", true, null, null, true);
+		SC.StageFiles(StagedFileType.NonUFS, CombinePaths(SC.LocalRoot, "Engine/Binaries/ThirdParty/CEF3", SC.PlatformDir), "*", true, null, null, true);
 		StageExecutable("exe", SC, CommandUtils.CombinePaths(SC.LocalRoot, "Engine/Binaries", SC.PlatformDir), "UnrealCEFSubProcess.", false, null, null, true);
 
         if (Params.StageNonMonolithic)
