fix for trying to delete a directory that doesn't exist
UE-7640
#ios

diff --git a/Engine/Source/Programs/AutomationTool/IOS/IOSPlatform.Automation.cs b/Engine/Source/Programs/AutomationTool/IOS/IOSPlatform.Automation.cs
--- a/Engine/Source/Programs/AutomationTool/IOS/IOSPlatform.Automation.cs
+++ b/Engine/Source/Programs/AutomationTool/IOS/IOSPlatform.Automation.cs
@@ -1006,9 +1006,11 @@ public class IOSPlatform : Platform
 	{
 		if (UnrealBuildTool.BuildHostPlatform.Current.Platform == UnrealTargetPlatform.Mac)
 		{
-			Console.WriteLine (Directory.GetCurrentDirectory ());
 			Console.WriteLine ("Deleting " + Params.BaseStageDirectory + "/IOS/launch.trace");
-			Directory.Delete (Params.BaseStageDirectory + "/IOS/launch.trace", true);
+			if (Directory.Exists(Params.BaseStageDirectory + "/IOS/launch.trace"))
+			{
+				Directory.Delete (Params.BaseStageDirectory + "/IOS/launch.trace", true);
+			}
 		}
 	}
 
