Make Rocket builds restartable by archiving generated headers from root editor and monolithic nodes rather than relying on them not being touched.

diff --git a/Engine/Source/Programs/AutomationTool/Scripts/GUBP.Automation.cs b/Engine/Source/Programs/AutomationTool/Scripts/GUBP.Automation.cs
--- a/Engine/Source/Programs/AutomationTool/Scripts/GUBP.Automation.cs
+++ b/Engine/Source/Programs/AutomationTool/Scripts/GUBP.Automation.cs
@@ -1573,6 +1573,27 @@ public class GUBP : BuildCommand
         }
     }
 
+	public static class HeadersNode
+	{
+		public static void ZipHeaders(IEnumerable<string> HeaderFiles, string ZipFileName)
+		{
+			Ionic.Zip.ZipFile Zip = new Ionic.Zip.ZipFile();
+			foreach(string HeaderFile in HeaderFiles)
+			{
+				string NormalizedDirectoryName = CommandUtils.ConvertSeparators(PathSeparator.Slash, Path.GetDirectoryName(HeaderFile));
+				if(NormalizedDirectoryName.StartsWith(CmdEnv.LocalRoot))
+				{
+					Zip.AddFile(HeaderFile, NormalizedDirectoryName.Substring(CmdEnv.LocalRoot.Length).TrimStart('/'));
+				}
+				else
+				{
+					throw new AutomationException("Header file was not under root directory");
+				}
+			}
+			Zip.Save(ZipFileName);
+		}
+	}
+
     public class RootEditorHeadersNode : HostPlatformNode
     {
         public RootEditorHeadersNode(UnrealTargetPlatform InHostPlatform)
@@ -1595,11 +1616,10 @@ public class GUBP : BuildCommand
         }
         public override void DoBuild(GUBP bp)
         {
-            BuildProducts = new List<string>();
-
+			HashSet<string> HeaderFiles = new HashSet<string>(StringComparer.InvariantCultureIgnoreCase);
             foreach (var FileToCopy in CommandUtils.FindFiles("*.h", true, CommandUtils.CombinePaths(CmdEnv.LocalRoot, @"Engine\Intermediate\Build\", HostPlatform.ToString(), "Inc")))
             {
-                AddBuildProduct(FileToCopy);
+				HeaderFiles.Add(FileToCopy);
             }
             var Targets = new List<string> { bp.Branch.BaseEngineProject.Properties.Targets[TargetRules.TargetType.Editor].TargetName };
             foreach (var ProgramTarget in bp.Branch.BaseEngineProject.Properties.Programs)
@@ -1613,7 +1633,7 @@ public class GUBP : BuildCommand
             {
                 foreach (var FileToCopy in CommandUtils.FindFiles("*.h", true, CommandUtils.CombinePaths(CmdEnv.LocalRoot, @"Engine\Intermediate\Build\", HostPlatform.ToString(), Target, "Inc")))
                 {
-                    AddBuildProduct(FileToCopy);
+					HeaderFiles.Add(FileToCopy);
                 }
             }
 
@@ -1625,10 +1645,14 @@ public class GUBP : BuildCommand
             {
                 foreach (var FileToCopy in CommandUtils.FindFiles("*", true, CommandUtils.CombinePaths(EnginePlugin.Directory, @"Intermediate\Build", HostPlatform.ToString(), "Inc")))
                 {
-                    AddBuildProduct(FileToCopy);
+					HeaderFiles.Add(FileToCopy);
                 }
             }
-            RemoveOveralppingBuildProducts();
+
+			// Create a zip file containing the headers. 
+			string ZipFileName = CommandUtils.CombinePaths(CmdEnv.LocalRoot, "Engine\\Intermediate\\Build", HostPlatform.ToString(), "RootEditorHeaders.zip");
+			HeadersNode.ZipHeaders(HeaderFiles, ZipFileName);
+            BuildProducts = new List<string>{ ZipFileName };
         }
     }
 
@@ -1661,14 +1685,15 @@ public class GUBP : BuildCommand
             {
                 throw new AutomationException("these rocket header node require real mac path.");
             }
-            BuildProducts = new List<string>();
+
+			HashSet<string> HeaderFiles = new HashSet<string>(StringComparer.InvariantCultureIgnoreCase);
 			if (TargetPlatform != HostPlatform)
 			{
 				// host platform overlaps with the editor, so we don't want them here
 				var PlatformDir = TargetPlatform.ToString();
 				foreach (var FileToCopy in CommandUtils.FindFiles("*.h", true, CommandUtils.CombinePaths(CmdEnv.LocalRoot, @"Engine\Intermediate\Build\", PlatformDir, "Inc")))
 				{
-					AddBuildProduct(FileToCopy);
+					HeaderFiles.Add(FileToCopy);
 				}
 
 				// these may not be any new build products here, but we will check anyway
@@ -1680,15 +1705,15 @@ public class GUBP : BuildCommand
 				{
 					foreach (var FileToCopy in CommandUtils.FindFiles("*", true, CommandUtils.CombinePaths(EnginePlugin.Directory, @"Intermediate\Build", PlatformDir, "Inc")))
 					{
-						AddBuildProduct(FileToCopy);
+						HeaderFiles.Add(FileToCopy);
 					}
 				}
 			}
-            RemoveOveralppingBuildProducts();
-            if (BuildProducts.Count == 0)
-            {
-                SaveRecordOfSuccessAndAddToBuildProducts(); // could be empty
-            }
+
+			// Create a zip containing all the headers
+			string ZipFileName = CommandUtils.CombinePaths(CmdEnv.LocalRoot, "Engine\\Intermediate\\Build", TargetPlatform.ToString(), "MonolithicHeaders.zip");
+			HeadersNode.ZipHeaders(HeaderFiles, ZipFileName);
+            BuildProducts = new List<string>{ ZipFileName };
         }
     }
 
