Add more descriptive error message when zipping headers fails.

diff --git a/Engine/Source/Programs/AutomationTool/Scripts/GUBP.Automation.cs b/Engine/Source/Programs/AutomationTool/Scripts/GUBP.Automation.cs
--- a/Engine/Source/Programs/AutomationTool/Scripts/GUBP.Automation.cs
+++ b/Engine/Source/Programs/AutomationTool/Scripts/GUBP.Automation.cs
@@ -1577,17 +1577,19 @@ public class GUBP : BuildCommand
 	{
 		public static void ZipHeaders(IEnumerable<string> HeaderFiles, string ZipFileName)
 		{
+			string NormalizedPrefix = CommandUtils.ConvertSeparators(PathSeparator.Slash, CmdEnv.LocalRoot).TrimEnd('/') + "/";
+
 			Ionic.Zip.ZipFile Zip = new Ionic.Zip.ZipFile();
 			foreach(string HeaderFile in HeaderFiles)
 			{
 				string NormalizedDirectoryName = CommandUtils.ConvertSeparators(PathSeparator.Slash, Path.GetDirectoryName(HeaderFile));
-				if(NormalizedDirectoryName.StartsWith(CmdEnv.LocalRoot))
+				if(NormalizedDirectoryName.StartsWith(NormalizedPrefix))
 				{
-					Zip.AddFile(HeaderFile, NormalizedDirectoryName.Substring(CmdEnv.LocalRoot.Length).TrimStart('/'));
+					Zip.AddFile(HeaderFile, NormalizedDirectoryName.Substring(NormalizedPrefix.Length));
 				}
 				else
 				{
-					throw new AutomationException("Header file was not under root directory");
+					throw new AutomationException("Header file '{0}' was not under root directory ('{1}')", NormalizedDirectoryName, NormalizedPrefix);
 				}
 			}
 			Zip.Save(ZipFileName);
