SDK path editing in editor.
- Added class to hold values (AndroidSDKSettings)
- Added code to Android editor module to setup the values on load
- Changed device detection so that the thread is always started and the SDK path can be changed on the fly
- Changed Platform Target Management Module so that a single platform can be checked via UBT
- Settings can now be sourced from an ini file via temp direct reading code (auto seralisation is disabled until later changes are made)
-- Settings currently not exposed in editor

- Added Mac Environment var setting support

Unreleated
- Envars can now be passed down to process start points in tools

diff --git a/Engine/Source/Developer/Android/AndroidDeviceDetection/Private/AndroidDeviceDetectionModule.cpp b/Engine/Source/Developer/Android/AndroidDeviceDetection/Private/AndroidDeviceDetectionModule.cpp
--- a/Engine/Source/Developer/Android/AndroidDeviceDetection/Private/AndroidDeviceDetectionModule.cpp
+++ b/Engine/Source/Developer/Android/AndroidDeviceDetection/Private/AndroidDeviceDetectionModule.cpp
@@ -12,11 +12,13 @@
 class FAndroidDeviceDetectionRunnable : public FRunnable
 {
 public:
-	FAndroidDeviceDetectionRunnable(const FString& InADBPath, TMap<FString,FAndroidDeviceInfo>& InDeviceMap, FCriticalSection* InDeviceMapLock) :
-		ADBPath(InADBPath),
+	FAndroidDeviceDetectionRunnable(TMap<FString,FAndroidDeviceInfo>& InDeviceMap, FCriticalSection* InDeviceMapLock, FCriticalSection* InADBPathCheckLock) :
 		StopTaskCounter(0),
 		DeviceMap(InDeviceMap),
-		DeviceMapLock(InDeviceMapLock)
+		DeviceMapLock(InDeviceMapLock),
+		ADBPathCheckLock(InADBPathCheckLock),
+		HasADBPath(false)
+
 	{
 	}
 
@@ -46,7 +48,11 @@ public:
 			// query every 10 seconds
 			if (LoopCount++ >= 10)
 			{
-				QueryConnectedDevices();
+				// Make sure we have an ADB path before checking
+				FScopeLock PathLock(ADBPathCheckLock);
+				if (HasADBPath)
+					QueryConnectedDevices();
+
 				LoopCount = 0;
 			}
 
@@ -56,6 +62,12 @@ public:
 		return 0;
 	}
 
+	void UpdateADBPath(FString &InADBPath)
+	{
+		ADBPath = InADBPath;
+		HasADBPath = !ADBPath.IsEmpty();
+	}
+
 private:
 
 	bool ExecuteAdbCommand( const FString& CommandLine, FString* OutStdOut, FString* OutStdErr ) const
@@ -241,6 +253,9 @@ private:
 
 	TMap<FString,FAndroidDeviceInfo>& DeviceMap;
 	FCriticalSection* DeviceMapLock;
+
+	FCriticalSection* ADBPathCheckLock;
+	bool HasADBPath;
 };
 
 class FAndroidDeviceDetection : public IAndroidDeviceDetection
@@ -251,81 +266,12 @@ public:
 		DetectionThread(nullptr),
 		DetectionThreadRunnable(nullptr)
 	{
-		// get the SDK binaries folder
-		TCHAR AndroidDirectory[32768] = { 0 };
-		FPlatformMisc::GetEnvironmentVariable(TEXT("ANDROID_HOME"), AndroidDirectory, 32768);
-
-#if PLATFORM_MAC
-		TCHAR AntDirectory[32768] = { 0 };
-		TCHAR NDKDirectory[32768] = { 0 };
-		FPlatformMisc::GetEnvironmentVariable(TEXT("ANT_HOME"), AntDirectory, 32768);
-		FPlatformMisc::GetEnvironmentVariable(TEXT("NDKROOT"), NDKDirectory, 32768);
-
-		if (AndroidDirectory[0] == 0 || AntDirectory[0] == 0 || NDKDirectory[0] == 0)
-		{
-			FArchive* FileReader = IFileManager::Get().CreateFileReader(*FString([@"~/.bash_profile" stringByExpandingTildeInPath]));
-			if (FileReader)
-			{
-				const int64 FileSize = FileReader->TotalSize();
-				ANSICHAR* AnsiContents = (ANSICHAR*)FMemory::Malloc(FileSize);
-				FileReader->Serialize(AnsiContents, FileSize);
-				FileReader->Close();
-				delete FileReader;
-
-				TArray<FString> Lines;
-				FString(ANSI_TO_TCHAR(AnsiContents)).ParseIntoArrayLines(&Lines);
-				FMemory::Free(AnsiContents);
-
-				for (int32 Index = 0; Index < Lines.Num(); Index++)
-				{
-					if (AndroidDirectory[0] == 0 && Lines[Index].StartsWith(TEXT("export ANDROID_HOME=")))
-					{
-						FString Directory;
-						Lines[Index].Split(TEXT("="), NULL, &Directory);
-						Directory = Directory.Replace(TEXT("\""), TEXT(""));
-						FCString::Strcpy(AndroidDirectory, *Directory);
-						setenv("ANDROID_HOME", TCHAR_TO_ANSI(AndroidDirectory), 1);
-					}
-					else if (AntDirectory[0] == 0 && Lines[Index].StartsWith(TEXT("export ANT_HOME=")))
-					{
-						FString Directory;
-						Lines[Index].Split(TEXT("="), NULL, &Directory);
-						Directory = Directory.Replace(TEXT("\""), TEXT(""));
-						setenv("ANT_HOME", TCHAR_TO_ANSI(*Directory), 1);
-					}
-					else if (NDKDirectory[0] == 0 && Lines[Index].StartsWith(TEXT("export NDKROOT=")))
-					{
-						FString Directory;
-						Lines[Index].Split(TEXT("="), NULL, &Directory);
-						Directory = Directory.Replace(TEXT("\""), TEXT(""));
-						setenv("NDKROOT", TCHAR_TO_ANSI(*Directory), 1);
-					}
-				}
-			}
-		}
-#endif
-
-		if (AndroidDirectory[0] == 0)
-		{
-			return;
-		}
-
-#if PLATFORM_WINDOWS
-		FString ADBPath = FString::Printf(TEXT("%s\\platform-tools\\adb.exe"), AndroidDirectory);
-#else
-		FString ADBPath = FString::Printf(TEXT("%s/platform-tools/adb"), AndroidDirectory);
-#endif
-
-		// if it doesn't exist, no SDK installed, so don't bother creating a thread to look for devices
-		if (!FPaths::FileExists(*ADBPath))
-		{
-			ADBPath.Empty();
-			return;
-		}
-
 		// create and fire off our device detection thread
-		DetectionThreadRunnable = new FAndroidDeviceDetectionRunnable(ADBPath, DeviceMap, &DeviceMapLock);
+		DetectionThreadRunnable = new FAndroidDeviceDetectionRunnable(DeviceMap, &DeviceMapLock, &ADBPathCheckLock);
 		DetectionThread = FRunnableThread::Create(DetectionThreadRunnable, TEXT("FAndroidDeviceDetectionRunnable"));
+
+		// get the SDK binaries folder and throw it to the runnable
+		UpdateADBPath();
 	}
 
 	virtual ~FAndroidDeviceDetection()
@@ -347,13 +293,40 @@ public:
 		return &DeviceMapLock;
 	}
 
+	virtual void UpdateADBPath()
+	{
+		FScopeLock PathUpdateLock(&ADBPathCheckLock);
+		TCHAR AndroidDirectory[32768] = { 0 };
+		FPlatformMisc::GetEnvironmentVariable(TEXT("ANDROID_HOME"), AndroidDirectory, 32768);
+
+		FString ADBPath;
+
+		if (AndroidDirectory[0] != 0)
+		{
+#if PLATFORM_WINDOWS
+			ADBPath = FString::Printf(TEXT("%s\\platform-tools\\adb.exe"), AndroidDirectory);
+#else
+			ADBPath = FString::Printf(TEXT("%s/platform-tools/adb"), AndroidDirectory);
+#endif
+
+			// if it doesn't exist then just clear the path as we might set it later
+			if (!FPaths::FileExists(*ADBPath))
+			{
+				ADBPath.Empty();
+			}
+		}
+		DetectionThreadRunnable->UpdateADBPath(ADBPath);
+	}
+
 private:
 
+
 	FRunnableThread* DetectionThread;
 	FAndroidDeviceDetectionRunnable* DetectionThreadRunnable;
 
 	TMap<FString,FAndroidDeviceInfo> DeviceMap;
 	FCriticalSection DeviceMapLock;
+	FCriticalSection ADBPathCheckLock;
 };
 
 
diff --git a/Engine/Source/Developer/Android/AndroidDeviceDetection/Public/Interfaces/IAndroidDeviceDetection.h b/Engine/Source/Developer/Android/AndroidDeviceDetection/Public/Interfaces/IAndroidDeviceDetection.h
--- a/Engine/Source/Developer/Android/AndroidDeviceDetection/Public/Interfaces/IAndroidDeviceDetection.h
+++ b/Engine/Source/Developer/Android/AndroidDeviceDetection/Public/Interfaces/IAndroidDeviceDetection.h
@@ -50,7 +50,7 @@ class IAndroidDeviceDetection
 public:
 	virtual const TMap<FString,FAndroidDeviceInfo>& GetDeviceMap() = 0;
 	virtual FCriticalSection* GetDeviceMapLock() = 0;
-
+	virtual void UpdateADBPath() = 0;
 protected:
 
 	/**
diff --git a/Engine/Source/Developer/Android/AndroidPlatformEditor/AndroidPlatformEditor.Build.cs b/Engine/Source/Developer/Android/AndroidPlatformEditor/AndroidPlatformEditor.Build.cs
--- a/Engine/Source/Developer/Android/AndroidPlatformEditor/AndroidPlatformEditor.Build.cs
+++ b/Engine/Source/Developer/Android/AndroidPlatformEditor/AndroidPlatformEditor.Build.cs
@@ -17,10 +17,14 @@ public class AndroidPlatformEditor : ModuleRules
 				"Slate",
 				"SlateCore",
 				"EditorStyle",
+                "EditorWidgets",
+                "DesktopWidgets",
 				"PropertyEditor",
 				"SharedSettingsWidgets",
 				"SourceControl",
 				"AndroidRuntimeSettings",
+                "AndroidDeviceDetection",
+                "TargetPlatform"
 			}
 		);
 
diff --git a/Engine/Source/Developer/Android/AndroidPlatformEditor/Private/AndroidPlatformEditorModule.cpp b/Engine/Source/Developer/Android/AndroidPlatformEditor/Private/AndroidPlatformEditorModule.cpp
--- a/Engine/Source/Developer/Android/AndroidPlatformEditor/Private/AndroidPlatformEditorModule.cpp
+++ b/Engine/Source/Developer/Android/AndroidPlatformEditor/Private/AndroidPlatformEditorModule.cpp
@@ -2,6 +2,7 @@
 
 #include "AndroidPlatformEditorPrivatePCH.h"
 #include "AndroidTargetSettingsCustomization.h"
+#include "AndroidSDKSettingsCustomization.h"
 #include "ModuleInterface.h"
 #include "ISettingsModule.h"
 #include "ModuleManager.h"
@@ -27,6 +28,11 @@ class FAndroidPlatformEditorModule
 			FOnGetDetailCustomizationInstance::CreateStatic(&FAndroidTargetSettingsCustomization::MakeInstance)
 		);
 
+		PropertyModule.RegisterCustomClassLayout(
+			UAndroidSDKSettings::StaticClass()->GetFName(),
+			FOnGetDetailCustomizationInstance::CreateStatic(&FAndroidSDKSettingsCustomization::MakeInstance)
+			);
+
 		PropertyModule.NotifyCustomizationModuleChanged();
 
 		// register settings
@@ -39,7 +45,23 @@ class FAndroidPlatformEditorModule
 				LOCTEXT("RuntimeSettingsDescription", "Settings and resources for Android platforms"),
 				GetMutableDefault<UAndroidRuntimeSettings>()
 			);
+
+			// This will need to be added back in once we know where we are going to be setting the settings from
+// 			SettingsModule->RegisterSettings("Project", "User Settings", "Android",
+// 				LOCTEXT("RuntimeSettingsName", "Android"),
+// 				LOCTEXT("RuntimeSettingsDescription", "Settings for Android SDK"),
+// 				GetMutableDefault<UAndroidSDKSettings>()
+//			);
 		}
+
+		// Force the SDK settings into a sane state initially so we can make use of them
+		auto &TargetPlatformManagerModule = FModuleManager::LoadModuleChecked<ITargetPlatformManagerModule>("TargetPlatform");
+		UAndroidSDKSettings * settings = GetMutableDefault<UAndroidSDKSettings>();
+		settings->SetTargetModule(&TargetPlatformManagerModule);
+		auto &AndroidDeviceDetection = FModuleManager::LoadModuleChecked<IAndroidDeviceDetection>("AndroidDeviceDetection");
+		settings->SetDeviceDetection(&AndroidDeviceDetection);
+		settings->SetupInitialTargetPaths();
+
 	}
 
 	virtual void ShutdownModule() override
@@ -49,6 +71,7 @@ class FAndroidPlatformEditorModule
 		if (SettingsModule != nullptr)
 		{
 			SettingsModule->UnregisterSettings("Project", "Platforms", "Android");
+			// SettingsModule->UnregisterSettings("Project", "User Settings", "Android");
 		}
 	}
 };
diff --git a/Engine/Source/Developer/Android/AndroidPlatformEditor/Private/AndroidPlatformEditorPrivatePCH.h b/Engine/Source/Developer/Android/AndroidPlatformEditor/Private/AndroidPlatformEditorPrivatePCH.h
--- a/Engine/Source/Developer/Android/AndroidPlatformEditor/Private/AndroidPlatformEditorPrivatePCH.h
+++ b/Engine/Source/Developer/Android/AndroidPlatformEditor/Private/AndroidPlatformEditorPrivatePCH.h
@@ -2,9 +2,9 @@
 
 #pragma once
 
-
-/* Private dependencies
- *****************************************************************************/
-
 #include "Core.h"
+#include "ModuleInterface.h"
+#include "ModuleManager.h"
 #include "Engine.h"
+
+#include "AndroidDeviceDetection.h"
\ No newline at end of file
diff --git a/Engine/Source/Developer/Android/AndroidPlatformEditor/Private/AndroidSDKSettings.cpp b/Engine/Source/Developer/Android/AndroidPlatformEditor/Private/AndroidSDKSettings.cpp
new file mode 100644
--- /dev/null
+++ b/Engine/Source/Developer/Android/AndroidPlatformEditor/Private/AndroidSDKSettings.cpp
@@ -0,0 +1,184 @@
+// Copyright 1998-2014 Epic Games, Inc. All Rights Reserved.
+
+#include "AndroidPlatformEditorPrivatePCH.h"
+
+//#include "EngineTypes.h"
+#include "AndroidSDKSettings.h"
+
+DEFINE_LOG_CATEGORY_STATIC(AndroidSDKSettings, Log, All);
+
+UAndroidSDKSettings::UAndroidSDKSettings(const FObjectInitializer& ObjectInitializer)
+	: Super(ObjectInitializer)
+{
+
+}
+
+#if WITH_EDITOR
+void UAndroidSDKSettings::PostEditChangeProperty(struct FPropertyChangedEvent& PropertyChangedEvent)
+{
+	Super::PostEditChangeProperty(PropertyChangedEvent);
+	UpdateTargetModulePaths(true);	
+}
+
+void UAndroidSDKSettings::SetTargetModule(ITargetPlatformManagerModule * TargetManagerModule)
+{
+	this->TargetManagerModule = TargetManagerModule;
+}
+
+void UAndroidSDKSettings::SetDeviceDetection(IAndroidDeviceDetection * AndroidDeviceDetection)
+{
+	this->AndroidDeviceDetection = AndroidDeviceDetection;
+}
+
+void UAndroidSDKSettings::SetupInitialTargetPaths()
+{
+	// Temp code to pull directly from GConfig until the correct serialisation code exists
+	const FString UserIni = GEngineIni; //  (TEXT("DefaultUser"));
+	FString TempPath;
+	if (GConfig->GetString(TEXT("/Script/AndroidPlatformEditor.AndroidSDKSettings"), TEXT("SDKPath"), TempPath, UserIni))
+	{
+		UE_LOG(AndroidSDKSettings, Warning, TEXT("SDK Path = %s"), *TempPath);
+
+		if (TempPath.Contains(TEXT("Path=")))
+		{
+			SDKPath.Path = TempPath.Mid(7, TempPath.Len() - 9); // Pull the value from the string Path="<payload>"
+		}
+		else
+		{
+			SDKPath.Path = "Bar";
+		}
+	}
+	else
+	{
+		SDKPath.Path = "Foo";
+	}
+	UE_LOG(AndroidSDKSettings, Warning, TEXT("SDK Path = %s"), *SDKPath.Path);
+
+	if (GConfig->GetString(TEXT("/Script/AndroidPlatformEditor.AndroidSDKSettings"), TEXT("NDKPath"), TempPath, UserIni))
+	{
+		if (TempPath.Contains(TEXT("Path=")))
+		{
+			NDKPath.Path = TempPath.Mid(7, TempPath.Len() - 9); // Pull the value from the string Path="<payload>"
+		}
+	}
+
+	if (GConfig->GetString(TEXT("/Script/AndroidPlatformEditor.AndroidSDKSettings"), TEXT("ANTPath"), TempPath, UserIni))
+	{
+		if (TempPath.Contains(TEXT("Path=")))
+		{
+			ANTPath.Path = TempPath.Mid(7, TempPath.Len() - 9); // Pull the value from the string Path="<payload>"
+		}
+	}
+	// End Temp section - this will be removed asap
+
+	// Check each path, if we don't have it then try to look up the values from a known location instead
+	if (SDKPath.Path.IsEmpty())
+	{
+		TCHAR AndroidSDKPath[256];
+		FPlatformMisc::GetEnvironmentVariable(TEXT("ANDROID_HOME"), AndroidSDKPath, ARRAY_COUNT(AndroidSDKPath));
+		if (AndroidSDKPath[0])
+		{
+			SDKPath.Path = AndroidSDKPath;
+		}
+	}
+	UE_LOG(AndroidSDKSettings, Warning, TEXT("SDK Path = %s"), *SDKPath.Path);
+
+	if (NDKPath.Path.IsEmpty())
+	{
+		TCHAR AndroidNDKPath[256];
+		FPlatformMisc::GetEnvironmentVariable(TEXT("NDKROOT"), AndroidNDKPath, ARRAY_COUNT(AndroidNDKPath));
+		if (AndroidNDKPath[0])
+		{
+			NDKPath.Path = AndroidNDKPath;
+		}
+	}
+
+	if (ANTPath.Path.IsEmpty())
+	{
+		TCHAR AndroidANTPath[256];
+		FPlatformMisc::GetEnvironmentVariable(TEXT("ANT_HOME"), AndroidANTPath, ARRAY_COUNT(AndroidANTPath));
+		if (AndroidANTPath[0])
+		{
+			ANTPath.Path = AndroidANTPath;
+		}
+	}
+
+
+#if PLATFORM_MAC
+	// On a Mac, if we still don't have the values then check the .bash_profile file for them
+	if (SDKPath.Path.IsEmpty() || NDKPath.Path.IsEmpty() || ANTPath.Path.IsEmpty())
+	{
+		FArchive* FileReader = IFileManager::Get().CreateFileReader(*FString([@"~/.bash_profile" stringByExpandingTildeInPath]));
+		if (FileReader)
+		{
+			const int64 FileSize = FileReader->TotalSize();
+			ANSICHAR* AnsiContents = (ANSICHAR*)FMemory::Malloc(FileSize);
+			FileReader->Serialize(AnsiContents, FileSize);
+			FileReader->Close();
+			delete FileReader;
+
+			TArray<FString> Lines;
+			FString(ANSI_TO_TCHAR(AnsiContents)).ParseIntoArrayLines(&Lines);
+			FMemory::Free(AnsiContents);
+
+			for (int32 Index = 0; Index < Lines.Num(); Index++)
+			{
+				if (SDKPath.Path.IsEmpty() && Lines[Index].StartsWith(TEXT("export ANDROID_HOME=")))
+				{
+					FString Directory;
+					Lines[Index].Split(TEXT("="), NULL, &Directory);
+					SDKPath.Path = Directory.Replace(TEXT("\""), TEXT(""));
+				}
+				else if (ANTPath.Path.IsEmpty() && Lines[Index].StartsWith(TEXT("export ANT_HOME=")))
+				{
+					FString Directory;
+					Lines[Index].Split(TEXT("="), NULL, &Directory);
+					ANTPath.Path = Directory.Replace(TEXT("\""), TEXT(""));
+				}
+				else if (NDKPath.Path.IsEmpty() && Lines[Index].StartsWith(TEXT("export NDKROOT=")))
+				{
+					FString Directory;
+					Lines[Index].Split(TEXT("="), NULL, &Directory);
+					NDKPath.Path = Directory.Replace(TEXT("\""), TEXT(""));
+				}
+			}
+		}
+	}
+
+#endif
+	UpdateTargetModulePaths(false);
+}
+void UAndroidSDKSettings::UpdateTargetModulePaths(bool bForceUpdate)
+{
+	TArray<FString> Keys;
+	TArray<FString> Values;
+	
+	if (bForceUpdate || !SDKPath.Path.IsEmpty())
+	{
+		Keys.Add(TEXT("ANDROID_HOME"));
+		Values.Add(SDKPath.Path);
+		UE_LOG(AndroidSDKSettings, Log, TEXT("%s"), *SDKPath.Path);
+	}
+	
+	if (bForceUpdate || !NDKPath.Path.IsEmpty())
+	{
+		Keys.Add(TEXT("NDKROOT"));
+		Values.Add(NDKPath.Path);
+	}
+	
+	if (bForceUpdate || !ANTPath.Path.IsEmpty())
+	{
+		Keys.Add(TEXT("ANT_HOME"));
+		Values.Add(ANTPath.Path);
+	}
+
+	SaveConfig();
+	
+	if (Keys.Num() != 0)
+	{
+		TargetManagerModule->UpdatePlatformEnvironment(TEXT("Android"), Keys, Values);
+		AndroidDeviceDetection->UpdateADBPath();
+	}
+}
+
+#endif
diff --git a/Engine/Source/Developer/Android/AndroidPlatformEditor/Private/AndroidSDKSettings.h b/Engine/Source/Developer/Android/AndroidPlatformEditor/Private/AndroidSDKSettings.h
new file mode 100644
--- /dev/null
+++ b/Engine/Source/Developer/Android/AndroidPlatformEditor/Private/AndroidSDKSettings.h
@@ -0,0 +1,43 @@
+// Copyright 1998-2014 Epic Games, Inc. All Rights Reserved.
+
+#pragma once
+
+#include "ITargetPlatformManagerModule.h"
+#include "IAndroidDeviceDetection.h"
+#include "AndroidSDKSettings.generated.h"
+
+
+/**
+ * Implements the settings for the Android SDK setup.
+ */
+//UCLASS(config=Engine, defaultconfig)
+UCLASS()
+class ANDROIDPLATFORMEDITOR_API UAndroidSDKSettings : public UObject
+{
+public:
+	GENERATED_UCLASS_BODY()
+
+//		UPROPERTY(GlobalConfig, EditDefaultsOnly, Category = SDKConfig)
+		UPROPERTY()
+		FDirectoryPath SDKPath;
+
+//	UPROPERTY(GlobalConfig, EditDefaultsOnly, Category = SDKConfig)
+	UPROPERTY()
+		FDirectoryPath NDKPath;
+
+//	UPROPERTY(GlobalConfig, EditDefaultsOnly, Category = SDKConfig)
+	UPROPERTY()
+		FDirectoryPath ANTPath;
+
+#if WITH_EDITOR
+	// UObject interface
+	virtual void PostEditChangeProperty(struct FPropertyChangedEvent& PropertyChangedEvent) override;
+	// End of UObject interface
+	void SetTargetModule(ITargetPlatformManagerModule * TargetManagerModule);
+	void SetDeviceDetection(IAndroidDeviceDetection * AndroidDeviceDetection);
+	void SetupInitialTargetPaths();
+	void UpdateTargetModulePaths(bool bForceUpdate);
+	ITargetPlatformManagerModule * TargetManagerModule;
+	IAndroidDeviceDetection * AndroidDeviceDetection;
+#endif
+};
diff --git a/Engine/Source/Developer/Android/AndroidPlatformEditor/Private/AndroidSDKSettingsCustomization.cpp b/Engine/Source/Developer/Android/AndroidPlatformEditor/Private/AndroidSDKSettingsCustomization.cpp
new file mode 100644
--- /dev/null
+++ b/Engine/Source/Developer/Android/AndroidPlatformEditor/Private/AndroidSDKSettingsCustomization.cpp
@@ -0,0 +1,105 @@
+// Copyright 1998-2014 Epic Games, Inc. All Rights Reserved.
+
+#include "AndroidPlatformEditorPrivatePCH.h"
+#include "AndroidSDKSettingsCustomization.h"
+#include "DetailLayoutBuilder.h"
+#include "DetailCategoryBuilder.h"
+#include "PropertyEditing.h"
+
+#include "ScopedTransaction.h"
+#include "SExternalImageReference.h"
+#include "SHyperlinkLaunchURL.h"
+#include "SPlatformSetupMessage.h"
+#include "SFilePathPicker.h"
+#include "PlatformIconInfo.h"
+#include "SourceControlHelpers.h"
+#include "ManifestUpdateHelper.h"
+#include "SNotificationList.h"
+#include "NotificationManager.h"
+#include "ITargetPlatformManagerModule.h"
+
+#define LOCTEXT_NAMESPACE "AndroidSDKSettings"
+
+//////////////////////////////////////////////////////////////////////////
+// FAndroidSDKSettingsCustomization
+
+TSharedRef<IDetailCustomization> FAndroidSDKSettingsCustomization::MakeInstance()
+{
+	return MakeShareable(new FAndroidSDKSettingsCustomization);
+}
+
+FAndroidSDKSettingsCustomization::FAndroidSDKSettingsCustomization()
+{
+	TargetPlatformManagerModule = &FModuleManager::LoadModuleChecked<ITargetPlatformManagerModule>("TargetPlatform");
+	SetupSDKPaths();
+}
+
+void FAndroidSDKSettingsCustomization::CustomizeDetails(IDetailLayoutBuilder& DetailLayout)
+{
+	SavedLayoutBuilder = &DetailLayout;
+	BuildSDKPathSection(DetailLayout);
+}
+
+
+void FAndroidSDKSettingsCustomization::BuildSDKPathSection(IDetailLayoutBuilder& DetailLayout)
+{
+	SetupSDKPaths();}
+
+
+void FAndroidSDKSettingsCustomization::SetupSDKPaths()
+{
+	// Try and query the env vars
+	UAndroidSDKSettings * settings = GetMutableDefault<UAndroidSDKSettings>();
+	bool changed = false;
+	
+	if (settings->SDKPath.Path.IsEmpty())
+	{
+		TCHAR AndroidSDKPath[256];
+		FPlatformMisc::GetEnvironmentVariable(TEXT("ANDROID_HOME"), AndroidSDKPath, ARRAY_COUNT(AndroidSDKPath));
+		if (AndroidSDKPath[0])
+		{
+			settings->SDKPath.Path = AndroidSDKPath;
+			changed |= true;
+		}
+				
+	}
+	
+
+	if (settings->NDKPath.Path.IsEmpty())
+	{
+		TCHAR AndroidNDKPath[256];
+		FPlatformMisc::GetEnvironmentVariable(TEXT("NDKROOT"), AndroidNDKPath, ARRAY_COUNT(AndroidNDKPath));
+		if (AndroidNDKPath[0])
+		{
+			settings->NDKPath.Path = AndroidNDKPath;
+			changed |= true;
+		}
+		
+	}
+	
+
+	if (settings->ANTPath.Path.IsEmpty())
+	{
+		TCHAR AndroidANTPath[256];
+		FPlatformMisc::GetEnvironmentVariable(TEXT("ANT_HOME"), AndroidANTPath, ARRAY_COUNT(AndroidANTPath));
+		if (AndroidANTPath[0])
+		{
+			settings->ANTPath.Path = AndroidANTPath;
+			changed |= true;
+		}
+		
+	}
+	
+	if (changed)
+	{
+		settings->Modify();
+		settings->SaveConfig();
+	}
+
+	settings->UpdateTargetModulePaths(false);
+
+}
+
+//////////////////////////////////////////////////////////////////////////
+
+#undef LOCTEXT_NAMESPACE
diff --git a/Engine/Source/Developer/Android/AndroidPlatformEditor/Private/AndroidSDKSettingsCustomization.h b/Engine/Source/Developer/Android/AndroidPlatformEditor/Private/AndroidSDKSettingsCustomization.h
new file mode 100644
--- /dev/null
+++ b/Engine/Source/Developer/Android/AndroidPlatformEditor/Private/AndroidSDKSettingsCustomization.h
@@ -0,0 +1,32 @@
+// Copyright 1998-2014 Epic Games, Inc. All Rights Reserved.
+
+#pragma once
+
+#include "EditorStyle.h"
+#include "PropertyEditorModule.h"
+#include "AndroidSDKSettings.h"
+
+//////////////////////////////////////////////////////////////////////////
+// FAndroidSDKSettingsCustomization
+
+class FAndroidSDKSettingsCustomization : public IDetailCustomization
+{
+public:
+	// Makes a new instance of this detail layout class for a specific detail view requesting it
+	static TSharedRef<IDetailCustomization> MakeInstance();
+
+	// IDetailCustomization interface
+	virtual void CustomizeDetails(IDetailLayoutBuilder& DetailLayout) override;
+	// End of IDetailCustomization interface
+
+private:
+	FAndroidSDKSettingsCustomization();
+
+	void BuildSDKPathSection(IDetailLayoutBuilder& DetailLayout);
+	void SetupSDKPaths();
+
+private:
+	IDetailLayoutBuilder* SavedLayoutBuilder;
+	ITargetPlatformManagerModule * TargetPlatformManagerModule;
+
+};
diff --git a/Engine/Source/Developer/Android/AndroidTargetPlatform/Private/AndroidTargetPlatform.inl b/Engine/Source/Developer/Android/AndroidTargetPlatform/Private/AndroidTargetPlatform.inl
--- a/Engine/Source/Developer/Android/AndroidTargetPlatform/Private/AndroidTargetPlatform.inl
+++ b/Engine/Source/Developer/Android/AndroidTargetPlatform/Private/AndroidTargetPlatform.inl
@@ -89,29 +89,13 @@ inline bool FAndroidTargetPlatform<TPlatformProperties>::IsRunningPlatform( ) co
 template<class TPlatformProperties>
 inline bool FAndroidTargetPlatform<TPlatformProperties>::IsSdkInstalled(bool bProjectHasCode, FString& OutDocumentationPath) const
 {
+
 	OutDocumentationPath = FString("Shared/Tutorials/SettingUpAndroidTutorial");
 
-	TCHAR ANDROID_HOME[MAX_PATH];
+#if PLATFORM_WINDOWS
 	TCHAR JAVA_HOME[MAX_PATH];
-	TCHAR ANT_HOME[MAX_PATH];
-	TCHAR NDKROOT[MAX_PATH];
-	FPlatformMisc::GetEnvironmentVariable(TEXT("ANDROID_HOME"), ANDROID_HOME, MAX_PATH);
 	FPlatformMisc::GetEnvironmentVariable(TEXT("JAVA_HOME"), JAVA_HOME, MAX_PATH);
-	FPlatformMisc::GetEnvironmentVariable(TEXT("ANT_HOME"), ANT_HOME, MAX_PATH);
-	FPlatformMisc::GetEnvironmentVariable(TEXT("NDKROOT"), NDKROOT, MAX_PATH);
-
-	// make sure ANDROID_HOME points to the right thing
-	if (ANDROID_HOME[0] == 0 ||
-#if PLATFORM_WINDOWS
-		IFileManager::Get().FileSize(*(FString(ANDROID_HOME) / TEXT("platform-tools/adb.exe"))) < 0)
-#else
-		IFileManager::Get().FileSize(*(FString(ANDROID_HOME) / TEXT("platform-tools/adb"))) < 0)
-#endif
-	{
-		return false;
-	}
 
-#if PLATFORM_WINDOWS
 	// make sure that JAVA_HOME points to the right thing
 	if (JAVA_HOME[0] == 0 ||
 		IFileManager::Get().FileSize(*(FString(JAVA_HOME) / TEXT("bin/javac.exe"))) < 0)
@@ -120,26 +104,6 @@ inline bool FAndroidTargetPlatform<TPlatformProperties>::IsSdkInstalled(bool bPr
 	}
 #endif
 
-	// now look for ANT_HOME, or the ADT workaround of looking for a plugin
-	if (ANT_HOME[0] == 0)
-	{
-		// look for plugins in eclipse (this is enough to assume we have an ant plugin)
-		if (!IFileManager::Get().DirectoryExists(*(FString(ANDROID_HOME) / TEXT("../eclipse/plugins"))))
-		{
-			return false;
-		}
-	}
-
-	// we need NDKROOT if the game has code
-	if (bProjectHasCode)
-	{
-		if (NDKROOT[0] == 0 ||
-			IFileManager::Get().FileSize(*(FString(NDKROOT) / TEXT("ndk-build.cmd"))) < 0)
-		{
-			return false;
-		}
-	}
-
 	return true;
 }
 
diff --git a/Engine/Source/Developer/TargetPlatform/Private/TargetPlatformManagerModule.cpp b/Engine/Source/Developer/TargetPlatform/Private/TargetPlatformManagerModule.cpp
--- a/Engine/Source/Developer/TargetPlatform/Private/TargetPlatformManagerModule.cpp
+++ b/Engine/Source/Developer/TargetPlatform/Private/TargetPlatformManagerModule.cpp
@@ -590,6 +590,12 @@ protected:
 		}
 	}
 
+	bool UpdatePlatformEnvironment(FString PlatformName, TArray<FString> &Keys, TArray<FString> &Values)
+	{
+		SetupEnvironmentVariables(Keys, Values);
+		return SetupSDKStatus(PlatformName);	
+	}
+
 	bool SetupAndValidateAutoSDK(const FString& AutoSDKPath)
 	{
 #if AUTOSDKS_ENABLED
@@ -740,13 +746,8 @@ protected:
 			// don't actually set anything until we successfully validate and read all values in.
 			// we don't want to set a few vars, return a failure, and then have a platform try to
 			// build against a manually installed SDK with half-set env vars.
-			for (int i = 0; i < EnvVarNames.Num(); ++i)
-			{
-				const FString& EnvVarName = EnvVarNames[i];
-				const FString& EnvVarValue = EnvVarValues[i];
-				UE_LOG(LogTargetPlatformManager, Verbose, TEXT("Setting variable '%s' to '%s'."), *EnvVarName, *EnvVarValue);
-				FPlatformMisc::SetEnvironmentVar(*EnvVarName, *EnvVarValue);
-			}
+			SetupEnvironmentVariables(EnvVarNames, EnvVarValues);
+
 
 			const int32 MaxPathVarLen = 32768;
 			TCHAR OrigPathVarMem[MaxPathVarLen];
@@ -821,16 +822,25 @@ protected:
 	}
 
 	bool SetupSDKStatus()
+	{
+		return SetupSDKStatus(TEXT(""));
+	}
+
+	bool SetupSDKStatus(FString targetPlatforms)
 	{
 		// run UBT with -validate -allplatforms and read the output
 #if PLATFORM_MAC
 		FString CmdExe = TEXT("/bin/sh");
         FString ScriptPath = FPaths::ConvertRelativePathToFull(FPaths::EngineDir() / TEXT("Build/BatchFiles/Mac/RunMono.sh"));
-		FString CommandLine = TEXT("\"") + ScriptPath + TEXT("\" \"") + FPaths::ConvertRelativePathToFull(FPaths::EngineDir() / TEXT("Binaries/DotNet/UnrealBuildTool.exe")) + TEXT("\" -validateplatform -allplatforms");
+		FString CommandLine = TEXT("\"") + ScriptPath + TEXT("\" \"") + FPaths::ConvertRelativePathToFull(FPaths::EngineDir() / TEXT("Binaries/DotNet/UnrealBuildTool.exe")) + TEXT("\" -validateplatform");
 #else
 		FString CmdExe = FPaths::ConvertRelativePathToFull(FPaths::EngineDir() / TEXT("Binaries/DotNet/UnrealBuildTool.exe"));
-		FString CommandLine = TEXT("-validateplatform -allplatforms");
+		FString CommandLine = TEXT("-validateplatform");
 #endif
+
+		// Allow for only a subset of platforms to be reparsed - needed when kicking a change from the UI
+		CommandLine += targetPlatforms.IsEmpty() ? TEXT(" -allplatforms") : (TEXT(" -platforms=") + targetPlatforms);
+
 		TSharedPtr<FMonitoredProcess> UBTProcess = MakeShareable(new FMonitoredProcess(CmdExe, CommandLine, true));
 		UBTProcess->OnOutput().BindStatic(&FTargetPlatformManagerModule::OnStatusOutput);
 		SDKStatusMessage = TEXT("");
@@ -891,6 +901,17 @@ protected:
 
 private:
 
+	void SetupEnvironmentVariables(TArray<FString> &EnvVarNames, TArray<FString> EnvVarValues)
+	{
+		for (int i = 0; i < EnvVarNames.Num(); ++i)
+		{
+			const FString& EnvVarName = EnvVarNames[i];
+			const FString& EnvVarValue = EnvVarValues[i];
+			UE_LOG(LogTargetPlatformManager, Verbose, TEXT("Setting variable '%s' to '%s'."), *EnvVarName, *EnvVarValue);
+			FPlatformMisc::SetEnvironmentVar(*EnvVarName, *EnvVarValue);
+		}
+	}
+
 	void ModulesChangesCallback(FName ModuleName, EModuleChangeReason ReasonForChange)
 	{
 		if (ModuleName.ToString().Contains(TEXT("TargetPlatform")) )
diff --git a/Engine/Source/Developer/TargetPlatform/Public/Interfaces/ITargetPlatformManagerModule.h b/Engine/Source/Developer/TargetPlatform/Public/Interfaces/ITargetPlatformManagerModule.h
--- a/Engine/Source/Developer/TargetPlatform/Public/Interfaces/ITargetPlatformManagerModule.h
+++ b/Engine/Source/Developer/TargetPlatform/Public/Interfaces/ITargetPlatformManagerModule.h
@@ -2,6 +2,8 @@
 
 #pragma once
 
+#include "TargetDeviceID.h"
+#include "ITargetDevice.h"
 
 /**
  * Module for the target platform manager
@@ -140,6 +142,11 @@ public:
 	 */
 	virtual uint16 ShaderFormatVersion(FName Name) = 0;
 
+	/**
+	 * Allows changes to environment for a given platform
+	 */
+	virtual bool UpdatePlatformEnvironment(FString PlatformName, TArray<FString> &Keys, TArray<FString> &Values) = 0;
+
 public:
 
 	/** Virtual destructor. */
diff --git a/Engine/Source/Programs/AutomationTool/Android/AndroidPlatform.Automation.cs b/Engine/Source/Programs/AutomationTool/Android/AndroidPlatform.Automation.cs
--- a/Engine/Source/Programs/AutomationTool/Android/AndroidPlatform.Automation.cs
+++ b/Engine/Source/Programs/AutomationTool/Android/AndroidPlatform.Automation.cs
@@ -18,24 +18,6 @@ public class AndroidPlatform : Platform
 	public AndroidPlatform()
 		: base(UnrealTargetPlatform.Android)
 	{
-		bool bNeedsAndroidHome = string.IsNullOrEmpty(Environment.GetEnvironmentVariable("ANDROID_HOME"));
-		if (Utils.IsRunningOnMono && bNeedsAndroidHome)
-		{
-			// Try reading env variable we need from .bash_profile
-			string BashProfilePath = Path.Combine(Environment.GetFolderPath(Environment.SpecialFolder.Personal), ".bash_profile");
-			if(File.Exists(BashProfilePath))
-			{
-				string[] BashProfileContents = File.ReadAllLines(BashProfilePath);
-				foreach (string Line in BashProfileContents)
-				{
-					if (Line.StartsWith("export ANDROID_HOME="))
-					{
-						string PathVar = Line.Split('=')[1].Replace("\"", "");
-						Environment.SetEnvironmentVariable("ANDROID_HOME", PathVar);
-					}
-				}
-			}
-		}
 	}
 
 	private static string GetSONameWithoutArchitecture(ProjectParams Params, string DecoratedExeName)
@@ -106,8 +88,8 @@ public class AndroidPlatform : Platform
 
 	private static string GetDeviceObbName(string ApkName)
 	{
-		string ObbName = GetFinalObbName(ApkName);
-		string PackageName = GetPackageInfo(ApkName, false);
+        string ObbName = GetFinalObbName(ApkName);
+        string PackageName = GetPackageInfo(ApkName, false);
 		return "obb/" + PackageName + "/" + Path.GetFileName(ObbName);
 	}
 
@@ -128,6 +110,80 @@ public class AndroidPlatform : Platform
 		return Path.Combine(Path.GetDirectoryName(ApkName), "Install_" + Params.ShortProjectName + "_" + Params.ClientConfigsToBuild[0].ToString() + Architecture + GPUArchitecture + (Utils.IsRunningOnMono ? ".command" : ".bat"));
 	}
 
+//     public static Dictionary<string, string> SetupEnv(string ProjectPath)
+//     {
+//         if(configCacheIni == null)
+//         {
+//             CommandUtils.LogWarning(Path.GetDirectoryName(ProjectPath));
+//             configCacheIni = new UnrealBuildTool.ConfigCacheIni("Engine", Path.GetDirectoryName(ProjectPath));
+// 
+//             string[] EnvVarNames = {"ANDROID_HOME", "NDKHOME", "ANT_HOME"};
+// 
+//             string path;
+//             if(configCacheIni.GetPath("/Script/AndroidPlatformEditor.AndroidSDKSettings", "SDKPath", out path))
+//             {
+//                 AndroidEnv.Add("ANDROID_HOME", path);
+//             }
+//             else // if we fail then try and get it from the real environment
+//             {
+//                 AndroidEnv.Add("ANDROID_HOME", Environment.GetEnvironmentVariable("ANDROID_HOME"));
+//             }
+// 
+//             if (configCacheIni.GetPath("/Script/AndroidPlatformEditor.AndroidSDKSettings", "NDKPath", out path))
+//             {
+//                 AndroidEnv.Add("NDKHOME", path);
+//             }
+//             else // if we fail then try and get it from the real environment
+//             {
+//                 AndroidEnv.Add("NDKHOME", Environment.GetEnvironmentVariable("NDKHOME"));
+//             }
+// 
+//             if (configCacheIni.GetPath("/Script/AndroidPlatformEditor.AndroidSDKSettings", "ANTPath", out path))
+//             {
+//                 AndroidEnv.Add("ANT_HOME", path);
+//             }
+//             else // if we fail then try and get it from the real environment
+//             {
+//                 AndroidEnv.Add("ANT_HOME", Environment.GetEnvironmentVariable("ANT_HOME"));
+//             }
+// 
+//             // If we are on Mono and we are still missing a key then go and find it from the .bash_profile
+//             if (Utils.IsRunningOnMono && EnvVarNames.All( s => AndroidEnv.ContainsKey(s)))
+//             {
+//                 string BashProfilePath = Path.Combine(Environment.GetFolderPath(Environment.SpecialFolder.Personal), ".bash_profile");
+//                 if (File.Exists(BashProfilePath))
+//                 {
+//                     string[] BashProfileContents = File.ReadAllLines(BashProfilePath);
+//                     foreach (string Line in BashProfileContents)
+//                     {
+//                         foreach (var key in EnvVarNames)
+//                         {
+//                             if (AndroidEnv.ContainsKey(key))
+//                             {
+//                                 continue;
+//                             }
+//                             
+//                             if (Line.StartsWith("export " + key + "="))
+//                             {
+//                                 string PathVar = Line.Split('=')[1].Replace("\"", "");
+//                                 AndroidEnv.Add(key, PathVar);
+//                             }
+//                         }
+//                     }
+//                 }                    
+//             }
+// 
+//             // Set for the process
+//             foreach(var kvp in AndroidEnv)
+//             {
+//                 Environment.SetEnvironmentVariable(kvp.Key, kvp.Value);
+//             }
+// 
+//         }
+// 
+//         return AndroidEnv;
+//     }
+
 	public override void Package(ProjectParams Params, DeploymentContext SC, int WorkingCL)
 	{
 		string[] Architectures = UnrealBuildTool.AndroidToolChain.GetAllArchitectures();
@@ -189,19 +245,19 @@ public class AndroidPlatform : Platform
 					Deploy.PrepForUATPackageOrDeploy(Params.ShortProjectName, SC.ProjectRoot, SOName, SC.LocalRoot + "/Engine", Params.Distribution, CookFlavor);
 				}
 
-				// Create APK specific OBB in case we have a detached OBB.
-				string DeviceObbName = "";
-				string ObbName = "";
-				if (!Params.OBBinAPK)
-				{
-					DeviceObbName = GetDeviceObbName(ApkName);
-					ObbName = GetFinalObbName(ApkName);
-					CopyFile(LocalObbName, ObbName);
-				}
+			    // Create APK specific OBB in case we have a detached OBB.
+			    string DeviceObbName = "";
+			    string ObbName = "";
+			    if (!Params.OBBinAPK)
+			    {
+				    DeviceObbName = GetDeviceObbName(ApkName);
+				    ObbName = GetFinalObbName(ApkName);
+				    CopyFile(LocalObbName, ObbName);
+			    }
 
 				// Write install batch file(s).
 
-				string PackageName = GetPackageInfo(ApkName, false);
+                string PackageName = GetPackageInfo(ApkName, false);
 				// make a batch file that can be used to install the .apk and .obb files
 				string[] BatchLines;
 				if (Utils.IsRunningOnMono)
@@ -210,7 +266,8 @@ public class AndroidPlatform : Platform
 					BatchLines = new string[] {
 						"#!/bin/sh",
 						"cd \"`dirname \"$0\"`\"",
-						"ADB=$ANDROID_HOME/platform-tools/adb",
+                        "ADB=",
+						"if [ \"$ANDROID_HOME\" != \"\"]; then ADB=$ANDROID_HOME/platform-tools/adb; else ADB=" +Environment.GetEnvironmentVariable("ANDROID_HOME") + "; fi",
 						"DEVICE=",
 						"if [ \"$1\" != \"\" ]; then DEVICE=\"-s $1\"; fi",
 						"$ADB $DEVICE uninstall " + PackageName,
@@ -241,10 +298,12 @@ public class AndroidPlatform : Platform
 					Log("Writing bat for install with {0}", Params.OBBinAPK ? "OBB in APK" : "OBB separate");
 					BatchLines = new string[] {
 						"setlocal",
-						"set ADB=%ANDROID_HOME%\\platform-tools\\adb.exe",
+                        "set ANDROIDHOME=%ANDROID_HOME%",
+                        "if \"%ANDROIDHOME%\"==\"\" set ANDROIDHOME="+Environment.GetEnvironmentVariable("ANDROID_HOME"),
+						"set ADB=%ANDROIDHOME%\\platform-tools\\adb.exe",
 						"set DEVICE=",
-						"if not \"%1\"==\"\" set DEVICE=-s %1",
-						"for /f \"delims=\" %%A in ('%ADB% %DEVICE% " + GetStorageQueryCommand() + "') do @set STORAGE=%%A",
+                        "if not \"%1\"==\"\" set DEVICE=-s %1",
+                        "for /f \"delims=\" %%A in ('%ADB% %DEVICE% " + GetStorageQueryCommand() +"') do @set STORAGE=%%A",
 						"%ADB% %DEVICE% uninstall " + PackageName,
 						"%ADB% %DEVICE% install " + Path.GetFileName(ApkName),
 						"@if \"%ERRORLEVEL%\" NEQ \"0\" goto Error",
@@ -491,8 +550,8 @@ public class AndroidPlatform : Platform
 		{
 			// cache some strings
 			string BaseCommandline = "push";
-			string RemoteDir = StorageLocation + "/" + Params.ShortProjectName;
-            string UE4GameRemoteDir = StorageLocation + "/" + Params.ShortProjectName;
+			string RemoteDir = StorageLocation + "/" /*+ Params.ProjectDirPrefix + "/"*/ + Params.ShortProjectName;
+            string UE4GameRemoteDir = StorageLocation + "/" /*+ Params.ProjectDirPrefix + "/"*/ + Params.ShortProjectName;
 
 			// make sure device is at a clean state
 			RunAdbCommand(Params, "shell rm -r " + RemoteDir);
@@ -618,7 +677,7 @@ public class AndroidPlatform : Platform
 		{
 			// cache some strings
 			string BaseCommandline = "push";
-            string RemoteDir = StorageLocation + "/" + Params.ShortProjectName;
+            string RemoteDir = StorageLocation + "/" /*+ Params.ProjectDirPrefix + "/"*/ + Params.ShortProjectName;
 
 			string FinalRemoteDir = RemoteDir;
 			/*
@@ -695,12 +754,13 @@ public class AndroidPlatform : Platform
 	private static string GetAaptPath()
 	{
 		// there is a numbered directory in here, hunt it down
-		string[] Subdirs = Directory.GetDirectories(Environment.ExpandEnvironmentVariables("%ANDROID_HOME%/build-tools/"));
-		if (Subdirs.Length == 0)
-		{
+        string path = Environment.ExpandEnvironmentVariables("%ANDROID_HOME%/build-tools/");
+		string[] Subdirs = Directory.GetDirectories(path);
+        if (Subdirs.Length == 0)
+        {
             ErrorReporter.Error("Failed to find %ANDROID_HOME%/build-tools subdirectory", (int)ErrorCodes.Error_AndroidBuildToolsPathNotFound);
-			throw new AutomationException("Failed to find %ANDROID_HOME%/build-tools subdirectory");
-		}
+            throw new AutomationException("Failed to find %ANDROID_HOME%/build-tools subdirectory");
+        }
 		// we expect there to be one, so use the first one
 		return Path.Combine(Subdirs[0], Utils.IsRunningOnMono ? "aapt" : "aapt.exe");
 	}
diff --git a/Engine/Source/Programs/AutomationTool/ProcessUtils.cs b/Engine/Source/Programs/AutomationTool/ProcessUtils.cs
--- a/Engine/Source/Programs/AutomationTool/ProcessUtils.cs
+++ b/Engine/Source/Programs/AutomationTool/ProcessUtils.cs
@@ -760,9 +760,9 @@ namespace AutomationTool
 		/// <param name="LogName">Name of the logfile ( if null, executable name is used )</param>
 		/// <param name="Input">Optional Input for the program (will be provided as stdin)</param>
 		/// <param name="Options">Defines the options how to run. See ERunOptions.</param>
-		public static void RunAndLog(CommandEnvironment Env, string App, string CommandLine, string LogName = null, int MaxSuccessCode = 0, string Input = null, ERunOptions Options = ERunOptions.Default)
+		public static void RunAndLog(CommandEnvironment Env, string App, string CommandLine, string LogName = null, int MaxSuccessCode = 0, string Input = null, ERunOptions Options = ERunOptions.Default, Dictionary<string, string> EnvVars = null)
 		{
-            RunAndLog(App, CommandLine, GetRunAndLogLogName(Env, App, LogName), MaxSuccessCode, Input, Options);
+            RunAndLog(App, CommandLine, GetRunAndLogLogName(Env, App, LogName), MaxSuccessCode, Input, Options, EnvVars);
 		}
         /// <summary>
         /// Runs external program and writes the output to a logfile.
@@ -772,9 +772,9 @@ namespace AutomationTool
         /// <param name="Logfile">Full path to the logfile, where the application output should be written to.</param>
 		/// <param name="Input">Optional Input for the program (will be provided as stdin)</param>
 		/// <param name="Options">Defines the options how to run. See ERunOptions.</param>
-        public static string RunAndLog(string App, string CommandLine, string Logfile = null, int MaxSuccessCode = 0, string Input = null, ERunOptions Options = ERunOptions.Default)
+        public static string RunAndLog(string App, string CommandLine, string Logfile = null, int MaxSuccessCode = 0, string Input = null, ERunOptions Options = ERunOptions.Default, Dictionary<string, string> EnvVars = null)
         {
-            ProcessResult Result = Run(App, CommandLine, Input, Options);
+            ProcessResult Result = Run(App, CommandLine, Input, Options, EnvVars);
             if (Result.Output.Length > 0 && Logfile != null)
             {
                 WriteToFile(Logfile, Result.Output);
@@ -809,9 +809,9 @@ namespace AutomationTool
         /// <param name="CommandLine">Commandline to pass on to the executable</param>
         /// <param name="Logfile">Full path to the logfile, where the application output should be written to.</param>
         /// <returns>Whether the program executed successfully or not.</returns>
-        public static string RunAndLog(string App, string CommandLine, out int SuccessCode, string Logfile = null)
+        public static string RunAndLog(string App, string CommandLine, out int SuccessCode, string Logfile = null, Dictionary<string, string> EnvVars = null)
         {
-            ProcessResult Result = Run(App, CommandLine);
+            ProcessResult Result = Run(App, CommandLine, Env: EnvVars);
             SuccessCode = Result.ExitCode;
             if (Result.Output.Length > 0 && Logfile != null)
             {
@@ -832,9 +832,9 @@ namespace AutomationTool
         /// <param name="CommandLine">Commandline to pass on to the executable</param>
         /// <param name="LogName">Name of the logfile ( if null, executable name is used )</param>
         /// <returns>Whether the program executed successfully or not.</returns>
-        public static string RunAndLog(CommandEnvironment Env, string App, string CommandLine, out int SuccessCode, string LogName = null)
+        public static string RunAndLog(CommandEnvironment Env, string App, string CommandLine, out int SuccessCode, string LogName = null, Dictionary<string, string> EnvVars = null)
         {
-            return RunAndLog(App, CommandLine, out SuccessCode, GetRunAndLogLogName(Env, App, LogName));
+            return RunAndLog(App, CommandLine, out SuccessCode, GetRunAndLogLogName(Env, App, LogName), EnvVars);
         }
 
 		/// <summary>
diff --git a/Engine/Source/Programs/AutomationTool/ProjectParams.cs b/Engine/Source/Programs/AutomationTool/ProjectParams.cs
--- a/Engine/Source/Programs/AutomationTool/ProjectParams.cs
+++ b/Engine/Source/Programs/AutomationTool/ProjectParams.cs
@@ -1622,6 +1622,17 @@ namespace AutomationTool
 			get { return ProjectUtils.GetShortProjectName(RawProjectPath); }
 		}
 
+        /// <summary>
+        /// Get a directory pre-fix for the game's deploy location. Defaults to UE4Game is not supplied
+        /// </summary>
+
+        public string ProjectDirPrefix
+        {
+            get { return String.IsNullOrEmpty(ProjectDirPrefixValue) ? "UE4Game" : ProjectDirPrefixValue;  }
+            set { ProjectDirPrefixValue = value; }
+        }
+        private string ProjectDirPrefixValue;
+
 		/// <summary>
 		/// True if this project contains source code.
 		/// </summary>	
diff --git a/Engine/Source/Programs/AutomationTool/Scripts/BuildCookRun.Automation.cs b/Engine/Source/Programs/AutomationTool/Scripts/BuildCookRun.Automation.cs
--- a/Engine/Source/Programs/AutomationTool/Scripts/BuildCookRun.Automation.cs
+++ b/Engine/Source/Programs/AutomationTool/Scripts/BuildCookRun.Automation.cs
@@ -133,9 +133,9 @@ public class BuildCookRun : BuildCommand
 			}
 		}
 
-        var configCache = new UnrealBuildTool.ConfigCacheIni("Game", Path.GetDirectoryName(ProjectPath));
+        var gameConfigCache = new UnrealBuildTool.ConfigCacheIni("Game", Path.GetDirectoryName(ProjectPath));
         bool obbInAPK = false;
-        configCache.GetBool("/Script/UnrealEd.ProjectPackagingSettings", "UseOBB_InAPK", out obbInAPK);
+        gameConfigCache.GetBool("/Script/UnrealEd.ProjectPackagingSettings", "UseOBB_InAPK", out obbInAPK);
         Params.OBBinAPK = obbInAPK;         
 
 		Params.ValidateAndLog();
diff --git a/Engine/Source/Programs/AutomationTool/UBTUtils.cs b/Engine/Source/Programs/AutomationTool/UBTUtils.cs
--- a/Engine/Source/Programs/AutomationTool/UBTUtils.cs
+++ b/Engine/Source/Programs/AutomationTool/UBTUtils.cs
@@ -15,7 +15,7 @@ namespace AutomationTool
 		/// <param name="Env">Environment to use.</param>
 		/// <param name="CommandLine">Commandline to pass on to UBT.</param>
 		/// <param name="LogName">Optional logfile name.</param>
-		public static void RunUBT(CommandEnvironment Env, string UBTExecutable, string CommandLine, string LogName = null)
+        public static void RunUBT(CommandEnvironment Env, string UBTExecutable, string CommandLine, string LogName = null, Dictionary<string, string> EnvVars = null)
 		{
 			if (!FileExists(UBTExecutable))
 			{
@@ -30,7 +30,7 @@ namespace AutomationTool
 			{
 				CommandLine += " -nocreatestub";
 			}
-			CommandUtils.RunAndLog(Env, UBTExecutable, CommandLine, LogName);
+			CommandUtils.RunAndLog(Env, UBTExecutable, CommandLine, LogName, EnvVars: EnvVars);
 		}
 
 		/// <summary>
@@ -66,9 +66,9 @@ namespace AutomationTool
 		/// <param name="Config">Configuration to build.</param>
 		/// <param name="AdditionalArgs">Additional arguments to pass on to UBT.</param>
 		/// <param name="LogName">Optional logifle name.</param>
-		public static void RunUBT(CommandEnvironment Env, string UBTExecutable, string Project, string Target, string Platform, string Config, string AdditionalArgs = "", string LogName = null)
+		public static void RunUBT(CommandEnvironment Env, string UBTExecutable, string Project, string Target, string Platform, string Config, string AdditionalArgs = "", string LogName = null, Dictionary<string, string> EnvVars = null)
 		{
-			RunUBT(Env, UBTExecutable, UBTCommandline(Project, Target, Platform, Config, AdditionalArgs), LogName);
+			RunUBT(Env, UBTExecutable, UBTCommandline(Project, Target, Platform, Config, AdditionalArgs), LogName, EnvVars);
 		}
 
 	}
diff --git a/Engine/Source/Programs/AutomationTool/UE4Build.cs b/Engine/Source/Programs/AutomationTool/UE4Build.cs
--- a/Engine/Source/Programs/AutomationTool/UE4Build.cs
+++ b/Engine/Source/Programs/AutomationTool/UE4Build.cs
@@ -140,7 +140,7 @@ namespace AutomationTool
 			public string OutputCaption;
 		}
 
-		XGEItem XGEPrepareBuildWithUBT(string ProjectName, string TargetName, UnrealBuildTool.UnrealTargetPlatform Platform, string Config, string UprojectPath, bool ForceMonolithic = false, bool ForceNonUnity = false, bool ForceDebugInfo = false, string InAddArgs = "", bool ForceUnity = false)
+		XGEItem XGEPrepareBuildWithUBT(string ProjectName, string TargetName, UnrealBuildTool.UnrealTargetPlatform Platform, string Config, string UprojectPath, bool ForceMonolithic = false, bool ForceNonUnity = false, bool ForceDebugInfo = false, string InAddArgs = "", bool ForceUnity = false, Dictionary<string, string> EnvVars = null)
 		{
 			string AddArgs = "";
 			if (string.IsNullOrEmpty(UprojectPath) == false)
@@ -184,7 +184,7 @@ namespace AutomationTool
 
 			ClearExportedXGEXML();
 
-			RunUBT(CmdEnv, UBTExecutable: UBTExecutable, Project: ProjectName, Target: TargetName, Platform: Platform.ToString(), Config: Config, AdditionalArgs: "-generatemanifest -nobuilduht -xgeexport" + AddArgs);
+			RunUBT(CmdEnv, UBTExecutable: UBTExecutable, Project: ProjectName, Target: TargetName, Platform: Platform.ToString(), Config: Config, AdditionalArgs: "-generatemanifest -nobuilduht -xgeexport" + AddArgs, EnvVars: EnvVars);
 
 			PrepareManifest(UBTManifest);
 
@@ -247,7 +247,7 @@ namespace AutomationTool
 			}
 		}
 
-		void CleanWithUBT(string ProjectName, string TargetName, UnrealBuildTool.UnrealTargetPlatform Platform, string Config, string UprojectPath, bool ForceMonolithic = false, bool ForceNonUnity = false, bool ForceDebugInfo = false, string InAddArgs = "", bool ForceUnity = false)
+		void CleanWithUBT(string ProjectName, string TargetName, UnrealBuildTool.UnrealTargetPlatform Platform, string Config, string UprojectPath, bool ForceMonolithic = false, bool ForceNonUnity = false, bool ForceDebugInfo = false, string InAddArgs = "", bool ForceUnity = false, Dictionary<string, string> EnvVars = null)
 		{
 			string AddArgs = "";
 			if (string.IsNullOrEmpty(UprojectPath) == false)
@@ -278,10 +278,10 @@ namespace AutomationTool
 
 			PrepareUBT();
 
-			RunUBT(CmdEnv, UBTExecutable: UBTExecutable, Project: ProjectName, Target: TargetName, Platform: Platform.ToString(), Config: Config, AdditionalArgs: "-clean" + AddArgs);
+			RunUBT(CmdEnv, UBTExecutable: UBTExecutable, Project: ProjectName, Target: TargetName, Platform: Platform.ToString(), Config: Config, AdditionalArgs: "-clean" + AddArgs, EnvVars: EnvVars);
 		}
 
-		void BuildWithUBT(string ProjectName, string TargetName, UnrealBuildTool.UnrealTargetPlatform TargetPlatform, string Config, string UprojectPath, bool ForceMonolithic = false, bool ForceNonUnity = false, bool ForceDebugInfo = false, bool ForceFlushMac = false, bool DisableXGE = false, string InAddArgs = "", bool ForceUnity = false)
+		void BuildWithUBT(string ProjectName, string TargetName, UnrealBuildTool.UnrealTargetPlatform TargetPlatform, string Config, string UprojectPath, bool ForceMonolithic = false, bool ForceNonUnity = false, bool ForceDebugInfo = false, bool ForceFlushMac = false, bool DisableXGE = false, string InAddArgs = "", bool ForceUnity = false, Dictionary<string, string> EnvVars = null)
 		{
 			string AddArgs = "";
 			if (string.IsNullOrEmpty(UprojectPath) == false)
@@ -336,12 +336,12 @@ namespace AutomationTool
 
 				DeleteFile(UBTManifest);
 
-				RunUBT(CmdEnv, UBTExecutable: UBTExecutable, Project: ProjectName, Target: TargetName, Platform: TargetPlatform.ToString(), Config: Config, AdditionalArgs: "-generatemanifest" + AddArgs);
+				RunUBT(CmdEnv, UBTExecutable: UBTExecutable, Project: ProjectName, Target: TargetName, Platform: TargetPlatform.ToString(), Config: Config, AdditionalArgs: "-generatemanifest" + AddArgs, EnvVars: EnvVars);
 
 				PrepareManifest(UBTManifest);
 			}
 
-			RunUBT(CmdEnv, UBTExecutable: UBTExecutable, Project: ProjectName, Target: TargetName, Platform: TargetPlatform.ToString(), Config: Config, AdditionalArgs: AddArgs);
+			RunUBT(CmdEnv, UBTExecutable: UBTExecutable, Project: ProjectName, Target: TargetName, Platform: TargetPlatform.ToString(), Config: Config, AdditionalArgs: AddArgs, EnvVars: EnvVars);
 
 			// allow the platform to perform any actions after building a target (seems almost like this should be done in UBT)
 			Platform.Platforms[TargetPlatform].PostBuildTarget(this, string.IsNullOrEmpty(ProjectName) ? TargetName : ProjectName, UprojectPath, Config);
@@ -1150,7 +1150,7 @@ namespace AutomationTool
 		/// <param name="InDeleteBuildProducts">if specified, determines if the build products will be deleted before building. If not specified -clean parameter will be used,</param>
 		/// <param name="InUpdateVersionFiles">True if the version files are to be updated </param>
 		/// <param name="InForceNoXGE">If true will force XGE off</param>
-		public void Build(BuildAgenda Agenda, bool? InDeleteBuildProducts = null, bool InUpdateVersionFiles = true, bool InForceNoXGE = false, bool InForceNonUnity = false, bool InForceUnity = false)
+		public void Build(BuildAgenda Agenda, bool? InDeleteBuildProducts = null, bool InUpdateVersionFiles = true, bool InForceNoXGE = false, bool InForceNonUnity = false, bool InForceUnity = false, Dictionary<UnrealBuildTool.UnrealTargetPlatform, Dictionary<string, string>> PlatformEnvVars = null)
 		{
 			if (!CmdEnv.HasCapabilityToCompile)
 			{
diff --git a/Engine/Source/Programs/UnrealBuildTool/Android/UEBuildAndroid.cs b/Engine/Source/Programs/UnrealBuildTool/Android/UEBuildAndroid.cs
--- a/Engine/Source/Programs/UnrealBuildTool/Android/UEBuildAndroid.cs
+++ b/Engine/Source/Programs/UnrealBuildTool/Android/UEBuildAndroid.cs
@@ -6,6 +6,7 @@ using System.Text;
 using System.Diagnostics;
 using System.IO;
 using System.Xml;
+using System.Linq;
 
 namespace UnrealBuildTool
 {
@@ -67,35 +68,69 @@ namespace UnrealBuildTool
 			bool bNeedsNDKPath = string.IsNullOrEmpty(NDKPath);
 			bool bNeedsAndroidHome = string.IsNullOrEmpty(Environment.GetEnvironmentVariable("ANDROID_HOME"));
 			bool bNeedsAntHome = string.IsNullOrEmpty(Environment.GetEnvironmentVariable("ANT_HOME"));
+            if((bNeedsNDKPath || bNeedsAndroidHome || bNeedsAntHome))
+            {
+                var configCacheIni = new ConfigCacheIni("Engine", null);
+                var AndroidEnv = new Dictionary<string, string>();
+
+                Dictionary<string, string> EnvVarNames = new Dictionary<string,string> { 
+                                                         {"ANDROID_HOME", "SDKPath"}, 
+                                                         {"NDKROOT", "NDKPath"}, 
+                                                         {"ANT_HOME", "ANTPath"}
+                                                         };
+
+                string path;
+                foreach(var kvp in EnvVarNames)
+                {
+                    if (configCacheIni.GetPath("/Script/AndroidPlatformEditor.AndroidSDKSettings", kvp.Value, out path))
+                    {
+                        AndroidEnv.Add(kvp.Key, path);
+                    }
+                    else
+                    {
+                        var envValue = Environment.GetEnvironmentVariable(kvp.Key);
+                        if(!String.IsNullOrEmpty(envValue))
+                        {
+                            AndroidEnv.Add(kvp.Key, envValue);
+                        }
+                    }
+                }
 
-			if (Utils.IsRunningOnMono && (bNeedsNDKPath || bNeedsAndroidHome || bNeedsAntHome))
-			{
-				// Try reading env variables we need from .bash_profile
-				string BashProfilePath = Path.Combine(Environment.GetFolderPath(Environment.SpecialFolder.Personal), ".bash_profile");
-				if(File.Exists(BashProfilePath))
-				{
-					string[] BashProfileContents = File.ReadAllLines(BashProfilePath);
-					foreach (string Line in BashProfileContents)
-					{
-						if (bNeedsAndroidHome && Line.StartsWith("export ANDROID_HOME="))
-						{
-							string PathVar = Line.Split('=')[1].Replace("\"", "");
-							Environment.SetEnvironmentVariable("ANDROID_HOME", PathVar);
-						}
-						else if (bNeedsNDKPath && Line.StartsWith("export NDKROOT="))
-						{
-							string PathVar = Line.Split('=')[1].Replace("\"", "");
-							Environment.SetEnvironmentVariable("NDKROOT", PathVar);
-							NDKPath = PathVar;
-						}
-						else if (bNeedsAntHome && Line.StartsWith("export ANT_HOME="))
-						{
-							string PathVar = Line.Split('=')[1].Replace("\"", "");
-							Environment.SetEnvironmentVariable("ANT_HOME", PathVar);
-						}
-					}
-				}
-			}
+                // If we are on Mono and we are still missing a key then go and find it from the .bash_profile
+                if (Utils.IsRunningOnMono && !EnvVarNames.All(s => AndroidEnv.ContainsKey(s.Key)))
+                {
+                    string BashProfilePath = Path.Combine(Environment.GetFolderPath(Environment.SpecialFolder.Personal), ".bash_profile");
+                    if (File.Exists(BashProfilePath))
+                    {
+                        string[] BashProfileContents = File.ReadAllLines(BashProfilePath);
+                        foreach (string Line in BashProfileContents)
+                        {
+                            foreach (var kvp in EnvVarNames)
+                            {
+                                if (AndroidEnv.ContainsKey(kvp.Key))
+                                {
+                                    continue;
+                                }
+
+                                if (Line.StartsWith("export " + kvp.Key + "="))
+                                {
+                                    string PathVar = Line.Split('=')[1].Replace("\"", "");
+                                    AndroidEnv.Add(kvp.Key, PathVar);
+                                }
+                            }
+                        }
+                    }
+                }
+
+                // Set for the process
+                foreach (var kvp in AndroidEnv)
+                {
+                    Environment.SetEnvironmentVariable(kvp.Key, kvp.Value);
+                }
+
+                // See if we have an NDK path now...
+                AndroidEnv.TryGetValue("NDKROOT", out NDKPath);
+            }
 
             // we don't have an NDKROOT specified
             if (String.IsNullOrEmpty(NDKPath))
diff --git a/Engine/Source/Programs/UnrealBuildTool/Configuration/EngineConfiguration.cs b/Engine/Source/Programs/UnrealBuildTool/Configuration/EngineConfiguration.cs
--- a/Engine/Source/Programs/UnrealBuildTool/Configuration/EngineConfiguration.cs
+++ b/Engine/Source/Programs/UnrealBuildTool/Configuration/EngineConfiguration.cs
@@ -268,6 +268,39 @@ namespace UnrealBuildTool
 			return Result;
 		}
 
+        private static bool ExtractPath(string Source, out string Path)
+        {
+            int start = Source.IndexOf('"');
+            int end = Source.LastIndexOf('"');    
+            if(start != 1 && end != -1 && start < end)
+            {
+                ++start;
+                Path = Source.Substring(start, end - start);
+                return true;
+            }
+            else
+            {
+                Path = "";
+            }
+
+            return false;
+        }
+
+        public bool GetPath(string SectionName, string Key, out string Value)
+        {
+            string temp;
+            if(GetString(SectionName, Key, out temp))
+            {
+                return ExtractPath(temp, out Value);
+            }
+            else
+            {
+                Value = "";
+            }
+
+            return false;
+        }
+
 		/// <summary>
 		/// List of actions that can be performed on a single line from ini file
 		/// </summary>
@@ -492,7 +525,10 @@ namespace UnrealBuildTool
 			yield return Path.Combine(EngineDirectory, "Config", "Base" + BaseIniName + ".ini");
 
 			// Game/Config/Default* ini
-			yield return Path.Combine(ProjectDirectory, "Config", "Default" + BaseIniName + ".ini");
+            if(!String.IsNullOrEmpty(ProjectDirectory))
+			{
+                yield return Path.Combine(ProjectDirectory, "Config", "Default" + BaseIniName + ".ini");
+            }
 			
 			string UserSettingsFolder = null; // Match FPlatformProcess::UserSettingsDir()
 			string PersonalFolder = null; // Match FPlatformProcess::UserDir()
@@ -518,7 +554,10 @@ namespace UnrealBuildTool
 			yield return Path.Combine(PersonalFolder, "Unreal Engine", "Engine", "Config", "User" + BaseIniName + ".ini");
 			
 			// Game/Config/User* ini
-			yield return Path.Combine(ProjectDirectory, "Config", "User" + BaseIniName + ".ini");
+            if (!String.IsNullOrEmpty(ProjectDirectory))
+            {
+                yield return Path.Combine(ProjectDirectory, "Config", "User" + BaseIniName + ".ini");
+            }
 		}
 
 		/// <summary>
@@ -532,7 +571,10 @@ namespace UnrealBuildTool
 			yield return Path.Combine(EngineDirectory, "Config", PlatformName, PlatformName + BaseIniName + ".ini");
 
 			// Game/Config/Platform/Platform* ini
-			yield return Path.Combine(ProjectDirectory, "Config", PlatformName, PlatformName + BaseIniName + ".ini");
+            if (!String.IsNullOrEmpty(ProjectDirectory))
+            {
+                yield return Path.Combine(ProjectDirectory, "Config", PlatformName, PlatformName + BaseIniName + ".ini");
+            }
 		}		
 
 		/// <summary>
diff --git a/Engine/Source/Programs/UnrealBuildTool/Configuration/UEBuildPlatform.cs b/Engine/Source/Programs/UnrealBuildTool/Configuration/UEBuildPlatform.cs
--- a/Engine/Source/Programs/UnrealBuildTool/Configuration/UEBuildPlatform.cs
+++ b/Engine/Source/Programs/UnrealBuildTool/Configuration/UEBuildPlatform.cs
@@ -42,11 +42,11 @@ namespace UnrealBuildTool
 		void RegisterBuildPlatform();
 
 		/**
-		 * Attempt to set up AutoSDK for this paltform
+		 * Attempt to set up AutoSDK for this platform
 		 */
 		void ManageAndValidateSDK();
 
-		/**
+        /**
 		 * Retrieve the CPPTargetPlatform for the given UnrealTargetPlatform
 		 *
 		 * @param InUnrealTargetPlatform The UnrealTargetPlatform being build
diff --git a/Engine/Source/Runtime/Core/Private/Mac/MacPlatformMisc.cpp b/Engine/Source/Runtime/Core/Private/Mac/MacPlatformMisc.cpp
--- a/Engine/Source/Runtime/Core/Private/Mac/MacPlatformMisc.cpp
+++ b/Engine/Source/Runtime/Core/Private/Mac/MacPlatformMisc.cpp
@@ -409,6 +409,21 @@ void FMacPlatformMisc::GetEnvironmentVariable(const TCHAR* InVariableName, TCHAR
 	}
 }
 
+void FMacPlatformMisc::SetEnvironmentVar(const TCHAR* InVariableName, const TCHAR* Value)
+{
+	FString VariableName = InVariableName;
+	VariableName.ReplaceInline(TEXT("-"), TEXT("_"));
+	if (Value == NULL || Value[0] == TEXT('\0'))
+	{
+		unsetenv(TCHAR_TO_ANSI(*VariableName));
+	}
+	else
+	{
+		setenv(TCHAR_TO_ANSI(*VariableName), TCHAR_TO_ANSI(Value), 1);
+	}
+}
+
+
 TArray<uint8> FMacPlatformMisc::GetMacAddress()
 {
 	TArray<uint8> Result;
diff --git a/Engine/Source/Runtime/Core/Public/Mac/MacPlatformMisc.h b/Engine/Source/Runtime/Core/Public/Mac/MacPlatformMisc.h
--- a/Engine/Source/Runtime/Core/Public/Mac/MacPlatformMisc.h
+++ b/Engine/Source/Runtime/Core/Public/Mac/MacPlatformMisc.h
@@ -20,7 +20,7 @@ struct CORE_API FMacPlatformMisc : public FGenericPlatformMisc
 	static void PlatformPostInit(bool ShowSplashScreen = false);
 	static class GenericApplication* CreateApplication();
 	static void GetEnvironmentVariable(const TCHAR* VariableName, TCHAR* Result, int32 ResultLength);
-
+	static void SetEnvironmentVar(const TCHAR* VariableName, const TCHAR* Value);
 	static const TCHAR* GetPathVarDelimiter()
 	{
 		return TEXT(":");
