Merging from John Barrett:

Adjusted AutomationTool to support running commandlets as part of game tests.

diff --git a/Engine/Source/Programs/AutomationTool/Scripts/RunProjectCommand.Automation.cs b/Engine/Source/Programs/AutomationTool/Scripts/RunProjectCommand.Automation.cs
--- a/Engine/Source/Programs/AutomationTool/Scripts/RunProjectCommand.Automation.cs
+++ b/Engine/Source/Programs/AutomationTool/Scripts/RunProjectCommand.Automation.cs
@@ -188,6 +188,8 @@ public partial class Project : CommandUtils
 		if (Params.Unattended)
 		{
 			string LookFor = "Bringing up level for play took";
+			bool bCommandlet = false;
+
 			if (Params.RunAutomationTest != "")
 			{
 				LookFor = "Automation Test Succeeded";
@@ -200,6 +202,13 @@ public partial class Project : CommandUtils
 			{
 				LookFor = "Asset discovery search completed in";
 			}
+			// If running a commandlet, just detect a normal exit
+			else if (ClientCmdLine.IndexOf("-run=", StringComparison.InvariantCultureIgnoreCase) >= 0)
+			{
+				LookFor = "Game engine shut down";
+				bCommandlet = true;
+			}
+
 			{
 
 				string AllClientOutput = "";
@@ -271,6 +280,38 @@ public partial class Project : CommandUtils
 								}
 							}
 						}
+						// Detect commandlet failure
+						else if (bCommandlet)
+						{
+							const string ResultLog = "Commandlet->Main return this error code: ";
+
+							int ResultStart = AllClientOutput.LastIndexOf(ResultLog);
+							int ResultValIdx = ResultStart + ResultLog.Length;
+
+							if (ResultStart >= 0 && ResultValIdx < AllClientOutput.Length &&
+								AllClientOutput.Substring(ResultValIdx, 1) == "1")
+							{
+								// Parse the full commandlet warning/error summary
+								string FullSummary = "";
+								int SummaryStart = AllClientOutput.LastIndexOf("Warning/Error Summary");
+								
+								if (SummaryStart >= 0 && SummaryStart < ResultStart)
+								{
+									FullSummary = AllClientOutput.Substring(SummaryStart, ResultStart - SummaryStart);
+								}
+
+
+								if (FullSummary.Length > 0)
+								{
+									Log(System.Diagnostics.TraceEventType.Error, "Commandlet failed, summary:" + Environment.NewLine +
+																					FullSummary);
+								}
+								else
+								{
+									Log(System.Diagnostics.TraceEventType.Error, "Commandlet failed.");
+								}
+							}
+						}
 					}
 					else if (bClientExited && (DateTime.UtcNow - ExitTime).TotalSeconds > 30)
 					{
