Make sure NotForLicensees files don't stage

diff --git a/Engine/Source/Programs/AutomationTool/DeploymentContext.cs b/Engine/Source/Programs/AutomationTool/DeploymentContext.cs
--- a/Engine/Source/Programs/AutomationTool/DeploymentContext.cs
+++ b/Engine/Source/Programs/AutomationTool/DeploymentContext.cs
@@ -346,6 +346,12 @@ public class DeploymentContext //: ProjectParams
 				{
 					continue;
 				}
+
+                if (FileToCopy.Contains("NotForLicensees"))
+                {
+                    continue;
+                }
+
 				bool OtherPlatform = false;
 				foreach (UnrealTargetPlatform Plat in Enum.GetValues(typeof(UnrealTargetPlatform)))
 				{
