StageFiles should only block NotForLicensees directories if specifically told to

Merging CL 2413745

//depot/UE4-UT/...

to //depot/UE4/...

diff --git a/Engine/Source/Programs/AutomationTool/DeploymentContext.cs b/Engine/Source/Programs/AutomationTool/DeploymentContext.cs
--- a/Engine/Source/Programs/AutomationTool/DeploymentContext.cs
+++ b/Engine/Source/Programs/AutomationTool/DeploymentContext.cs
@@ -317,7 +317,7 @@ public class DeploymentContext //: ProjectParams
 		ProjectArgForCommandLines = ProjectArgForCommandLines.Replace("\\", "/");
 	}
 
-	public int StageFiles(StagedFileType FileType, string InPath, string Wildcard = "*", bool bRecursive = true, string[] ExcludeWildcard = null, string NewPath = null, bool bAllowNone = false, bool bRemap = true, string NewName = null)
+    public int StageFiles(StagedFileType FileType, string InPath, string Wildcard = "*", bool bRecursive = true, string[] ExcludeWildcard = null, string NewPath = null, bool bAllowNone = false, bool bRemap = true, string NewName = null, bool bAllowNotForLicenseesFiles = true)
 	{
 		int FilesAdded = 0;
 		// make sure any ..'s are removed
@@ -347,7 +347,7 @@ public class DeploymentContext //: ProjectParams
 					continue;
 				}
 
-                if (FileToCopy.Contains("NotForLicensees"))
+                if (!bAllowNotForLicenseesFiles && FileToCopy.Contains("NotForLicensees"))
                 {
                     continue;
                 }
