build script, comment out depended on by to save some ec resources until real fix can be figured out

diff --git a/Engine/Source/Programs/AutomationTool/Scripts/GUBP.Automation.cs b/Engine/Source/Programs/AutomationTool/Scripts/GUBP.Automation.cs
--- a/Engine/Source/Programs/AutomationTool/Scripts/GUBP.Automation.cs
+++ b/Engine/Source/Programs/AutomationTool/Scripts/GUBP.Automation.cs
@@ -6517,10 +6517,10 @@ public class GUBP : BuildCommand
             {
                 ECProps.Add(string.Format("SortKey/{0}={1}", NodePair.Key, NodePair.Value));
             }
-			foreach (var NodePair in FullNodeDependedOnBy)
-			{
-				ECProps.Add(string.Format("DependedOnBy/{0}={1}", NodePair.Key, NodePair.Value));
-			}
+			//foreach (var NodePair in FullNodeDependedOnBy)
+			//{
+			//	ECProps.Add(string.Format("DependedOnBy/{0}={1}", NodePair.Key, NodePair.Value));
+			//}
 			foreach (var NodePair in FullNodeDependentPromotions)
 			{
 				ECProps.Add(string.Format("DependentPromotions/{0}={1}", NodePair.Key, NodePair.Value));
