Deployment change to filter out NoRedist if desired

Merging

//depot/UE4-UT/...

to //depot/UE4/...

diff --git a/Engine/Source/Programs/AutomationTool/DeploymentContext.cs b/Engine/Source/Programs/AutomationTool/DeploymentContext.cs
--- a/Engine/Source/Programs/AutomationTool/DeploymentContext.cs
+++ b/Engine/Source/Programs/AutomationTool/DeploymentContext.cs
@@ -347,7 +347,7 @@ public class DeploymentContext //: ProjectParams
 					continue;
 				}
 
-                if (!bAllowNotForLicenseesFiles && FileToCopy.Contains("NotForLicensees"))
+                if (!bAllowNotForLicenseesFiles && (FileToCopy.Contains("NotForLicensees") || FileToCopy.Contains("NoRedist"))
                 {
                     continue;
                 }
