Stage non-monolithic support for non-Development builds

diff --git a/Engine/Source/Programs/AutomationTool/Linux/LinuxPlatform.Automation.cs b/Engine/Source/Programs/AutomationTool/Linux/LinuxPlatform.Automation.cs
--- a/Engine/Source/Programs/AutomationTool/Linux/LinuxPlatform.Automation.cs
+++ b/Engine/Source/Programs/AutomationTool/Linux/LinuxPlatform.Automation.cs
@@ -60,7 +60,7 @@ public abstract class BaseLinuxPlatform : Platform
         {
             if (SC.DedicatedServer)
             {
-                SC.StageFiles(StagedFileType.NonUFS, CombinePaths(SC.LocalRoot, "Engine/Binaries", SC.PlatformDir), "UE4Server");
+                SC.StageFiles(StagedFileType.NonUFS, CombinePaths(SC.LocalRoot, "Engine/Binaries", SC.PlatformDir), "UE4Server*");
                 SC.StageFiles(StagedFileType.NonUFS, CombinePaths(SC.LocalRoot, "Engine/Binaries", SC.PlatformDir), "libUE4Server-*.so");
                 SC.StageFiles(StagedFileType.NonUFS, CombinePaths(SC.LocalRoot, "Engine/Plugins"), "libUE4Server-*.so", true, null, null, true);
                 SC.StageFiles(StagedFileType.NonUFS, CombinePaths(SC.ProjectRoot, "Binaries", SC.PlatformDir), "libUE4Server-*.so");
@@ -68,7 +68,7 @@ public abstract class BaseLinuxPlatform : Platform
             }
             else
             {
-                SC.StageFiles(StagedFileType.NonUFS, CombinePaths(SC.LocalRoot, "Engine/Binaries", SC.PlatformDir), "UE4");
+                SC.StageFiles(StagedFileType.NonUFS, CombinePaths(SC.LocalRoot, "Engine/Binaries", SC.PlatformDir), "UE4*");
                 SC.StageFiles(StagedFileType.NonUFS, CombinePaths(SC.LocalRoot, "Engine/Binaries", SC.PlatformDir), "libUE4-*.so");
                 SC.StageFiles(StagedFileType.NonUFS, CombinePaths(SC.LocalRoot, "Engine/Plugins"), "libUE4-*.so", true, null, null, true);
                 SC.StageFiles(StagedFileType.NonUFS, CombinePaths(SC.ProjectRoot, "Binaries", SC.PlatformDir), "libUE4-*.so");
diff --git a/Engine/Source/Programs/AutomationTool/Win/WinPlatform.Automation.cs b/Engine/Source/Programs/AutomationTool/Win/WinPlatform.Automation.cs
--- a/Engine/Source/Programs/AutomationTool/Win/WinPlatform.Automation.cs
+++ b/Engine/Source/Programs/AutomationTool/Win/WinPlatform.Automation.cs
@@ -98,7 +98,7 @@ public abstract class BaseWinPlatform : Platform
 
             if (SC.DedicatedServer)
             {
-                StageExecutable("exe", SC, CommandUtils.CombinePaths(SC.LocalRoot, "Engine/Binaries", SC.PlatformDir), "UE4Server.");
+                StageExecutable("exe", SC, CommandUtils.CombinePaths(SC.LocalRoot, "Engine/Binaries", SC.PlatformDir), "UE4Server*.");
                 StageExecutable("dll", SC, CommandUtils.CombinePaths(SC.LocalRoot, "Engine/Binaries", SC.PlatformDir), "UE4Server-*.");
                 StageExecutable("dll", SC, CommandUtils.CombinePaths(SC.LocalRoot, "Engine/Plugins"), "UE4Server-*.", true);
                 StageExecutable("dll", SC, CommandUtils.CombinePaths(SC.ProjectRoot, "Binaries", SC.PlatformDir), "UE4Server-*.");
@@ -106,7 +106,7 @@ public abstract class BaseWinPlatform : Platform
             }
             else
             {
-                StageExecutable("exe", SC, CommandUtils.CombinePaths(SC.LocalRoot, "Engine/Binaries", SC.PlatformDir), "UE4.");
+                StageExecutable("exe", SC, CommandUtils.CombinePaths(SC.LocalRoot, "Engine/Binaries", SC.PlatformDir), "UE4*.");
                 StageExecutable("dll", SC, CommandUtils.CombinePaths(SC.LocalRoot, "Engine/Binaries", SC.PlatformDir), "UE4-*.");
                 StageExecutable("dll", SC, CommandUtils.CombinePaths(SC.LocalRoot, "Engine/Plugins"), "UE4-*.", true);
                 StageExecutable("dll", SC, CommandUtils.CombinePaths(SC.ProjectRoot, "Binaries", SC.PlatformDir), "UE4-*.");
