build script, leave off some dashboard props

diff --git a/Engine/Source/Programs/AutomationTool/Scripts/GUBP.Automation.cs b/Engine/Source/Programs/AutomationTool/Scripts/GUBP.Automation.cs
--- a/Engine/Source/Programs/AutomationTool/Scripts/GUBP.Automation.cs
+++ b/Engine/Source/Programs/AutomationTool/Scripts/GUBP.Automation.cs
@@ -6509,10 +6509,10 @@ public class GUBP : BuildCommand
             {
                 ECProps.Add(string.Format("AllNodes/{0}={1}", NodePair.Key, NodePair.Value));
             }
-            foreach (var NodePair in FullNodeDirectDependencies)
-            {
-                ECProps.Add(string.Format("DirectDependencies/{0}={1}", NodePair.Key, NodePair.Value));
-            }
+			//foreach (var NodePair in FullNodeDirectDependencies)
+			//{
+			//	ECProps.Add(string.Format("DirectDependencies/{0}={1}", NodePair.Key, NodePair.Value));
+			//}
             foreach (var NodePair in FullNodeListSortKey)
             {
                 ECProps.Add(string.Format("SortKey/{0}={1}", NodePair.Key, NodePair.Value));
@@ -6521,10 +6521,10 @@ public class GUBP : BuildCommand
 			//{
 			//	ECProps.Add(string.Format("DependedOnBy/{0}={1}", NodePair.Key, NodePair.Value));
 			//}
-			foreach (var NodePair in FullNodeDependentPromotions)
-			{
-				ECProps.Add(string.Format("DependentPromotions/{0}={1}", NodePair.Key, NodePair.Value));
-			}
+			//foreach (var NodePair in FullNodeDependentPromotions)
+			//{
+			//	ECProps.Add(string.Format("DependentPromotions/{0}={1}", NodePair.Key, NodePair.Value));
+			//}
 			foreach (var Node in SeparatePromotables)
 			{
 				ECProps.Add(string.Format("PossiblePromotables/{0}={1}", Node, ""));
