Fix GamePlatformMonolithicsNode if GUBP so that for rocket builds we sqaush all but shipping/dev builds for html5.

diff --git a/Engine/Source/Programs/AutomationTool/Scripts/GUBP.Automation.cs b/Engine/Source/Programs/AutomationTool/Scripts/GUBP.Automation.cs
--- a/Engine/Source/Programs/AutomationTool/Scripts/GUBP.Automation.cs
+++ b/Engine/Source/Programs/AutomationTool/Scripts/GUBP.Automation.cs
@@ -1456,6 +1456,10 @@ public class GUBP : BuildCommand
                         {
                             continue;
                         }
+                        if (TargetPlatform == UnrealTargetPlatform.HTML5 && Config != UnrealTargetConfiguration.Shipping && Config != UnrealTargetConfiguration.Development)
+                        {
+                            continue;
+                        }
                         if (TargetPlatform == UnrealTargetPlatform.Linux && Config != UnrealTargetConfiguration.Shipping && Config != UnrealTargetConfiguration.Development)
                         {
                             continue;
@@ -1545,6 +1549,10 @@ public class GUBP : BuildCommand
                                     {
                                         continue;
                                     }
+                                    if (TargetPlatform == UnrealTargetPlatform.HTML5 && Config != UnrealTargetConfiguration.Shipping && Config != UnrealTargetConfiguration.Development)
+                                    {
+                                        continue;
+                                    }
                                     if (TargetPlatform == UnrealTargetPlatform.Linux && Config != UnrealTargetConfiguration.Shipping && Config != UnrealTargetConfiguration.Development)
                                     {
                                         continue;
