Merging stage non-monolithic adjustment

//depot/UE4-UT/...

to //depot/UE4/...

diff --git a/Engine/Source/Programs/AutomationTool/Win/WinPlatform.Automation.cs b/Engine/Source/Programs/AutomationTool/Win/WinPlatform.Automation.cs
--- a/Engine/Source/Programs/AutomationTool/Win/WinPlatform.Automation.cs
+++ b/Engine/Source/Programs/AutomationTool/Win/WinPlatform.Automation.cs
@@ -98,7 +98,21 @@ public abstract class BaseWinPlatform : Platform
 
             if (SC.DedicatedServer)
             {
-                StageExecutable("exe", SC, CommandUtils.CombinePaths(SC.LocalRoot, "Engine/Binaries", SC.PlatformDir), "UE4Server*.");
+                if (SC.StageTargetConfigurations.Any(x => x != UnrealTargetConfiguration.Development))
+                {
+                    StageExecutable("exe", SC, CommandUtils.CombinePaths(SC.LocalRoot, "Engine/Binaries", SC.PlatformDir), "UE4Server.");
+                }
+
+                if (SC.StageTargetConfigurations.Any(x => x != UnrealTargetConfiguration.Test))
+                {
+                    StageExecutable("exe", SC, CommandUtils.CombinePaths(SC.LocalRoot, "Engine/Binaries", SC.PlatformDir), "UE4Server*-Test.");
+                }
+
+                if (SC.StageTargetConfigurations.Any(x => x != UnrealTargetConfiguration.Shipping))
+                {
+                    StageExecutable("exe", SC, CommandUtils.CombinePaths(SC.LocalRoot, "Engine/Binaries", SC.PlatformDir), "UE4Server*-Shipping.");
+                }
+
                 StageExecutable("dll", SC, CommandUtils.CombinePaths(SC.LocalRoot, "Engine/Binaries", SC.PlatformDir), "UE4Server-*.");
                 StageExecutable("dll", SC, CommandUtils.CombinePaths(SC.LocalRoot, "Engine/Plugins"), "UE4Server-*.", true);
                 StageExecutable("dll", SC, CommandUtils.CombinePaths(SC.ProjectRoot, "Binaries", SC.PlatformDir), "UE4Server-*.");
@@ -106,7 +120,21 @@ public abstract class BaseWinPlatform : Platform
             }
             else
             {
-                StageExecutable("exe", SC, CommandUtils.CombinePaths(SC.LocalRoot, "Engine/Binaries", SC.PlatformDir), "UE4*.");
+                if (SC.StageTargetConfigurations.Any(x => x != UnrealTargetConfiguration.Development))
+                {
+                    StageExecutable("exe", SC, CommandUtils.CombinePaths(SC.LocalRoot, "Engine/Binaries", SC.PlatformDir), "UE4.");
+                }
+
+                if (SC.StageTargetConfigurations.Any(x => x != UnrealTargetConfiguration.Test))
+                {
+                    StageExecutable("exe", SC, CommandUtils.CombinePaths(SC.LocalRoot, "Engine/Binaries", SC.PlatformDir), "UE4*-Test.");
+                }
+
+                if (SC.StageTargetConfigurations.Any(x => x != UnrealTargetConfiguration.Shipping))
+                {
+                    StageExecutable("exe", SC, CommandUtils.CombinePaths(SC.LocalRoot, "Engine/Binaries", SC.PlatformDir), "UE4*-Shipping.");
+                }
+
                 StageExecutable("dll", SC, CommandUtils.CombinePaths(SC.LocalRoot, "Engine/Binaries", SC.PlatformDir), "UE4-*.");
                 StageExecutable("dll", SC, CommandUtils.CombinePaths(SC.LocalRoot, "Engine/Plugins"), "UE4-*.", true);
                 StageExecutable("dll", SC, CommandUtils.CombinePaths(SC.ProjectRoot, "Binaries", SC.PlatformDir), "UE4-*.");
