build script, add agent override param

diff --git a/Engine/Source/Programs/AutomationTool/Scripts/GUBP.Automation.cs b/Engine/Source/Programs/AutomationTool/Scripts/GUBP.Automation.cs
--- a/Engine/Source/Programs/AutomationTool/Scripts/GUBP.Automation.cs
+++ b/Engine/Source/Programs/AutomationTool/Scripts/GUBP.Automation.cs
@@ -4027,6 +4027,10 @@ public class GUBP : BuildCommand
             LastAgentGroup = GUBPNodes[NodeToDo].AgentSharingGroup;
 
             string Agent = GUBPNodes[NodeToDo].ECAgentString();
+			if(ParseParamValue("AgentOverride") != "" && !GUBPNodes[NodeToDo].GetFullName().Contains("Mac"))
+			{
+				Agent = ParseParamValue("AgentOverride");
+			}
             if (Agent != "")
             {
                 Agent = "[" + Agent + "]";
@@ -4940,7 +4944,7 @@ public class GUBP : BuildCommand
         bool bChanges = ParseParam("Changes") || ParseParam("AllChanges");
         bool bHistory = ParseParam("History") || bChanges;
         bool bListOnly = ParseParam("ListOnly") || bHistory;
-        bool bSkipTriggers = ParseParam("SkipTriggers");
+        bool bSkipTriggers = ParseParam("SkipTriggers");		
         bFake = ParseParam("fake");
         bool bFakeEC = ParseParam("FakeEC");
         TimeIndex = ParseParamInt("TimeIndex", 0);
