build script, actually override the agent string

diff --git a/Engine/Source/Programs/AutomationTool/Scripts/GUBP.Automation.cs b/Engine/Source/Programs/AutomationTool/Scripts/GUBP.Automation.cs
--- a/Engine/Source/Programs/AutomationTool/Scripts/GUBP.Automation.cs
+++ b/Engine/Source/Programs/AutomationTool/Scripts/GUBP.Automation.cs
@@ -4773,7 +4773,12 @@ public class GUBP : BuildCommand
         }
         if (!OnlyLateUpdates)
         {
-            ECProps.Add(string.Format("AgentRequirementString/{0}={1}", NodeToDo, GUBPNodes[NodeToDo].ECAgentString()));
+			string AgentReq = GUBPNodes[NodeToDo].ECAgentString();
+			if(ParseParamValue("AgentOverride") != "" && !GUBPNodes[NodeToDo].GetFullName().Contains("OnMac"))
+			{
+				AgentReq = ParseParamValue("AgentOverride");
+			}
+            ECProps.Add(string.Format("AgentRequirementString/{0}={1}", NodeToDo, AgentReq));
             ECProps.Add(string.Format("RequiredMemory/{0}={1}", NodeToDo, GUBPNodes[NodeToDo].AgentMemoryRequirement(this)));
             ECProps.Add(string.Format("Timeouts/{0}={1}", NodeToDo, GUBPNodes[NodeToDo].TimeoutInMinutes()));
             ECProps.Add(string.Format("JobStepPath/{0}={1}", NodeToDo, GetJobStepPath(NodeToDo)));
