fix for not generating the plist when changing content only projects
UE-8293
#ios

diff --git a/Engine/Source/Programs/AutomationTool/IOS/IOSPlatform.Automation.cs b/Engine/Source/Programs/AutomationTool/IOS/IOSPlatform.Automation.cs
--- a/Engine/Source/Programs/AutomationTool/IOS/IOSPlatform.Automation.cs
+++ b/Engine/Source/Programs/AutomationTool/IOS/IOSPlatform.Automation.cs
@@ -570,7 +570,7 @@ public class IOSPlatform : Platform
 			{
 				string SourcePath = CombinePaths((SC.IsCodeBasedProject ? SC.ProjectRoot : SC.LocalRoot + "\\Engine"), "Intermediate", "IOS");
 				string TargetPListFile = Path.Combine(SourcePath, (SC.IsCodeBasedProject ? SC.ShortProjectName : "UE4Game") + "-Info.plist");
-				if (!File.Exists(TargetPListFile))
+//				if (!File.Exists(TargetPListFile))
 				{
 					// ensure the plist, entitlements, and provision files are properly copied
 					UnrealBuildTool.IOS.UEDeployIOS.GeneratePList((SC.IsCodeBasedProject ? SC.ProjectRoot : SC.LocalRoot + "\\Engine"), !SC.IsCodeBasedProject, (SC.IsCodeBasedProject ? SC.ShortProjectName : "UE4Game"), SC.ShortProjectName, SC.LocalRoot + "\\Engine", "");
