fix for missing App directory when generating the plist on mac.
#ios

diff --git a/Engine/Source/Programs/AutomationTool/IOS/IOSPlatform.Automation.cs b/Engine/Source/Programs/AutomationTool/IOS/IOSPlatform.Automation.cs
--- a/Engine/Source/Programs/AutomationTool/IOS/IOSPlatform.Automation.cs
+++ b/Engine/Source/Programs/AutomationTool/IOS/IOSPlatform.Automation.cs
@@ -573,7 +573,7 @@ public class IOSPlatform : Platform
 //				if (!File.Exists(TargetPListFile))
 				{
 					// ensure the plist, entitlements, and provision files are properly copied
-					UnrealBuildTool.IOS.UEDeployIOS.GeneratePList((SC.IsCodeBasedProject ? SC.ProjectRoot : SC.LocalRoot + "\\Engine"), !SC.IsCodeBasedProject, (SC.IsCodeBasedProject ? SC.ShortProjectName : "UE4Game"), SC.ShortProjectName, SC.LocalRoot + "\\Engine", "");
+					UnrealBuildTool.IOS.UEDeployIOS.GeneratePList((SC.IsCodeBasedProject ? SC.ProjectRoot : SC.LocalRoot + "\\Engine"), !SC.IsCodeBasedProject, (SC.IsCodeBasedProject ? SC.ShortProjectName : "UE4Game"), SC.ShortProjectName, SC.LocalRoot + "\\Engine", SC.ProjectRoot + "\\Binaries\\IOS\\Payload\\" + (SC.IsCodeBasedProject ? SC.ShortProjectName : "UE4Game") + ".app");
 				}
 
 				SC.StageFiles(StagedFileType.NonUFS, SourcePath, Path.GetFileName(TargetPListFile), false, null, "", false, false, "Info.plist");
