build script, re-enable dependent promotions and dependedonby

diff --git a/Engine/Source/Programs/AutomationTool/Scripts/GUBP.Automation.cs b/Engine/Source/Programs/AutomationTool/Scripts/GUBP.Automation.cs
--- a/Engine/Source/Programs/AutomationTool/Scripts/GUBP.Automation.cs
+++ b/Engine/Source/Programs/AutomationTool/Scripts/GUBP.Automation.cs
@@ -6534,14 +6534,14 @@ public class GUBP : BuildCommand
             {
                 ECProps.Add(string.Format("SortKey/{0}={1}", NodePair.Key, NodePair.Value));
             }
-			//foreach (var NodePair in FullNodeDependedOnBy)
-			//{
-			//	ECProps.Add(string.Format("DependedOnBy/{0}={1}", NodePair.Key, NodePair.Value));
-			//}
-			//foreach (var NodePair in FullNodeDependentPromotions)
-			//{
-			//	ECProps.Add(string.Format("DependentPromotions/{0}={1}", NodePair.Key, NodePair.Value));
-			//}
+			foreach (var NodePair in FullNodeDependedOnBy)
+			{
+				ECProps.Add(string.Format("DependedOnBy/{0}={1}", NodePair.Key, NodePair.Value));
+			}
+			foreach (var NodePair in FullNodeDependentPromotions)
+			{
+				ECProps.Add(string.Format("DependentPromotions/{0}={1}", NodePair.Key, NodePair.Value));
+			}
 			foreach (var Node in SeparatePromotables)
 			{
 				ECProps.Add(string.Format("PossiblePromotables/{0}={1}", Node, ""));
