Merging staging adjustments for non-monolithic builds

//depot/UE4-UT/...

to //depot/UE4/...

diff --git a/Engine/Source/Programs/AutomationTool/Win/WinPlatform.Automation.cs b/Engine/Source/Programs/AutomationTool/Win/WinPlatform.Automation.cs
--- a/Engine/Source/Programs/AutomationTool/Win/WinPlatform.Automation.cs
+++ b/Engine/Source/Programs/AutomationTool/Win/WinPlatform.Automation.cs
@@ -98,17 +98,17 @@ public abstract class BaseWinPlatform : Platform
 
             if (SC.DedicatedServer)
             {
-                if (SC.StageTargetConfigurations.Any(x => x != UnrealTargetConfiguration.Development))
+                if (SC.StageTargetConfigurations.Contains(UnrealTargetConfiguration.Development))
                 {
                     StageExecutable("exe", SC, CommandUtils.CombinePaths(SC.LocalRoot, "Engine/Binaries", SC.PlatformDir), "UE4Server.");
                 }
 
-                if (SC.StageTargetConfigurations.Any(x => x != UnrealTargetConfiguration.Test))
+                if (SC.StageTargetConfigurations.Contains(UnrealTargetConfiguration.Test))
                 {
                     StageExecutable("exe", SC, CommandUtils.CombinePaths(SC.LocalRoot, "Engine/Binaries", SC.PlatformDir), "UE4Server*-Test.");
                 }
 
-                if (SC.StageTargetConfigurations.Any(x => x != UnrealTargetConfiguration.Shipping))
+                if (SC.StageTargetConfigurations.Contains(UnrealTargetConfiguration.Shipping))
                 {
                     StageExecutable("exe", SC, CommandUtils.CombinePaths(SC.LocalRoot, "Engine/Binaries", SC.PlatformDir), "UE4Server*-Shipping.");
                 }
@@ -120,17 +120,17 @@ public abstract class BaseWinPlatform : Platform
             }
             else
             {
-                if (SC.StageTargetConfigurations.Any(x => x != UnrealTargetConfiguration.Development))
+                if (SC.StageTargetConfigurations.Contains(UnrealTargetConfiguration.Development))
                 {
                     StageExecutable("exe", SC, CommandUtils.CombinePaths(SC.LocalRoot, "Engine/Binaries", SC.PlatformDir), "UE4.");
                 }
 
-                if (SC.StageTargetConfigurations.Any(x => x != UnrealTargetConfiguration.Test))
+                if (SC.StageTargetConfigurations.Contains(UnrealTargetConfiguration.Test))
                 {
                     StageExecutable("exe", SC, CommandUtils.CombinePaths(SC.LocalRoot, "Engine/Binaries", SC.PlatformDir), "UE4*-Test.");
                 }
 
-                if (SC.StageTargetConfigurations.Any(x => x != UnrealTargetConfiguration.Shipping))
+                if (SC.StageTargetConfigurations.Contains(UnrealTargetConfiguration.Shipping))
                 {
                     StageExecutable("exe", SC, CommandUtils.CombinePaths(SC.LocalRoot, "Engine/Binaries", SC.PlatformDir), "UE4*-Shipping.");
                 }
