build script, add cis quantum branch hacker

diff --git a/Engine/Source/Programs/AutomationTool/Scripts/GUBP.Automation.cs b/Engine/Source/Programs/AutomationTool/Scripts/GUBP.Automation.cs
--- a/Engine/Source/Programs/AutomationTool/Scripts/GUBP.Automation.cs
+++ b/Engine/Source/Programs/AutomationTool/Scripts/GUBP.Automation.cs
@@ -204,6 +204,7 @@ public class GUBP : BuildCommand
             public List<string> ExcludeNodes = new List<string>();
 			public bool bNoAutomatedTesting = false;
 			public bool bNoDocumentation = false;
+			public int QuantumOverride = 0;
         }
         public virtual void ModifyOptions(GUBP bp, ref BranchOptions Options, string Branch)
         {
@@ -3760,7 +3761,12 @@ public class GUBP : BuildCommand
         int Quantum = GUBPNodes[NodeToDo].DependentCISFrequencyQuantumShift();
         if (Quantum > 0)
         {
-            int Minutes = 20 * (1 << Quantum);
+			int TimeQuantum = 20;
+			if(BranchOptions.QuantumOverride != 0)
+			{
+				TimeQuantum = BranchOptions.QuantumOverride;
+			}
+            int Minutes = TimeQuantum * (1 << Quantum);
             if (Minutes < 60)
             {
                 FrequencyString = string.Format(" ({0}m)", Minutes);
@@ -5083,9 +5089,12 @@ public class GUBP : BuildCommand
                 int Minutes = int.Parse(Parts[1]);
 
                 int DeltaMinutes = NowMinutes - Minutes;
-
-                const int TimeQuantum = 20;
-
+				
+                int TimeQuantum = 20;
+				if(BranchOptions.QuantumOverride != 0)
+				{
+					TimeQuantum = BranchOptions.QuantumOverride;
+				}
                 int NewIndex = Index + 1;
 
                 if (DeltaMinutes > TimeQuantum * 2)
