Clean up Android install batch file
#android
#ue4

diff --git a/Engine/Source/Programs/AutomationTool/Android/AndroidPlatform.Automation.cs b/Engine/Source/Programs/AutomationTool/Android/AndroidPlatform.Automation.cs
--- a/Engine/Source/Programs/AutomationTool/Android/AndroidPlatform.Automation.cs
+++ b/Engine/Source/Programs/AutomationTool/Android/AndroidPlatform.Automation.cs
@@ -315,14 +315,14 @@ public class AndroidPlatform : Platform
 						"@echo Uninstalling existing application. Failures here can almost always be ignored.",
 						"%ADB% %DEVICE% uninstall " + PackageName,
 						"@echo.",
-						"@echo Installing existing application. Failures here indicate a problem with the device \\(connection or storage permissions\\) and are fatal.",
+						"@echo Installing existing application. Failures here indicate a problem with the device (connection or storage permissions) and are fatal.",
 						"%ADB% %DEVICE% install " + Path.GetFileName(ApkName),
 						"@if \"%ERRORLEVEL%\" NEQ \"0\" goto Error",
 						"%ADB% %DEVICE% shell rm -r %STORAGE%/" + Params.ShortProjectName,
 						"%ADB% %DEVICE% shell rm -r %STORAGE%/UE4Game/UE4CommandLine.txt", // we need to delete the commandline in UE4Game or it will mess up loading
 						"%ADB% %DEVICE% shell rm -r %STORAGE%/obb/" + PackageName,
 						bPackageDataInsideApk ? "" : "@echo.",
-						bPackageDataInsideApk ? "" : "@echo Installing new data. Failures here indicate storage problems \\(missing SD card or bad permissions\\) and are fatal.",
+						bPackageDataInsideApk ? "" : "@echo Installing new data. Failures here indicate storage problems (missing SD card or bad permissions) and are fatal.",
 						bPackageDataInsideApk ? "" : "%ADB% %DEVICE% push " + Path.GetFileName(ObbName) + " %STORAGE%/" + DeviceObbName,
 						bPackageDataInsideApk ? "" : "if \"%ERRORLEVEL%\" NEQ \"0\" goto Error",
 						"@echo.",
