fix for path separators on Mac
UE-8395
#ios

diff --git a/Engine/Source/Programs/AutomationTool/IOS/IOSPlatform.Automation.cs b/Engine/Source/Programs/AutomationTool/IOS/IOSPlatform.Automation.cs
--- a/Engine/Source/Programs/AutomationTool/IOS/IOSPlatform.Automation.cs
+++ b/Engine/Source/Programs/AutomationTool/IOS/IOSPlatform.Automation.cs
@@ -568,12 +568,12 @@ public class IOSPlatform : Platform
 			// copy the plist (only if code signing, as it's protected by the code sign blob in the executable and can't be modified independently)
 			if (GetCodeSignDesirability(Params))
 			{
-				string SourcePath = CombinePaths((SC.IsCodeBasedProject ? SC.ProjectRoot : SC.LocalRoot + "\\Engine"), "Intermediate", "IOS");
+				string SourcePath = CombinePaths((SC.IsCodeBasedProject ? SC.ProjectRoot : SC.LocalRoot + "/Engine"), "Intermediate", "IOS");
 				string TargetPListFile = Path.Combine(SourcePath, (SC.IsCodeBasedProject ? SC.ShortProjectName : "UE4Game") + "-Info.plist");
 //				if (!File.Exists(TargetPListFile))
 				{
 					// ensure the plist, entitlements, and provision files are properly copied
-					UnrealBuildTool.IOS.UEDeployIOS.GeneratePList((SC.IsCodeBasedProject ? SC.ProjectRoot : SC.LocalRoot + "\\Engine"), !SC.IsCodeBasedProject, (SC.IsCodeBasedProject ? SC.ShortProjectName : "UE4Game"), SC.ShortProjectName, SC.LocalRoot + "\\Engine", SC.ProjectRoot + "\\Binaries\\IOS\\Payload\\" + (SC.IsCodeBasedProject ? SC.ShortProjectName : "UE4Game") + ".app");
+					UnrealBuildTool.IOS.UEDeployIOS.GeneratePList((SC.IsCodeBasedProject ? SC.ProjectRoot : SC.LocalRoot + "/Engine"), !SC.IsCodeBasedProject, (SC.IsCodeBasedProject ? SC.ShortProjectName : "UE4Game"), SC.ShortProjectName, SC.LocalRoot + "/Engine", (SC.IsCodeBasedProject ? SC.ProjectRoot : SC.LocalRoot + "/Engine") + "/Binaries/IOS/Payload/" + (SC.IsCodeBasedProject ? SC.ShortProjectName : "UE4Game") + ".app");
 				}
 
 				SC.StageFiles(StagedFileType.NonUFS, SourcePath, Path.GetFileName(TargetPListFile), false, null, "", false, false, "Info.plist");
