Merging staging option to include directories that include non-platform files

//depot/UE4-UT/...

to //depot/UE4/...

diff --git a/Engine/Source/Programs/AutomationTool/DeploymentContext.cs b/Engine/Source/Programs/AutomationTool/DeploymentContext.cs
--- a/Engine/Source/Programs/AutomationTool/DeploymentContext.cs
+++ b/Engine/Source/Programs/AutomationTool/DeploymentContext.cs
@@ -317,7 +317,7 @@ public class DeploymentContext //: ProjectParams
 		ProjectArgForCommandLines = ProjectArgForCommandLines.Replace("\\", "/");
 	}
 
-    public int StageFiles(StagedFileType FileType, string InPath, string Wildcard = "*", bool bRecursive = true, string[] ExcludeWildcard = null, string NewPath = null, bool bAllowNone = false, bool bRemap = true, string NewName = null, bool bAllowNotForLicenseesFiles = true)
+    public int StageFiles(StagedFileType FileType, string InPath, string Wildcard = "*", bool bRecursive = true, string[] ExcludeWildcard = null, string NewPath = null, bool bAllowNone = false, bool bRemap = true, string NewName = null, bool bAllowNotForLicenseesFiles = true, bool bStripFilesForOtherPlatforms = true)
 	{
 		int FilesAdded = 0;
 		// make sure any ..'s are removed
@@ -335,7 +335,7 @@ public class DeploymentContext //: ProjectParams
 					var Remove = CommandUtils.FindFiles(Excl, bRecursive, InPath);
 					foreach (var File in Remove)
 					{
-						Exclude.Add(CommandUtils.CombinePaths(File));
+                        Exclude.Add(CommandUtils.CombinePaths(File));
 					}
 				}
 			}
@@ -346,40 +346,43 @@ public class DeploymentContext //: ProjectParams
 				{
 					continue;
 				}
-
+                
                 if (!bAllowNotForLicenseesFiles && (FileToCopy.Contains("NotForLicensees") || FileToCopy.Contains("NoRedist")))
                 {
                     continue;
                 }
 
-				bool OtherPlatform = false;
-				foreach (UnrealTargetPlatform Plat in Enum.GetValues(typeof(UnrealTargetPlatform)))
-				{
-					if (Plat != StageTargetPlatform.PlatformType && Plat != UnrealTargetPlatform.Unknown)
-					{
-                        var Search = FileToCopy;
-                        if (Search.StartsWith(LocalRoot, StringComparison.InvariantCultureIgnoreCase))
+                if (bStripFilesForOtherPlatforms)
+                {
+                    bool OtherPlatform = false;
+                    foreach (UnrealTargetPlatform Plat in Enum.GetValues(typeof(UnrealTargetPlatform)))
+                    {
+                        if (Plat != StageTargetPlatform.PlatformType && Plat != UnrealTargetPlatform.Unknown)
                         {
-                            if (LocalRoot.EndsWith("\\") || LocalRoot.EndsWith("/"))
+                            var Search = FileToCopy;
+                            if (Search.StartsWith(LocalRoot, StringComparison.InvariantCultureIgnoreCase))
                             {
-                                Search = Search.Substring(LocalRoot.Length - 1);
+                                if (LocalRoot.EndsWith("\\") || LocalRoot.EndsWith("/"))
+                                {
+                                    Search = Search.Substring(LocalRoot.Length - 1);
+                                }
+                                else
+                                {
+                                    Search = Search.Substring(LocalRoot.Length);
+                                }
                             }
-                            else
+                            if (Search.IndexOf(CommandUtils.CombinePaths("/" + Plat.ToString() + "/"), 0, StringComparison.InvariantCultureIgnoreCase) >= 0)
                             {
-                                Search = Search.Substring(LocalRoot.Length);
+                                OtherPlatform = true;
+                                break;
                             }
                         }
-                        if (Search.IndexOf(CommandUtils.CombinePaths("/" + Plat.ToString() + "/"), 0, StringComparison.InvariantCultureIgnoreCase) >= 0)
-						{
-							OtherPlatform = true;
-							break;
-						}
-					}
-				}
-				if (OtherPlatform)
-				{
-					continue;
-				}
+                    }
+                    if (OtherPlatform)
+                    {
+                        continue;
+                    }
+                }
 
 				string Dest;
 				if (!FileToCopy.StartsWith(InPath))
