BPT UAT updates

diff --git a/Engine/Source/Programs/AutomationTool/MCPPublic.cs b/Engine/Source/Programs/AutomationTool/MCPPublic.cs
--- a/Engine/Source/Programs/AutomationTool/MCPPublic.cs
+++ b/Engine/Source/Programs/AutomationTool/MCPPublic.cs
@@ -292,6 +292,10 @@ namespace EpicGames.MCP.Automation
             /// <summary>
             /// Matches the corresponding BuildPatchTool command line argument.
             /// </summary>
+            public string FileAttributeList;
+            /// <summary>
+            /// Matches the corresponding BuildPatchTool command line argument.
+            /// </summary>
             public string AppLaunchCmd;
             /// <summary>
             /// Matches the corresponding BuildPatchTool command line argument.
