UE-8433

Added an error code for missing ue4game binaries when packaging a non code project.
Added a check in main frame actions to determine whether we display a dialog with the error message in it.

Fix for incorrect IOS stub path being logged
Fix for incorrect Android apk path being generated for non ue4game renames in ::GetFinalApkName

diff --git a/Engine/Source/Editor/MainFrame/Private/Frame/MainFrameActions.cpp b/Engine/Source/Editor/MainFrame/Private/Frame/MainFrameActions.cpp
--- a/Engine/Source/Editor/MainFrame/Private/Frame/MainFrameActions.cpp
+++ b/Engine/Source/Editor/MainFrame/Private/Frame/MainFrameActions.cpp
@@ -525,7 +525,7 @@ void FMainFrameActionCallbacks::PackageProject( const FName InPlatformInfoName )
 	}
 
 	// only build if the project has code that might need to be built
-	if ((bProjectHasCode || bContentProjectCanBeBuilt) && FSourceCodeNavigation::IsCompilerAvailable())
+	if (bProjectHasCode && FSourceCodeNavigation::IsCompilerAvailable())
 	{
 		OptionalParams += TEXT(" -build");
 	}
@@ -1140,6 +1140,7 @@ void FMainFrameActionCallbacks::HandleUatProcessCanceled( TWeakPtr<SNotification
 }
 
 
+DECLARE_CYCLE_STAT(TEXT("Requesting FMainFrameActionCallbacks::HandleUatProcessCompleted message dialog to present the error message"), STAT_FMainFrameActionCallbacks_HandleUatProcessCompleted_DialogMessage, STATGROUP_TaskGraphTasks);
 void FMainFrameActionCallbacks::HandleUatProcessCompleted( int32 ReturnCode, TWeakPtr<SNotificationItem> NotificationItemPtr, FText PlatformDisplayName, FText TaskName, EventData Event )
 {
 	FFormatNamedArguments Arguments;
@@ -1172,6 +1173,18 @@ void FMainFrameActionCallbacks::HandleUatProcessCompleted( int32 ReturnCode, TWe
 		ParamArray.Add(FAnalyticsEventAttribute(TEXT("Time"), FPlatformTime::Seconds() - Event.StartTime));
 		FEditorAnalytics::ReportEvent(Event.EventName + TEXT(".Failed"), PlatformDisplayName.ToString(), Event.bProjectHasCode, ReturnCode, ParamArray);
 
+		// Present a message dialog if we want the error message to be prominent.
+		if (FEditorAnalytics::ShouldElevateMessageThroughDialog(ReturnCode))
+		{
+			FSimpleDelegateGraphTask::CreateAndDispatchWhenReady(
+				FSimpleDelegateGraphTask::FDelegate::CreateLambda([=](){
+					FMessageDialog::Open(EAppMsgType::Ok, FText::FromString(FEditorAnalytics::TranslateErrorCode(ReturnCode)));
+				}),
+				GET_STATID(STAT_FMainFrameActionCallbacks_HandleUatProcessCompleted_DialogMessage),
+				nullptr,
+				ENamedThreads::GameThread
+			);
+		}
 //		FMessageLog("PackagingResults").Info(FText::Format(LOCTEXT("UatProcessFailedMessageLog", "{TaskName} for {Platform} failed"), Arguments));
 	}
 }
diff --git a/Engine/Source/Editor/UnrealEd/Private/Analytics/EditorAnalytics.cpp b/Engine/Source/Editor/UnrealEd/Private/Analytics/EditorAnalytics.cpp
--- a/Engine/Source/Editor/UnrealEd/Private/Analytics/EditorAnalytics.cpp
+++ b/Engine/Source/Editor/UnrealEd/Private/Analytics/EditorAnalytics.cpp
@@ -9,6 +9,8 @@
 #include "EditorAnalytics.h"
 #include "GeneralProjectSettings.h"
 
+#define LOCTEXT_NAMESPACE "EditorAnalytics"
+
 void FEditorAnalytics::ReportBuildRequirementsFailure(FString EventName, FString PlatformName, bool bHasCode, int32 Requirements)
 {
 	TArray<FAnalyticsEventAttribute> ParamArray;
@@ -74,7 +76,8 @@ void FEditorAnalytics::ReportEvent(FString EventName, FString PlatformName, bool
 		ParamArray.Add(FAnalyticsEventAttribute(TEXT("Platform"), PlatformName));
 		ParamArray.Add(FAnalyticsEventAttribute(TEXT("ProjectType"), bHasCode ? TEXT("C++ Code") : TEXT("Content Only")));
 		ParamArray.Add(FAnalyticsEventAttribute(TEXT("ErrorCode"), ErrorCode));
-		ParamArray.Add(FAnalyticsEventAttribute(TEXT("ErrorName"), TranslateErrorCode(ErrorCode)));
+		const FString ErrorMessage = TranslateErrorCode(ErrorCode);
+		ParamArray.Add(FAnalyticsEventAttribute(TEXT("ErrorName"), ErrorMessage));
 		ParamArray.Append(ExtraParams);
 
 		FEngineAnalytics::GetProvider().RecordEvent( EventName, ParamArray );
@@ -163,7 +166,15 @@ FString FEditorAnalytics::TranslateErrorCode(int32 ErrorCode)
 		return TEXT("No APK suitable for architecture found");
 	case EAnalyticsErrorCodes::FailedToDeleteStagingDirectory :
 		return TEXT("Failed to delete staging directory.  This could be because something is currently using the staging directory (ps4/xbox/etc)");
+	case EAnalyticsErrorCodes::MissingExecutable:
+		return LOCTEXT("UATErrorMissingExecutable", "Missing UE4Game binary.\nYou may have to build the UE4 project with your IDE. Alternatively, build using UnrealBuildTool with the commandline:\nUE4Game <Platform> <Configuration>").ToString();
 	}
 	return TEXT("Unknown Error");
 }
 
+bool FEditorAnalytics::ShouldElevateMessageThroughDialog(const int32 ErrorCode)
+{
+	return (EAnalyticsErrorCodes::Type)ErrorCode == EAnalyticsErrorCodes::MissingExecutable;
+}
+
+#undef LOCTEXT_NAMESPACE
\ No newline at end of file
diff --git a/Engine/Source/Editor/UnrealEd/Public/EditorAnalytics.h b/Engine/Source/Editor/UnrealEd/Public/EditorAnalytics.h
--- a/Engine/Source/Editor/UnrealEd/Public/EditorAnalytics.h
+++ b/Engine/Source/Editor/UnrealEd/Public/EditorAnalytics.h
@@ -56,6 +56,7 @@ namespace EAnalyticsErrorCodes
 		LauncherFailed = 100,
 		UATLaunchFailure = 101,
 		FailedToDeleteStagingDirectory = 102,
+		MissingExecutable = 103,
 	};
 };
 
@@ -101,5 +102,8 @@ public:
 public:
 
 	/* Translates the error code to something human readable */
-	static FString TranslateErrorCode(int32 ErrorCode);
+	UNREALED_API static FString TranslateErrorCode(int32 ErrorCode);
+
+	/* Determine whether the error code should cause the reporter to display the message in a dialog. */
+	UNREALED_API static bool ShouldElevateMessageThroughDialog(const int32 ErrorCode);
 };
diff --git a/Engine/Source/Programs/AutomationTool/Android/AndroidPlatform.Automation.cs b/Engine/Source/Programs/AutomationTool/Android/AndroidPlatform.Automation.cs
--- a/Engine/Source/Programs/AutomationTool/Android/AndroidPlatform.Automation.cs
+++ b/Engine/Source/Programs/AutomationTool/Android/AndroidPlatform.Automation.cs
@@ -48,7 +48,7 @@ public class AndroidPlatform : Platform
 			else
 			{
 				// if we want to use UE4 directly then use it from the engine directory not project directory
-				ApkName = ApkName.Replace(ProjectDir, Path.Combine(CmdEnv.LocalRoot, "Engine"));
+				ApkName = ApkName.Replace(ProjectDir, Path.Combine(CmdEnv.LocalRoot, "Engine/Binaries/Android"));
 			}
 		}
 
@@ -234,6 +234,17 @@ public class AndroidPlatform : Platform
 			foreach (string GPUArchitecture in GPUArchitectures)
 			{
 				string ApkName = GetFinalApkName(Params, SC.StageExecutables[0], true, bMakeSeparateApks ? Architecture : "", bMakeSeparateApks ? GPUArchitecture : "");
+				if (!SC.IsCodeBasedProject)
+				{
+					string UE4GameApkName = GetFinalApkName(Params, SC.StageExecutables[0], false, bMakeSeparateApks ? Architecture : "", bMakeSeparateApks ? GPUArchitecture : "");
+					if (FileExists_NoExceptions(UE4GameApkName) == false)
+					{
+						Log("Failed to find game apk " + UE4GameApkName);
+						AutomationTool.ErrorReporter.Error("Stage Failed.", (int)AutomationTool.ErrorCodes.Error_MissingExecutable);
+						throw new AutomationException("Could not find apk {0}. You may need to build the UE4 project with your target configuration and platform.", UE4GameApkName);
+					}
+				}
+				
 				string BatchName = GetFinalBatchName(ApkName, Params, bMakeSeparateApks ? Architecture : "", bMakeSeparateApks ? GPUArchitecture : "");
 
 				if (!Params.Prebuilt)
diff --git a/Engine/Source/Programs/AutomationTool/HTML5/HTML5Platform.Automation.cs b/Engine/Source/Programs/AutomationTool/HTML5/HTML5Platform.Automation.cs
--- a/Engine/Source/Programs/AutomationTool/HTML5/HTML5Platform.Automation.cs
+++ b/Engine/Source/Programs/AutomationTool/HTML5/HTML5Platform.Automation.cs
@@ -44,7 +44,17 @@ public class HTML5Platform : Platform
             {
                 GameExe += "-HTML5-" + Params.ClientConfigsToBuild[0].ToString();
             }
-            GameExe += ".js";
+			GameExe += ".js";
+			
+			// ensure the ue4game binary exists, if applicable
+			string FullGameExePath = Path.Combine(Path.GetDirectoryName(Params.ProjectGameExeFilename), GameExe);
+			if (!SC.IsCodeBasedProject && !FileExists_NoExceptions(FullGameExePath))
+			{
+				Log("Failed to find game application " + FullGameExePath);
+				AutomationTool.ErrorReporter.Error("Stage Failed.", (int)AutomationTool.ErrorCodes.Error_MissingExecutable);
+				throw new AutomationException("Could not find application {0}. You may need to build the UE4 project with your target configuration and platform.", FullGameExePath);
+			}
+
             if (Path.Combine(Path.GetDirectoryName(Params.ProjectGameExeFilename), GameExe) != Path.Combine(PackagePath, GameExe))
             {
                 File.Copy(Path.Combine(Path.GetDirectoryName(Params.ProjectGameExeFilename), GameExe), Path.Combine(PackagePath, GameExe), true);
diff --git a/Engine/Source/Programs/AutomationTool/IOS/IOSPlatform.Automation.cs b/Engine/Source/Programs/AutomationTool/IOS/IOSPlatform.Automation.cs
--- a/Engine/Source/Programs/AutomationTool/IOS/IOSPlatform.Automation.cs
+++ b/Engine/Source/Programs/AutomationTool/IOS/IOSPlatform.Automation.cs
@@ -139,6 +139,16 @@ public class IOSPlatform : Platform
 		{
 			// copy in all of the artwork and plist
 			var DeployHandler = UEBuildDeploy.GetBuildDeploy(UnrealTargetPlatform.IOS);
+
+			string FullExePath = CombinePaths(Path.GetDirectoryName (Params.ProjectGameExeFilename), SC.StageExecutables [0]);
+			// ensure the ue4game binary exists, if applicable
+			if (!SC.IsCodeBasedProject && !FileExists_NoExceptions(FullExePath))
+			{
+				Log("Failed to find game binary " + FullExePath);
+				AutomationTool.ErrorReporter.Error("Stage Failed.", (int)AutomationTool.ErrorCodes.Error_MissingExecutable);
+				throw new AutomationException("Could not find binary {0}. You may need to build the UE4 project with your target configuration and platform.", FullExePath);
+			}
+
 			DeployHandler.PrepForUATPackageOrDeploy(Params.ShortProjectName,
 				Path.GetDirectoryName(Params.RawProjectPath),
 				CombinePaths(Path.GetDirectoryName(Params.ProjectGameExeFilename), SC.StageExecutables[0]),
@@ -200,7 +210,7 @@ public class IOSPlatform : Platform
 		{
 			var TargetConfiguration = SC.StageTargetConfigurations[0];
 			var ProjectIPA = MakeIPAFileName(TargetConfiguration, Params);
-			var ProjectStub = Path.GetFullPath(Params.ProjectGameExeFilename) + "/" + Params.ShortProjectName + Path.GetExtension(Params.ProjectGameExeFilename);
+			var ProjectStub = Path.GetFullPath(Params.ProjectGameExeFilename);
 
 			// package a .ipa from the now staged directory
 			var IPPExe = CombinePaths(CmdEnv.LocalRoot, "Engine/Binaries/DotNET/IOS/IPhonePackager.exe");
diff --git a/Engine/Source/Programs/AutomationTool/Linux/LinuxPlatform.Automation.cs b/Engine/Source/Programs/AutomationTool/Linux/LinuxPlatform.Automation.cs
--- a/Engine/Source/Programs/AutomationTool/Linux/LinuxPlatform.Automation.cs
+++ b/Engine/Source/Programs/AutomationTool/Linux/LinuxPlatform.Automation.cs
@@ -99,6 +99,14 @@ public abstract class BaseLinuxPlatform : Platform
 				    // Move the executable for content-only projects into the project directory, using the project name, so it can figure out the UProject to look for and is consistent with code projects.
 				    if (!Params.IsCodeBasedProject && Exe == Exes[0])
 				    {
+						// ensure the ue4game binary exists, if applicable
+						if (!SC.IsCodeBasedProject && !FileExists_NoExceptions(Params.ProjectGameExeFilename))
+						{
+							Log("Failed to find game binary " + Params.ProjectGameExeFilename);
+							AutomationTool.ErrorReporter.Error("Stage Failed.", (int)AutomationTool.ErrorCodes.Error_MissingExecutable);
+							throw new AutomationException("Could not find game binary {0}. You may need to build the UE4 project with your target configuration and platform.", Params.ProjectGameExeFilename);
+						}
+
 					    SC.StageFiles(StagedFileType.NonUFS, CombinePaths(SC.LocalRoot, "Engine/Binaries", SC.PlatformDir), Path.GetFileNameWithoutExtension(Exe), true, null, CommandUtils.CombinePaths(SC.RelativeProjectRootForStage, "Binaries", SC.PlatformDir), false, true, SC.ShortProjectName);
 				    }
 				    else
diff --git a/Engine/Source/Programs/AutomationTool/Mac/MacPlatform.Automation.cs b/Engine/Source/Programs/AutomationTool/Mac/MacPlatform.Automation.cs
--- a/Engine/Source/Programs/AutomationTool/Mac/MacPlatform.Automation.cs
+++ b/Engine/Source/Programs/AutomationTool/Mac/MacPlatform.Automation.cs
@@ -58,6 +58,16 @@ public class MacPlatform : Platform
 			else if (Exe.StartsWith(CombinePaths(SC.RuntimeRootDir, "Engine/Binaries", SC.PlatformDir)))
 			{
 				AppBundlePath = CombinePaths("Engine/Binaries", SC.PlatformDir, Path.GetFileNameWithoutExtension(Exe) + ".app");
+
+				string AbsoluteBundlePath = CombinePaths (SC.LocalRoot, AppBundlePath);
+				// ensure the ue4game binary exists, if applicable
+				if (!SC.IsCodeBasedProject && !Directory.Exists(AbsoluteBundlePath))
+				{
+					Log("Failed to find app bundle " + AbsoluteBundlePath);
+					AutomationTool.ErrorReporter.Error("Stage Failed.", (int)AutomationTool.ErrorCodes.Error_MissingExecutable);
+					throw new AutomationException("Could not find app bundle {0}. You may need to build the UE4 project with your target configuration and platform.", AbsoluteBundlePath);
+				}
+
 				StageAppBundle(SC, CombinePaths(SC.LocalRoot, "Engine/Binaries", SC.PlatformDir, Path.GetFileNameWithoutExtension(Exe) + ".app"), AppBundlePath);
 			}
 
diff --git a/Engine/Source/Programs/AutomationTool/Program.cs b/Engine/Source/Programs/AutomationTool/Program.cs
--- a/Engine/Source/Programs/AutomationTool/Program.cs
+++ b/Engine/Source/Programs/AutomationTool/Program.cs
@@ -56,7 +56,8 @@ namespace AutomationTool
         Error_NoApkSuitableForArchitecture = 55,
         Error_LauncherFailed = 100,
 		Error_UATLaunchFailure = 101,
-        Error_FailedToDeleteStagingDirectory = 102,
+		Error_FailedToDeleteStagingDirectory = 102,
+		Error_MissingExecutable = 103,
 	};
 
 	public class ErrorReporter
diff --git a/Engine/Source/Programs/AutomationTool/Win/WinPlatform.Automation.cs b/Engine/Source/Programs/AutomationTool/Win/WinPlatform.Automation.cs
--- a/Engine/Source/Programs/AutomationTool/Win/WinPlatform.Automation.cs
+++ b/Engine/Source/Programs/AutomationTool/Win/WinPlatform.Automation.cs
@@ -152,12 +152,18 @@ public abstract class BaseWinPlatform : Platform
             {
 				if (Exe.StartsWith(CombinePaths(SC.RuntimeProjectRootDir, "Binaries", SC.PlatformDir)))
 				{
+					string BaseExeFileName = Path.GetFileName(Exe);
+					if (!String.IsNullOrEmpty(Params.OverrideMinimumOS) && (Params.OverrideMinimumOS == "WinXP"))
+					{
+						BaseExeFileName = Path.GetFileNameWithoutExtension(Exe) + "_xp" + Path.GetExtension(Exe);
+					}
+
 					// remap the project root.
-					string SourceFile = CombinePaths(SC.ProjectRoot, "Binaries", SC.PlatformDir, Path.GetFileName(Exe));
+					string SourceFile = CombinePaths(SC.ProjectRoot, "Binaries", SC.PlatformDir, BaseExeFileName);
 					StageExecutable("exe", SC, Path.GetDirectoryName(SourceFile), Path.GetFileNameWithoutExtension(SourceFile) + ".", true, null, CommandUtils.CombinePaths(SC.RelativeProjectRootForStage, "Binaries", SC.PlatformDir), false, WorkingFileType);
 					if(Exe == Exes[0] && !Params.NoBootstrapExe)
 					{
-						StageBootstrapExecutable(SC, SourceFile, CombinePaths(SC.RelativeProjectRootForStage, "Binaries", SC.PlatformDir, Path.GetFileName(Exe)), "");
+						StageBootstrapExecutable(SC, SourceFile, CombinePaths(SC.RelativeProjectRootForStage, "Binaries", SC.PlatformDir, BaseExeFileName), "");
 					}
 				}
 				else if (Exe.StartsWith(CombinePaths(SC.RuntimeRootDir, "Engine/Binaries", SC.PlatformDir)))
@@ -171,6 +177,15 @@ public abstract class BaseWinPlatform : Platform
 
 					// keep it in the engine directory.
 					string SourceFile = CombinePaths(SC.LocalRoot, "Engine", "Binaries", SC.PlatformDir, ExeFileName);
+
+					// ensure the ue4game binary exists, if applicable
+					if (!SC.IsCodeBasedProject && !FileExists_NoExceptions(SourceFile))
+					{
+						Log("Failed to find game executable " + SourceFile);
+						AutomationTool.ErrorReporter.Error("Stage Failed.", (int)AutomationTool.ErrorCodes.Error_MissingExecutable);
+						throw new AutomationException("Could not find exe {0}. You may need to build the UE4 project with your target configuration and platform.", SourceFile);
+					}
+
 					StageExecutable("exe", SC, CombinePaths(SC.LocalRoot, "Engine/Binaries", SC.PlatformDir), Path.GetFileNameWithoutExtension(SourceFile) + ".", true, null, null, false, WorkingFileType);
 					if(Exe == Exes[0] && !Params.NoBootstrapExe)
 					{
