HTML5 always uses Pak files when packaging. Pak files are compressed for HTML5.
Fixes UE-6883

diff --git a/Engine/Source/Programs/AutomationTool/HTML5/HTML5Platform.Automation.cs b/Engine/Source/Programs/AutomationTool/HTML5/HTML5Platform.Automation.cs
--- a/Engine/Source/Programs/AutomationTool/HTML5/HTML5Platform.Automation.cs
+++ b/Engine/Source/Programs/AutomationTool/HTML5/HTML5Platform.Automation.cs
@@ -392,9 +392,14 @@ public class HTML5Platform : Platform
 
     public override PakType RequiresPak(ProjectParams Params)
     {
-        return PakType.Never;
+        return PakType.Always;
     }
 
+	public override string GetPlatformPakCommandLine()
+	{
+		return " -compress";
+	}
+
 	public override bool DeployLowerCaseFilenames(bool bUFSFile)
 	{
 		return false;
