Fixed network file server not enumerating the DirectoriesToAlwaysStageAsUFS.
Fixed cook on the fly not deploying DirectoriesToAlwaysStageNonUFS.

diff --git a/Engine/Source/Programs/AutomationTool/Scripts/CopyBuildToStagingDirectory.Automation.cs b/Engine/Source/Programs/AutomationTool/Scripts/CopyBuildToStagingDirectory.Automation.cs
--- a/Engine/Source/Programs/AutomationTool/Scripts/CopyBuildToStagingDirectory.Automation.cs
+++ b/Engine/Source/Programs/AutomationTool/Scripts/CopyBuildToStagingDirectory.Automation.cs
@@ -242,9 +242,13 @@ public partial class Project : CommandUtils
 		}
 		SC.StageFiles(StagedFileType.NonUFS, GetIntermediateCommandlineDir(SC), CommandLineFile, false, null, "", true, false);
 
-		if (!Params.CookOnTheFly && !Params.SkipCookOnTheFly) // only stage the UFS files if we are not using cook on the fly
-		{
-            ConfigCacheIni PlatformGameConfig = new ConfigCacheIni(SC.StageTargetPlatform.PlatformType, "Game", CommandUtils.GetDirectoryName(Params.RawProjectPath));
+        ConfigCacheIni PlatformGameConfig = new ConfigCacheIni(SC.StageTargetPlatform.PlatformType, "Game", CommandUtils.GetDirectoryName(Params.RawProjectPath));
+        var ProjectContentRoot = CombinePaths(SC.ProjectRoot, "Content");
+        var StageContentRoot = CombinePaths(SC.RelativeProjectRootForStage, "Content");
+        
+        if (!Params.CookOnTheFly && !Params.SkipCookOnTheFly) // only stage the UFS files if we are not using cook on the fly
+        {
+
 
             // Initialize internationalization preset.
             string InternationalizationPreset = null;
@@ -297,116 +301,149 @@ public partial class Project : CommandUtils
             // Stage ICU internationalization data from Engine.
             SC.StageFiles(StagedFileType.UFS, CombinePaths(SC.LocalRoot, "Engine", "Content", "Internationalization", InternationalizationPreset), "*", true, null, CombinePaths("Engine", "Content", "Internationalization"), true, !Params.UsePak(SC.StageTargetPlatform));
 
-			// Engine ufs (content)
-			SC.StageFiles(StagedFileType.UFS, CombinePaths(SC.LocalRoot, "Engine/Config"), "*", true, null, null, false, !Params.UsePak(SC.StageTargetPlatform)); // TODO: Exclude localization data generation config files.
+            // Engine ufs (content)
+            SC.StageFiles(StagedFileType.UFS, CombinePaths(SC.LocalRoot, "Engine/Config"), "*", true, null, null, false, !Params.UsePak(SC.StageTargetPlatform)); // TODO: Exclude localization data generation config files.
 
-			// Stat prefix/suffix files (used for FPSChart, etc...)
-			//@TODO: Avoid packaging stat files in shipping builds (only 6 KB, but still more than zero)
-			SC.StageFiles(StagedFileType.UFS, CombinePaths(SC.LocalRoot, "Engine/Content/Stats"), "*", true, null, null, false, !Params.UsePak(SC.StageTargetPlatform));
+            // Stat prefix/suffix files (used for FPSChart, etc...)
+            //@TODO: Avoid packaging stat files in shipping builds (only 6 KB, but still more than zero)
+            SC.StageFiles(StagedFileType.UFS, CombinePaths(SC.LocalRoot, "Engine/Content/Stats"), "*", true, null, null, false, !Params.UsePak(SC.StageTargetPlatform));
 
-			if (Params.bUsesSlate)
-			{
-				if (Params.bUsesSlateEditorStyle)
-				{
-					SC.StageFiles(StagedFileType.UFS, CombinePaths(SC.LocalRoot, "Engine/Content/Editor/Slate"), "*", true, null, null, false, !Params.UsePak(SC.StageTargetPlatform));
-				}
-				SC.StageFiles(StagedFileType.UFS, CombinePaths(SC.LocalRoot, "Engine/Content/Slate"), "*", true, null, null, false, !Params.UsePak(SC.StageTargetPlatform));
-				SC.StageFiles(StagedFileType.UFS, CombinePaths(SC.ProjectRoot, "Content/Slate"), "*", true, new string[] {"*.uasset"}, CombinePaths(SC.RelativeProjectRootForStage, "Content/Slate"), true, !Params.UsePak(SC.StageTargetPlatform));
-			}
+            if (Params.bUsesSlate)
+            {
+                if (Params.bUsesSlateEditorStyle)
+                {
+                    SC.StageFiles(StagedFileType.UFS, CombinePaths(SC.LocalRoot, "Engine/Content/Editor/Slate"), "*", true, null, null, false, !Params.UsePak(SC.StageTargetPlatform));
+                }
+                SC.StageFiles(StagedFileType.UFS, CombinePaths(SC.LocalRoot, "Engine/Content/Slate"), "*", true, null, null, false, !Params.UsePak(SC.StageTargetPlatform));
+                SC.StageFiles(StagedFileType.UFS, CombinePaths(SC.ProjectRoot, "Content/Slate"), "*", true, new string[] { "*.uasset" }, CombinePaths(SC.RelativeProjectRootForStage, "Content/Slate"), true, !Params.UsePak(SC.StageTargetPlatform));
+            }
             foreach (string Culture in CulturesToStage)
             {
                 StageLocalizationDataForCulture(SC, Culture, CombinePaths(SC.LocalRoot, "Engine/Content/Localization/Engine"), null, !Params.UsePak(SC.StageTargetPlatform));
             }
-			SC.StageFiles(StagedFileType.UFS, CombinePaths(SC.LocalRoot, "Engine/Plugins"), "*.uplugin", true, null, null, true, !Params.UsePak(SC.StageTargetPlatform));
+            SC.StageFiles(StagedFileType.UFS, CombinePaths(SC.LocalRoot, "Engine/Plugins"), "*.uplugin", true, null, null, true, !Params.UsePak(SC.StageTargetPlatform));
 
-			// Game ufs (content)
+            // Game ufs (content)
 
-			SC.StageFiles(StagedFileType.UFS, CombinePaths(SC.ProjectRoot), "*.uproject", false, null, CombinePaths(SC.RelativeProjectRootForStage), true, !Params.UsePak(SC.StageTargetPlatform));
-			SC.StageFiles(StagedFileType.UFS, CombinePaths(SC.ProjectRoot, "Config"), "*", true, null, CombinePaths(SC.RelativeProjectRootForStage, "Config"), true, !Params.UsePak(SC.StageTargetPlatform)); // TODO: Exclude localization data generation config files.
-			SC.StageFiles(StagedFileType.UFS, CombinePaths(SC.ProjectRoot, "Plugins"), "*.uplugin", true, null, null, true, !Params.UsePak(SC.StageTargetPlatform));
-			foreach (string Culture in CulturesToStage)
-			{
+            SC.StageFiles(StagedFileType.UFS, CombinePaths(SC.ProjectRoot), "*.uproject", false, null, CombinePaths(SC.RelativeProjectRootForStage), true, !Params.UsePak(SC.StageTargetPlatform));
+            SC.StageFiles(StagedFileType.UFS, CombinePaths(SC.ProjectRoot, "Config"), "*", true, null, CombinePaths(SC.RelativeProjectRootForStage, "Config"), true, !Params.UsePak(SC.StageTargetPlatform)); // TODO: Exclude localization data generation config files.
+            SC.StageFiles(StagedFileType.UFS, CombinePaths(SC.ProjectRoot, "Plugins"), "*.uplugin", true, null, null, true, !Params.UsePak(SC.StageTargetPlatform));
+            foreach (string Culture in CulturesToStage)
+            {
                 StageLocalizationDataForCulture(SC, Culture, CombinePaths(SC.ProjectRoot, "Content/Localization/Game"), CombinePaths(SC.RelativeProjectRootForStage, "Content/Localization/Game"), !Params.UsePak(SC.StageTargetPlatform));
-			}
+            }
 
-			// Stage any additional UFS and NonUFS paths specified in the project ini files; these dirs are relative to the game content directory
+            // Stage any additional UFS and NonUFS paths specified in the project ini files; these dirs are relative to the game content directory
             if (PlatformGameConfig != null)
-			{
-				var ProjectContentRoot = CombinePaths(SC.ProjectRoot, "Content");
-				var StageContentRoot = CombinePaths(SC.RelativeProjectRootForStage, "Content");
-				List<string> ExtraUFSDirs;
-				if (PlatformGameConfig.GetArray("/Script/UnrealEd.ProjectPackagingSettings", "DirectoriesToAlwaysStageAsUFS", out ExtraUFSDirs))
-				{
-					// Each string has the format '(Path="TheDirToStage")'
-					foreach (var PathStr in ExtraUFSDirs)
-					{
-						var PathParts = PathStr.Split('"');
-						if (PathParts.Length == 3)
-						{
-							var RelativePath = PathParts[1];
-							SC.StageFiles(StagedFileType.UFS, CombinePaths(ProjectContentRoot, RelativePath), "*", true, null, CombinePaths(StageContentRoot, RelativePath), true, !Params.UsePak(SC.StageTargetPlatform));
-						}
-					}
-				}
+            {
+                List<string> ExtraUFSDirs;
+                if (PlatformGameConfig.GetArray("/Script/UnrealEd.ProjectPackagingSettings", "DirectoriesToAlwaysStageAsUFS", out ExtraUFSDirs))
+                {
+                    // Each string has the format '(Path="TheDirToStage")'
+                    foreach (var PathStr in ExtraUFSDirs)
+                    {
+                        var PathParts = PathStr.Split('"');
+                        if (PathParts.Length == 3)
+                        {
+                            var RelativePath = PathParts[1];
+                            SC.StageFiles(StagedFileType.UFS, CombinePaths(ProjectContentRoot, RelativePath), "*", true, null, CombinePaths(StageContentRoot, RelativePath), true, !Params.UsePak(SC.StageTargetPlatform));
+                        }
+                        else if (PathParts.Length == 1)
+                        {
+                            var RelativePath = PathParts[0];
+                            SC.StageFiles(StagedFileType.UFS, CombinePaths(ProjectContentRoot, RelativePath), "*", true, null, CombinePaths(StageContentRoot, RelativePath), true, !Params.UsePak(SC.StageTargetPlatform));
+                        }
+                    }
+                }
 
-				List<string> ExtraNonUFSDirs;
-				if (PlatformGameConfig.GetArray("/Script/UnrealEd.ProjectPackagingSettings", "DirectoriesToAlwaysStageAsNonUFS", out ExtraNonUFSDirs))
-				{
-					// Each string has the format '(Path="TheDirToStage")'
-					foreach (var PathStr in ExtraNonUFSDirs)
-					{
-						var PathParts = PathStr.Split('"');
-						if (PathParts.Length == 3)
-						{
-							var RelativePath = PathParts[1];
-							SC.StageFiles(StagedFileType.NonUFS, CombinePaths(ProjectContentRoot, RelativePath));
-						}
-					}
-				}
-			}
+                List<string> ExtraNonUFSDirs;
+                if (PlatformGameConfig.GetArray("/Script/UnrealEd.ProjectPackagingSettings", "DirectoriesToAlwaysStageAsNonUFS", out ExtraNonUFSDirs))
+                {
+                    // Each string has the format '(Path="TheDirToStage")'
+                    foreach (var PathStr in ExtraNonUFSDirs)
+                    {
+                        var PathParts = PathStr.Split('"');
+                        if (PathParts.Length == 3)
+                        {
+                            var RelativePath = PathParts[1];
+                            SC.StageFiles(StagedFileType.NonUFS, CombinePaths(ProjectContentRoot, RelativePath));
+                        }
+                        else if (PathParts.Length == 1)
+                        {
+                            var RelativePath = PathParts[0];
+                            SC.StageFiles(StagedFileType.UFS, CombinePaths(ProjectContentRoot, RelativePath), "*", true, null, CombinePaths(StageContentRoot, RelativePath), true, !Params.UsePak(SC.StageTargetPlatform));
+                        }
+                    }
+                }
+            }
 
-			StagedFileType StagedFileTypeForMovies = StagedFileType.NonUFS;
-			if (Params.FileServer)
-			{
-				// UFS is required when using a file server
-				StagedFileTypeForMovies = StagedFileType.UFS;
-			}
+            StagedFileType StagedFileTypeForMovies = StagedFileType.NonUFS;
+            if (Params.FileServer)
+            {
+                // UFS is required when using a file server
+                StagedFileTypeForMovies = StagedFileType.UFS;
+            }
 
-			if (SC.StageTargetPlatform.StageMovies)
-			{
-				SC.StageFiles(StagedFileTypeForMovies, CombinePaths(SC.LocalRoot, "Engine/Content/Movies"), "*", true, null, CombinePaths(SC.RelativeProjectRootForStage, "Engine/Content/Movies"), true, !Params.UsePak(SC.StageTargetPlatform));
-				SC.StageFiles(StagedFileTypeForMovies, CombinePaths(SC.ProjectRoot, "Content/Movies"), "*", true, null, CombinePaths(SC.RelativeProjectRootForStage, "Content/Movies"), true, !Params.UsePak(SC.StageTargetPlatform));
-			}
+            if (SC.StageTargetPlatform.StageMovies)
+            {
+                SC.StageFiles(StagedFileTypeForMovies, CombinePaths(SC.LocalRoot, "Engine/Content/Movies"), "*", true, null, CombinePaths(SC.RelativeProjectRootForStage, "Engine/Content/Movies"), true, !Params.UsePak(SC.StageTargetPlatform));
+                SC.StageFiles(StagedFileTypeForMovies, CombinePaths(SC.ProjectRoot, "Content/Movies"), "*", true, null, CombinePaths(SC.RelativeProjectRootForStage, "Content/Movies"), true, !Params.UsePak(SC.StageTargetPlatform));
+            }
 
-			// eliminate the sand box
-			SC.StageFiles(StagedFileType.UFS, CombinePaths(SC.ProjectRoot, "Saved", "Cooked", SC.CookPlatform), "*", true, null, "", true, !Params.UsePak(SC.StageTargetPlatform));
+            // eliminate the sand box
+            SC.StageFiles(StagedFileType.UFS, CombinePaths(SC.ProjectRoot, "Saved", "Cooked", SC.CookPlatform), "*", true, null, "", true, !Params.UsePak(SC.StageTargetPlatform));
 
-			// CrashReportClient is a standalone slate app that does not look in the generated pak file, so it needs the Content/Slate and Shaders/StandaloneRenderer folders Non-UFS
-			// @todo Make CrashReportClient more portable so we don't have to do this
-			if (SC.bStageCrashReporter && UnrealBuildTool.UnrealBuildTool.PlatformSupportsCrashReporter(SC.StageTargetPlatform.PlatformType) && !SC.DedicatedServer)
-			{
-				//If the .dat file needs to be staged as NonUFS for non-Windows/Linux hosts we need to change the casing as we do with the build properties file above.
-				SC.StageFiles(StagedFileType.NonUFS, CombinePaths(SC.LocalRoot, "Engine/Content/Slate"));
-				SC.StageFiles(StagedFileType.NonUFS, CombinePaths(SC.LocalRoot, "Engine/Shaders/StandaloneRenderer"));
+            // CrashReportClient is a standalone slate app that does not look in the generated pak file, so it needs the Content/Slate and Shaders/StandaloneRenderer folders Non-UFS
+            // @todo Make CrashReportClient more portable so we don't have to do this
+            if (SC.bStageCrashReporter && UnrealBuildTool.UnrealBuildTool.PlatformSupportsCrashReporter(SC.StageTargetPlatform.PlatformType) && !SC.DedicatedServer)
+            {
+                //If the .dat file needs to be staged as NonUFS for non-Windows/Linux hosts we need to change the casing as we do with the build properties file above.
+                SC.StageFiles(StagedFileType.NonUFS, CombinePaths(SC.LocalRoot, "Engine/Content/Slate"));
+                SC.StageFiles(StagedFileType.NonUFS, CombinePaths(SC.LocalRoot, "Engine/Shaders/StandaloneRenderer"));
 
                 SC.StageFiles(StagedFileType.NonUFS, CombinePaths(SC.LocalRoot, "Engine", "Content", "Internationalization", InternationalizationPreset), "*", true, null, CombinePaths("Engine", "Content", "Internationalization"), false, true);
 
-				// Linux platform stages ICU in GetFilesToDeployOrStage(), accounting for the actual architecture
-				if (SC.StageTargetPlatform.PlatformType == UnrealTargetPlatform.Win64 ||
-					SC.StageTargetPlatform.PlatformType == UnrealTargetPlatform.Win32 ||
-					SC.StageTargetPlatform.PlatformType == UnrealTargetPlatform.Mac)
-				{
-					SC.StageFiles(StagedFileType.NonUFS, CombinePaths(SC.LocalRoot, "Engine/Binaries/ThirdParty/ICU"));
-				}
+                // Linux platform stages ICU in GetFilesToDeployOrStage(), accounting for the actual architecture
+                if (SC.StageTargetPlatform.PlatformType == UnrealTargetPlatform.Win64 ||
+                    SC.StageTargetPlatform.PlatformType == UnrealTargetPlatform.Win32 ||
+                    SC.StageTargetPlatform.PlatformType == UnrealTargetPlatform.Mac)
+                {
+                    SC.StageFiles(StagedFileType.NonUFS, CombinePaths(SC.LocalRoot, "Engine/Binaries/ThirdParty/ICU"));
+                }
 
-				// SSL libraries are only available for Win64 builds.
-				// @see FPerforceSourceControlProvider::LoadSSLLibraries
-				if (SC.StageTargetPlatform.PlatformType == UnrealTargetPlatform.Win64)
-				{
-					SC.StageFiles(StagedFileType.NonUFS, CombinePaths(SC.LocalRoot, "Engine/Binaries/ThirdParty/OpenSSL"));
-				}
-			}
-		}
+                // SSL libraries are only available for Win64 builds.
+                // @see FPerforceSourceControlProvider::LoadSSLLibraries
+                if (SC.StageTargetPlatform.PlatformType == UnrealTargetPlatform.Win64)
+                {
+                    SC.StageFiles(StagedFileType.NonUFS, CombinePaths(SC.LocalRoot, "Engine/Binaries/ThirdParty/OpenSSL"));
+                }
+            }
+        }
+        else
+        {
+            if (PlatformGameConfig != null)
+            {
+                List<string> ExtraNonUFSDirs;
+                if (PlatformGameConfig.GetArray("/Script/UnrealEd.ProjectPackagingSettings", "DirectoriesToAlwaysStageAsNonUFS", out ExtraNonUFSDirs))
+                {
+                    // Each string has the format '(Path="TheDirToStage")'
+                    foreach (var PathStr in ExtraNonUFSDirs)
+                    {
+                        var PathParts = PathStr.Split('"');
+                        if (PathParts.Length == 3)
+                        {
+                            var RelativePath = PathParts[1];
+                            SC.StageFiles(StagedFileType.NonUFS, CombinePaths(ProjectContentRoot, RelativePath));
+                        }
+                        else if (PathParts.Length == 1)
+                        {
+                            var RelativePath = PathParts[0];
+                            SC.StageFiles(StagedFileType.UFS, CombinePaths(ProjectContentRoot, RelativePath), "*", true, null, CombinePaths(StageContentRoot, RelativePath), true, !Params.UsePak(SC.StageTargetPlatform));
+                        }
+                    }
+                }
+            }
+        }
 	}
 
 	public static void DumpTargetManifest(Dictionary<string, string> Mapping, string Filename, string StageDir, List<string> CRCFiles)
diff --git a/Engine/Source/Runtime/NetworkFileSystem/Private/NetworkFileServerConnection.cpp b/Engine/Source/Runtime/NetworkFileSystem/Private/NetworkFileServerConnection.cpp
--- a/Engine/Source/Runtime/NetworkFileSystem/Private/NetworkFileServerConnection.cpp
+++ b/Engine/Source/Runtime/NetworkFileSystem/Private/NetworkFileServerConnection.cpp
@@ -770,6 +770,15 @@ bool FNetworkFileServerClientConnection::ProcessGetFileList( FArchive& In, FArch
 		UE_LOG(LogFileServer, Display, TEXT("\t%s"), *(RootDirectories[DumpIdx]));
 	}
 
+	TArray<FString> DirectoriesToAlwaysStageAsUFS;
+	if ( GConfig->GetArray(TEXT("/Script/UnrealEd.ProjectPackagingSettings"), TEXT("DirectoriesToAlwaysStageAsUFS"), DirectoriesToAlwaysStageAsUFS, GGameIni) )
+	{
+		for ( const auto& DirectoryToAlwaysStage : DirectoriesToAlwaysStageAsUFS )
+		{
+			RootDirectories.Add( DirectoryToAlwaysStage );
+		}
+	}
+
 	// list of directories to skip
 	TArray<FString> DirectoriesToSkip;
 	TArray<FString> DirectoriesToNotRecurse;
