Exposed versioned cooking in ufe also made iterative cooking use versioned cooks.
Made the unversionedcook flag be respected when passed into UAT.

diff --git a/Engine/Source/Developer/LauncherServices/Private/Launcher/LauncherWorker.cpp b/Engine/Source/Developer/LauncherServices/Private/Launcher/LauncherWorker.cpp
--- a/Engine/Source/Developer/LauncherServices/Private/Launcher/LauncherWorker.cpp
+++ b/Engine/Source/Developer/LauncherServices/Private/Launcher/LauncherWorker.cpp
@@ -357,7 +357,7 @@ FString FLauncherWorker::CreateUATCommand( const ILauncherProfileRef& InProfile,
 
 			if (InProfile->IsCookingUnversioned())
 			{
-				UATCommand += TEXT(" -Unversioned");
+				UATCommand += TEXT(" -unversionedcookedcontent");
 			}
 
 			FString additionalOptions = InProfile->GetCookOptions();
diff --git a/Engine/Source/Developer/LauncherServices/Private/Profiles/LauncherProfile.h b/Engine/Source/Developer/LauncherServices/Private/Profiles/LauncherProfile.h
--- a/Engine/Source/Developer/LauncherServices/Private/Profiles/LauncherProfile.h
+++ b/Engine/Source/Developer/LauncherServices/Private/Profiles/LauncherProfile.h
@@ -603,7 +603,7 @@ public:
 		CookMode = ELauncherProfileCookModes::OnTheFly;
 		CookOptions = FString();
 		CookIncremental = false;
-		CookUnversioned = false;
+		CookUnversioned = true;
 		CookedCultures.Reset();
 		CookedCultures.Add(I18N.GetCurrentCulture()->GetName());
 		CookedMaps.Reset();
@@ -1009,6 +1009,11 @@ protected:
 			}
 		}
 
+		if ( CookUnversioned && CookIncremental )
+		{
+			ValidationErrors.Add(ELauncherProfileValidationErrors::UnversionedAndIncrimental);
+		}
+
 		// Launch: when launching, all devices that the build is launched on must have content cooked for their platform
 		if ((LaunchMode != ELauncherProfileLaunchModes::DoNotLaunch) && (CookMode != ELauncherProfileCookModes::OnTheFly || CookMode != ELauncherProfileCookModes::OnTheFlyInEditor))
 		{
diff --git a/Engine/Source/Developer/LauncherServices/Public/Interfaces/ILauncherProfile.h b/Engine/Source/Developer/LauncherServices/Public/Interfaces/ILauncherProfile.h
--- a/Engine/Source/Developer/LauncherServices/Public/Interfaces/ILauncherProfile.h
+++ b/Engine/Source/Developer/LauncherServices/Public/Interfaces/ILauncherProfile.h
@@ -156,6 +156,9 @@ namespace ELauncherProfileValidationErrors
 		
 		/** The platform SDK is not installed but is required. */
 		NoPlatformSDKInstalled,
+
+		/** The profile has unversioned and incrimental specified these are not compatible together */
+		UnversionedAndIncrimental,
 	};
 }
 
diff --git a/Engine/Source/Editor/UnrealEd/Private/PlayLevel.cpp b/Engine/Source/Editor/UnrealEd/Private/PlayLevel.cpp
--- a/Engine/Source/Editor/UnrealEd/Private/PlayLevel.cpp
+++ b/Engine/Source/Editor/UnrealEd/Private/PlayLevel.cpp
@@ -1480,6 +1480,7 @@ void UEditorEngine::PlayUsingLauncher()
 			bIncrimentalCooking = false;
 		}
 		LauncherProfile->SetCookMode( CurrentLauncherCookMode );
+		LauncherProfile->SetUnversionedCooking(!bIncrimentalCooking);
 		LauncherProfile->SetIncrementalCooking(bIncrimentalCooking);
 		LauncherProfile->SetDeployedDeviceGroup(DeviceGroup);
 		LauncherProfile->SetIncrementalDeploying(bIncrimentalCooking);
diff --git a/Engine/Source/Programs/AutomationTool/ProjectParams.cs b/Engine/Source/Programs/AutomationTool/ProjectParams.cs
--- a/Engine/Source/Programs/AutomationTool/ProjectParams.cs
+++ b/Engine/Source/Programs/AutomationTool/ProjectParams.cs
@@ -200,6 +200,7 @@ namespace AutomationTool
             this.CreateReleaseVersion = InParams.CreateReleaseVersion;
             this.DLCName = InParams.DLCName;
             this.NewCook = InParams.NewCook;
+            this.OldCook = InParams.OldCook;
             this.AdditionalCookerOptions = InParams.AdditionalCookerOptions;
 			this.ClientCookedTargets = InParams.ClientCookedTargets;
 			this.ServerCookedTargets = InParams.ServerCookedTargets;
@@ -213,6 +214,8 @@ namespace AutomationTool
 			this.Run = InParams.Run;
 			this.Cook = InParams.Cook;
 			this.IterativeCooking = InParams.IterativeCooking;
+            this.CookAll = InParams.CookAll;
+            this.CookMapsOnly = InParams.CookMapsOnly;
 			this.CookFlavor = InParams.CookFlavor;
 			this.SkipCook = InParams.SkipCook;
 			this.SkipCookOnTheFly = InParams.SkipCookOnTheFly;
@@ -226,6 +229,7 @@ namespace AutomationTool
 			this.NoXGE = InParams.NoXGE;
 			this.CookOnTheFly = InParams.CookOnTheFly;
             this.CookOnTheFlyStreaming = InParams.CookOnTheFlyStreaming;
+            this.UnversionedCookedContent = InParams.UnversionedCookedContent;
 			this.FileServer = InParams.FileServer;
 			this.DedicatedServer = InParams.DedicatedServer;
 			this.Client = InParams.Client;
@@ -328,13 +332,17 @@ namespace AutomationTool
             bool? Compressed = null,
             bool? UseDebugParamForEditorExe = null,
             bool? IterativeCooking = null,
+            bool? CookAll = null,
+            bool? CookMapsOnly = null,
             bool? CookOnTheFly = null,
             bool? CookOnTheFlyStreaming = null,
+            bool? UnversionedCookedContent = null,
             string AdditionalCookerOptions = null,
             string BasedOnReleaseVersion = null,
             string CreateReleaseVersion = null,
             string DLCName = null,
             bool? NewCook = null,
+            bool? OldCook = null,
 			bool? CrashReporter = null,
 			bool? DedicatedServer = null,
 			bool? Client = null,
@@ -384,10 +392,6 @@ namespace AutomationTool
 			//
 
 			this.RawProjectPath = RawProjectPath;
-			if (MapsToCook != null)
-			{
-				this.MapsToCook = MapsToCook;
-			}
 			if (DirectoriesToCook != null)
 			{
 				this.DirectoriesToCook = DirectoriesToCook;
@@ -452,6 +456,7 @@ namespace AutomationTool
 			this.Cook = GetParamValueIfNotSpecified(Command, Cook, this.Cook, "cook");
 			this.CookFlavor = ParseParamValueIfNotSpecified(Command, CookFlavor, "cookflavor", String.Empty);
             this.NewCook = GetParamValueIfNotSpecified(Command, NewCook, this.NewCook, "NewCook");
+            this.OldCook = GetParamValueIfNotSpecified(Command, OldCook, this.OldCook, "OldCook");
             this.CreateReleaseVersion = ParseParamValueIfNotSpecified(Command, CreateReleaseVersion, "createreleaseversion", String.Empty);
             this.BasedOnReleaseVersion = ParseParamValueIfNotSpecified(Command, BasedOnReleaseVersion, "basedonreleaseversion", String.Empty);
             this.AdditionalCookerOptions = ParseParamValueIfNotSpecified(Command, AdditionalCookerOptions, "AdditionalCookerOptions", String.Empty);
@@ -477,10 +482,13 @@ namespace AutomationTool
                 this.Cook = false;
             }
             this.CookOnTheFlyStreaming = GetParamValueIfNotSpecified(Command, CookOnTheFlyStreaming, this.CookOnTheFlyStreaming, "cookontheflystreaming");
+            this.UnversionedCookedContent = GetParamValueIfNotSpecified(Command, UnversionedCookedContent, this.UnversionedCookedContent, "UnversionedCookedContent");
             this.Compressed = GetParamValueIfNotSpecified(Command, Compressed, this.Compressed, "compressed");
             this.UseDebugParamForEditorExe = GetParamValueIfNotSpecified(Command, UseDebugParamForEditorExe, this.UseDebugParamForEditorExe, "UseDebugParamForEditorExe");
             this.IterativeCooking = GetParamValueIfNotSpecified(Command, IterativeCooking, this.IterativeCooking, "iterativecooking", "iterate");
 			this.SkipCookOnTheFly = GetParamValueIfNotSpecified(Command, SkipCookOnTheFly, this.SkipCookOnTheFly, "skipcookonthefly");
+            this.CookAll = GetParamValueIfNotSpecified(Command, CookAll, this.CookAll, "CookAll");
+            this.CookMapsOnly = GetParamValueIfNotSpecified(Command, CookMapsOnly, this.CookMapsOnly, "CookMapsOnly");
 			this.FileServer = GetParamValueIfNotSpecified(Command, FileServer, this.FileServer, "fileserver");
 			this.DedicatedServer = GetParamValueIfNotSpecified(Command, DedicatedServer, this.DedicatedServer, "dedicatedserver", "server");
 			this.Client = GetParamValueIfNotSpecified(Command, Client, this.Client, "client");
@@ -603,6 +611,28 @@ namespace AutomationTool
                 this.Port = Port;
             }
 
+            if (MapsToCook == null)
+            {
+                if (Command != null)
+                {
+                    this.MapsToCook = new ParamList<string>();
+
+                    var MapsString = Command.ParseParamValue("MapsToCook");
+                    if (String.IsNullOrEmpty(MapsString) == false)
+                    {
+                        var MapNames = new ParamList<string>(MapsString.Split('+'));
+                        foreach ( var M in MapNames ) 
+                        {
+                            this.MapsToCook.Add( M );
+                        }
+                    }
+                }
+            }
+            else
+            {
+                this.MapsToCook = MapsToCook;
+            }
+
 			if (ServerConfigsToBuild == null)
 			{
 				if (Command != null)
@@ -981,6 +1011,11 @@ namespace AutomationTool
         /// </summary>
         public bool NewCook { private set; get; }
 
+        /// <summary>
+        /// Cook: Use old cooker (temporary cooker options until new cook replaces it)
+        /// </summary>
+        public bool OldCook { private set; get; }
+
         /// <summary>
         /// Cook: Base this cook of a already released version of the cooked data
         /// </summary>
@@ -1011,18 +1046,38 @@ namespace AutomationTool
         /// </summary>
         public bool UseDebugParamForEditorExe;
 
+        /// <summary>
+        /// 
+        /// </summary>
+        public bool UnversionedCookedContent;
+
 		/// <summary>
 		/// Cook: Uses the iterative cooking, command line: -iterativecooking or -iterate
 		/// </summary>
 		[Help( "iterativecooking", "Uses the iterative cooking, command line: -iterativecooking or -iterate" )]
 		public bool IterativeCooking;
 
+        /// <summary>
+        /// Cook: Only cook maps (and referenced content) instead of cooking everything only affects -cookall flag
+        /// </summary>
+        [Help("CookMapsOnly", "Cook only maps this only affects usage of -cookall the flag")]
+        public bool CookMapsOnly;
+
+        /// <summary>
+        /// Cook: Only cook maps (and referenced content) instead of cooking everything only affects cookall flag
+        /// </summary>
+        [Help("CookAll", "Cook all the things in the content directory for this project")]
+        public bool CookAll;
+
+
 		/// <summary>
 		/// Cook: Uses the iterative deploy, command line: -iterativedeploy or -iterate
 		/// </summary>
 		[Help("iterativecooking", "Uses the iterative cooking, command line: -iterativedeploy or -iterate")]
 		public bool IterativeDeploy;
 
+
+
 		#endregion
 
 		#region Stage
@@ -1886,6 +1941,7 @@ namespace AutomationTool
 				CommandUtils.Log("CookFlavor={0}", CookFlavor);
 				CommandUtils.Log("CookOnTheFly={0}", CookOnTheFly);
                 CommandUtils.Log("CookOnTheFlyStreaming={0}", CookOnTheFlyStreaming);
+                CommandUtils.Log("UnversionedCookedContent={0}", UnversionedCookedContent);
                 CommandUtils.Log("CreateReleaseVersion={0}", CreateReleaseVersion);
                 CommandUtils.Log("BasedOnReleaseVersion={0}", BasedOnReleaseVersion);
                 CommandUtils.Log("DLCName={0}", DLCName);
@@ -1898,7 +1954,9 @@ namespace AutomationTool
 				CommandUtils.Log("IsCodeBasedProject={0}", IsCodeBasedProject.ToString());
 				CommandUtils.Log("IsProgramTarget={0}", IsProgramTarget.ToString());
 				CommandUtils.Log("IterativeCooking={0}", IterativeCooking);
-				CommandUtils.Log("IterativeDeploy={0}", IterativeDeploy);
+                CommandUtils.Log("CookAll={0}", CookAll);
+                CommandUtils.Log("CookMapsOnly={0}", CookMapsOnly);
+                CommandUtils.Log("IterativeDeploy={0}", IterativeDeploy);
 				CommandUtils.Log("LogWindow={0}", LogWindow);
 				CommandUtils.Log("Manifests={0}", Manifests);
 				CommandUtils.Log("MapToRun={0}", MapToRun);
diff --git a/Engine/Source/Programs/AutomationTool/Scripts/CookCommand.Automation.cs b/Engine/Source/Programs/AutomationTool/Scripts/CookCommand.Automation.cs
--- a/Engine/Source/Programs/AutomationTool/Scripts/CookCommand.Automation.cs
+++ b/Engine/Source/Programs/AutomationTool/Scripts/CookCommand.Automation.cs
@@ -100,6 +100,14 @@ public partial class Project : CommandUtils
 			if (Params.HasMapsToCook)
 			{
 				Maps = Params.MapsToCook.ToArray();
+                foreach (var M in Maps)
+                {
+                    Log("HasMapsToCook " + M.ToString());
+                }
+                foreach (var M in Params.MapsToCook)
+                {
+                    Log("Params.HasMapsToCook " + M.ToString());
+                }
 			}
 
 			string[] Dirs = null;
@@ -122,7 +130,11 @@ public partial class Project : CommandUtils
 
             try
             {
-                var CommandletParams = "-buildmachine -Unversioned -fileopenlog";
+                var CommandletParams = "-buildmachine -fileopenlog";
+                if (Params.UnversionedCookedContent)
+                {
+                    CommandletParams += " -unversioned";
+                }
                 if (Params.UseDebugParamForEditorExe)
                 {
                     CommandletParams += " -debug";
@@ -135,10 +147,26 @@ public partial class Project : CommandUtils
                 {
                     CommandletParams += " -iterate";
                 }
+                if (Params.CookMapsOnly)
+                {
+                    CommandletParams += " -mapsonly";
+                }
                 if (Params.NewCook)
                 {
                     CommandletParams += " -newcook";
                 }
+                if (Params.OldCook)
+                {
+                    CommandletParams += " -oldcook";
+                }
+                if (Params.CookAll)
+                {
+                    CommandletParams += " -cookall";
+                }
+                if (Params.CookMapsOnly)
+                {
+                    CommandletParams += " -mapsonly";
+                }
                 if (Params.HasCreateReleaseVersion)
                 {
                     CommandletParams += " -createreleaseversion=" + Params.CreateReleaseVersion;
@@ -151,6 +179,7 @@ public partial class Project : CommandUtils
                 {
                     CommandletParams += " -dlcname=" + Params.DLCName;
                 }
+                
                 if (Params.HasAdditionalCookerOptions)
                 {
                     string FormatedAdditionalCookerParams = Params.AdditionalCookerOptions.TrimStart(new char[] { '\"', ' ' }).TrimEnd(new char[] { '\"', ' ' });
