Removed CookedAssetRegistry.json file from the final staged build.
No one wants it.

diff --git a/Engine/Source/Programs/AutomationTool/Scripts/CopyBuildToStagingDirectory.Automation.cs b/Engine/Source/Programs/AutomationTool/Scripts/CopyBuildToStagingDirectory.Automation.cs
--- a/Engine/Source/Programs/AutomationTool/Scripts/CopyBuildToStagingDirectory.Automation.cs
+++ b/Engine/Source/Programs/AutomationTool/Scripts/CopyBuildToStagingDirectory.Automation.cs
@@ -391,7 +391,7 @@ public partial class Project : CommandUtils
             }
 
             // eliminate the sand box
-            SC.StageFiles(StagedFileType.UFS, CombinePaths(SC.ProjectRoot, "Saved", "Cooked", SC.CookPlatform), "*", true, null, "", true, !Params.UsePak(SC.StageTargetPlatform));
+            SC.StageFiles(StagedFileType.UFS, CombinePaths(SC.ProjectRoot, "Saved", "Cooked", SC.CookPlatform), "*", true, new string[] { "*.json" }, "", true, !Params.UsePak(SC.StageTargetPlatform));
 
             // CrashReportClient is a standalone slate app that does not look in the generated pak file, so it needs the Content/Slate and Shaders/StandaloneRenderer folders Non-UFS
             // @todo Make CrashReportClient more portable so we don't have to do this
