UAT: Better support for host platform detection. Removed a workaround which prevented HostPlatform crash when it was uninitialized.

diff --git a/Engine/Source/Programs/AutomationTool/HostPlatform.cs b/Engine/Source/Programs/AutomationTool/HostPlatform.cs
--- a/Engine/Source/Programs/AutomationTool/HostPlatform.cs
+++ b/Engine/Source/Programs/AutomationTool/HostPlatform.cs
@@ -20,7 +20,14 @@ namespace AutomationTool
 		/// </summary>
 		public static HostPlatform Current
 		{
-			get { return RunningPlatform; }
+			get 
+			{ 
+				if (RunningPlatform == null)
+				{
+					throw new AutomationException("UnrealAutomationTool host platform not initialized.");
+				}
+				return RunningPlatform; 
+			}
 		}
 
 		/// <summary>
@@ -32,6 +39,9 @@ namespace AutomationTool
 			switch (Platform)
 			{
 				case PlatformID.Win32NT:
+				case PlatformID.Win32S:
+				case PlatformID.Win32Windows:
+				case PlatformID.WinCE:
 					RunningPlatform = new WindowsHostPlatform();
 					break;
 
@@ -46,6 +56,10 @@ namespace AutomationTool
 					}
 					break;
 
+				case PlatformID.MacOSX:
+					RunningPlatform = new MacHostPlatform();
+					break;
+
 				default:
 					throw new Exception ("Unhandled runtime platform " + Platform);
 			}
diff --git a/Engine/Source/Programs/AutomationTool/ProjectParams.cs b/Engine/Source/Programs/AutomationTool/ProjectParams.cs
--- a/Engine/Source/Programs/AutomationTool/ProjectParams.cs
+++ b/Engine/Source/Programs/AutomationTool/ProjectParams.cs
@@ -424,25 +424,7 @@ namespace AutomationTool
                 this.ClientDependentPlatformMap = ClientDependentPlatformMap;
             }
 
-			UnrealTargetPlatform TargetPlatform;
-			try
-			{
-				TargetPlatform = HostPlatform.Current.HostEditorPlatform;
-			}
-			catch
-			{
-				if (this.ClientDependentPlatformMap != null)
-				{
-					// if a platform map is provided, assume key of first element is the target platform
-					TargetPlatform = this.ClientDependentPlatformMap.Keys.First();
-				}
-				else
-				{
-					// should never get here
-					throw new Exception("A host platform is not defined.");
-				}
-			}
-			this.ClientTargetPlatforms = SetupTargetPlatforms(ref this.ClientDependentPlatformMap, Command, ClientTargetPlatforms, new ParamList<UnrealTargetPlatform>() { TargetPlatform }, true, "TargetPlatform", "Platform");
+			this.ClientTargetPlatforms = SetupTargetPlatforms(ref this.ClientDependentPlatformMap, Command, ClientTargetPlatforms, new ParamList<UnrealTargetPlatform>() { HostPlatform.Current.HostEditorPlatform }, true, "TargetPlatform", "Platform");
 
             // Parse command line params for server platforms "-ServerTargetPlatform=Win64+Mac", "-ServerPlatform=Win64+Mac". "-Win64" etc is not allowed here
             if (ServerDependentPlatformMap != null)
