UE-8433

Fix for missing checks for IOS UE4Game binaries in the Windows Editor

diff --git a/Engine/Source/Programs/AutomationTool/IOS/IOSPlatform.Automation.cs b/Engine/Source/Programs/AutomationTool/IOS/IOSPlatform.Automation.cs
--- a/Engine/Source/Programs/AutomationTool/IOS/IOSPlatform.Automation.cs
+++ b/Engine/Source/Programs/AutomationTool/IOS/IOSPlatform.Automation.cs
@@ -129,6 +129,15 @@ public class IOSPlatform : Platform
 	{
 		Log("Package {0}", Params.RawProjectPath);
 
+		// ensure the ue4game binary exists, if applicable
+		string FullExePath = CombinePaths(Path.GetDirectoryName (Params.ProjectGameExeFilename), SC.StageExecutables[0]);
+		if (!SC.IsCodeBasedProject && !FileExists_NoExceptions(FullExePath))
+		{
+			Log("Failed to find game binary " + FullExePath);
+			AutomationTool.ErrorReporter.Error("Stage Failed.", (int)AutomationTool.ErrorCodes.Error_MissingExecutable);
+			throw new AutomationException("Could not find binary {0}. You may need to build the UE4 project with your target configuration and platform.", FullExePath);
+		}
+
 		//@TODO: We should be able to use this code on both platforms, when the following issues are sorted:
 		//   - Raw executable is unsigned & unstripped (need to investigate adding stripping to IPP)
 		//   - IPP needs to be able to codesign a raw directory
@@ -140,15 +149,6 @@ public class IOSPlatform : Platform
 			// copy in all of the artwork and plist
 			var DeployHandler = UEBuildDeploy.GetBuildDeploy(UnrealTargetPlatform.IOS);
 
-			string FullExePath = CombinePaths(Path.GetDirectoryName (Params.ProjectGameExeFilename), SC.StageExecutables [0]);
-			// ensure the ue4game binary exists, if applicable
-			if (!SC.IsCodeBasedProject && !FileExists_NoExceptions(FullExePath))
-			{
-				Log("Failed to find game binary " + FullExePath);
-				AutomationTool.ErrorReporter.Error("Stage Failed.", (int)AutomationTool.ErrorCodes.Error_MissingExecutable);
-				throw new AutomationException("Could not find binary {0}. You may need to build the UE4 project with your target configuration and platform.", FullExePath);
-			}
-
 			DeployHandler.PrepForUATPackageOrDeploy(Params.ShortProjectName,
 				Path.GetDirectoryName(Params.RawProjectPath),
 				CombinePaths(Path.GetDirectoryName(Params.ProjectGameExeFilename), SC.StageExecutables[0]),
