- Stripping out some special .ini settings while staging (the change is lower level than it should be, but it's the fastest way to protect passwords) [UE-9113]

diff --git a/Engine/Source/Programs/AutomationTool/CommandUtils.cs b/Engine/Source/Programs/AutomationTool/CommandUtils.cs
--- a/Engine/Source/Programs/AutomationTool/CommandUtils.cs
+++ b/Engine/Source/Programs/AutomationTool/CommandUtils.cs
@@ -1247,7 +1247,7 @@ namespace AutomationTool
 		/// <param name="Dest">The full path to the destination file</param>
 		/// <param name="bAllowDifferingTimestamps">If true, will always skip a file if the destination exists, even if timestamp differs; defaults to false</param>
 		/// <returns>True if the operation was successful, false otherwise.</returns>
-		public static void CopyFileIncremental(string Source, string Dest, bool bAllowDifferingTimestamps = false)
+		public static void CopyFileIncremental(string Source, string Dest, bool bAllowDifferingTimestamps = false, bool bFilterSpecialLinesFromIniFiles = false)
 		{
 			Source = ConvertSeparators(PathSeparator.Default, Source);
 			Dest = ConvertSeparators(PathSeparator.Default, Dest);
@@ -1277,7 +1277,7 @@ namespace AutomationTool
 			{
 				throw new AutomationException("Failed to delete {0} for copy", Dest);
 			}
-			if (!InternalUtils.SafeCopyFile(Source, Dest))
+			if (!InternalUtils.SafeCopyFile(Source, Dest, bFilterSpecialLinesFromIniFiles:bFilterSpecialLinesFromIniFiles))
 			{
 				throw new AutomationException("Failed to copy {0} to {1}", Source, Dest);
 			}
diff --git a/Engine/Source/Programs/AutomationTool/Scripts/CopyBuildToStagingDirectory.Automation.cs b/Engine/Source/Programs/AutomationTool/Scripts/CopyBuildToStagingDirectory.Automation.cs
--- a/Engine/Source/Programs/AutomationTool/Scripts/CopyBuildToStagingDirectory.Automation.cs
+++ b/Engine/Source/Programs/AutomationTool/Scripts/CopyBuildToStagingDirectory.Automation.cs
@@ -488,7 +488,7 @@ public partial class Project : CommandUtils
 			string Dest = CombinePaths(StageDir, Pair.Value);
 			if (Src != Dest)  // special case for things created in the staging directory, like the pak file
 			{
-				CopyFileIncremental(Src, Dest);
+				CopyFileIncremental(Src, Dest, bFilterSpecialLinesFromIniFiles:true);
 			}
 		}
 		if (!String.IsNullOrEmpty(ManifestPath) && Mapping.Count > 0)
diff --git a/Engine/Source/Programs/AutomationTool/Utils.cs b/Engine/Source/Programs/AutomationTool/Utils.cs
--- a/Engine/Source/Programs/AutomationTool/Utils.cs
+++ b/Engine/Source/Programs/AutomationTool/Utils.cs
@@ -276,13 +276,56 @@ namespace AutomationTool
 			return Result;
 		}
 
+		// @todo: This could be passed in from elsewhere, and this should be somehow done per ini section
+		// but this will get it so that games won't ship passwords
+		private static string[] LinesToFilter = new string[]
+		{
+			"KeyStorePassword",
+		};
+
+		private static void FilterIniFile(string SourceName, string TargetName)
+		{
+			string[] Lines = File.ReadAllLines(SourceName);
+			StringBuilder NewLines = new StringBuilder("");
+
+			foreach (string Line in Lines)
+			{
+				// look for each filter on each line
+				bool bFiltered = false;
+				foreach (string Filter in LinesToFilter)
+				{
+					if (Line.StartsWith(Filter + "="))
+					{
+						bFiltered = true;
+						break;
+					}
+				}
+
+				// write out if it's not filtered out
+				if (!bFiltered)
+				{
+					NewLines.AppendLine(Line);
+				}
+			}
+
+			// now write out the final .ini file
+			if (File.Exists(TargetName))
+			{
+				File.Delete(TargetName);
+			}
+			File.WriteAllText(TargetName, NewLines.ToString());
+
+			// other code assumes same timestamp for source and dest
+			File.SetLastWriteTimeUtc(TargetName, File.GetLastWriteTimeUtc(SourceName));
+		}
+
 		/// <summary>
 		/// Copies a file.
 		/// </summary>
 		/// <param name="SourceName">Source name</param>
 		/// <param name="TargetName">Target name</param>
 		/// <returns>True if the operation was successful, false otherwise.</returns>
-		public static bool SafeCopyFile(string SourceName, string TargetName, bool bQuiet = false)
+		public static bool SafeCopyFile(string SourceName, string TargetName, bool bQuiet = false, bool bFilterSpecialLinesFromIniFiles = false)
 		{
 			if (!bQuiet)
 			{
@@ -298,13 +341,23 @@ namespace AutomationTool
 				bool Retry = true;
 				try
 				{
-					File.Copy(SourceName, TargetName, overwrite: true);
+					bool bSkipSizeCheck = false;
+					if (bFilterSpecialLinesFromIniFiles && Path.GetExtension(SourceName) == ".ini")
+					{
+						FilterIniFile(SourceName, TargetName);
+						// ini files may change size, don't check
+						bSkipSizeCheck = true;
+					}
+					else
+					{
+						File.Copy(SourceName, TargetName, overwrite: true);
+					}
 					Retry = !File.Exists(TargetName);
 					if (!Retry)
 					{
 						FileInfo SourceInfo = new FileInfo(SourceName);
 						FileInfo TargetInfo = new FileInfo(TargetName);
-						if (SourceInfo.Length != TargetInfo.Length)
+						if (!bSkipSizeCheck && SourceInfo.Length != TargetInfo.Length)
 						{
 							Log.WriteLine(TraceEventType.Warning, "Size mismatch {0} = {1} to {2} = {3}", SourceName, SourceInfo.Length, TargetName, TargetInfo.Length);
 							Retry = true;
