Fix incorrect HTML5 command line for content only games when enabling pak files.

diff --git a/Engine/Source/Programs/AutomationTool/HTML5/HTML5Platform.Automation.cs b/Engine/Source/Programs/AutomationTool/HTML5/HTML5Platform.Automation.cs
--- a/Engine/Source/Programs/AutomationTool/HTML5/HTML5Platform.Automation.cs
+++ b/Engine/Source/Programs/AutomationTool/HTML5/HTML5Platform.Automation.cs
@@ -141,7 +141,7 @@ public class HTML5Platform : Platform
 				{
 					InArguments = InArguments.Replace ("\"", "");
 					string[] Arguments = InArguments.Split(' ');
-					string ArgumentString = IsContentOnly ? "'" + InGameName + "/" + InGameName + ".uproject '," : "";
+					string ArgumentString = IsContentOnly ? "'../../../" + InGameName + "/" + InGameName + ".uproject '," : "";
 					for (int i = 0; i < Arguments.Length - 1; ++i)
 					{
 						ArgumentString += "'";
