fix for deploying the development executable instead of the specified executable to the payload directory
UE-7226
#ios

diff --git a/Engine/Source/Programs/AutomationTool/IOS/IOSPlatform.Automation.cs b/Engine/Source/Programs/AutomationTool/IOS/IOSPlatform.Automation.cs
--- a/Engine/Source/Programs/AutomationTool/IOS/IOSPlatform.Automation.cs
+++ b/Engine/Source/Programs/AutomationTool/IOS/IOSPlatform.Automation.cs
@@ -400,14 +400,6 @@ public class IOSPlatform : Platform
 		bool bWasGenerated = false;
 		string XcodeProj = EnsureXcodeProjectExists (RawProjectPath, LocalRoot, ProjectName, ProjectDirectory, IsCode, out bWasGenerated);
 
-		// ensure the correct data os acrpss
-		var DeployHandler = UEBuildDeploy.GetBuildDeploy(UnrealTargetPlatform.IOS);
-		DeployHandler.PrepForUATPackageOrDeploy(ProjectName,
-			Path.GetDirectoryName(RawProjectPath),
-			CombinePaths(BaseDirectory, GameName),
-			CombinePaths(LocalRoot, "Engine"),
-			Distribution, "", false);
-
 		string Arguments = "UBT_NO_POST_DEPLOY=true";
 		Arguments += " /usr/bin/xcrun xcodebuild build -project \"" + XcodeProj + "\"";
 		Arguments += " -scheme '";
