build script, remove waitfortestshared dependency if project disables testwithshared

diff --git a/Engine/Source/Programs/AutomationTool/Scripts/GUBP.Automation.cs b/Engine/Source/Programs/AutomationTool/Scripts/GUBP.Automation.cs
--- a/Engine/Source/Programs/AutomationTool/Scripts/GUBP.Automation.cs
+++ b/Engine/Source/Programs/AutomationTool/Scripts/GUBP.Automation.cs
@@ -5729,6 +5729,10 @@ public class GUBP : BuildCommand
                     foreach (var Test in EditorTests)
                     {
                         EditorTestNodes.Add(AddNode(new UATTestNode(this, HostPlatform, CodeProj, Test.Key, Test.Value, AgentSharingGroup)));
+						if (!Options.bTestWithShared)
+						{
+							RemovePseudodependencyFromNode((UATTestNode.StaticGetFullName(HostPlatform, CodeProj, Test.Key)), WaitForTestShared.StaticGetFullName());
+						}
                     }
                     if (EditorTestNodes.Count > 0)
                     {                        
@@ -5869,6 +5873,10 @@ public class GUBP : BuildCommand
 										{
 											var TestNodeName = Test.Key + "_" + Plat.ToString();
 											ThisMonoGameTestNodes.Add(AddNode(new UATTestNode(this, HostPlatform, CodeProj, TestNodeName, Test.Value, CookedAgentSharingGroup, false, RequiredPlatforms)));
+											if (!Options.bTestWithShared)
+											{
+												RemovePseudodependencyFromNode((UATTestNode.StaticGetFullName(HostPlatform, CodeProj, TestNodeName)), WaitForTestShared.StaticGetFullName());
+											}
 										}
 										if (ThisMonoGameTestNodes.Count > 0)
 										{
