Fixed up -iterate commandline option for UAT so that it actually enables iterative cooking instead of not :/.

diff --git a/Engine/Source/Programs/AutomationTool/ProjectParams.cs b/Engine/Source/Programs/AutomationTool/ProjectParams.cs
--- a/Engine/Source/Programs/AutomationTool/ProjectParams.cs
+++ b/Engine/Source/Programs/AutomationTool/ProjectParams.cs
@@ -467,7 +467,7 @@ namespace AutomationTool
             this.UnversionedCookedContent = GetParamValueIfNotSpecified(Command, UnversionedCookedContent, this.UnversionedCookedContent, "UnversionedCookedContent");
             this.Compressed = GetParamValueIfNotSpecified(Command, Compressed, this.Compressed, "compressed");
             this.UseDebugParamForEditorExe = GetParamValueIfNotSpecified(Command, UseDebugParamForEditorExe, this.UseDebugParamForEditorExe, "UseDebugParamForEditorExe");
-            this.IterativeCooking = GetParamValueIfNotSpecified(Command, IterativeCooking, this.IterativeCooking, "iterativecooking", "iterate");
+            this.IterativeCooking = GetParamValueIfNotSpecified(Command, IterativeCooking, this.IterativeCooking, new string[] { "iterativecooking", "iterate" } );
 			this.SkipCookOnTheFly = GetParamValueIfNotSpecified(Command, SkipCookOnTheFly, this.SkipCookOnTheFly, "skipcookonthefly");
             this.CookAll = GetParamValueIfNotSpecified(Command, CookAll, this.CookAll, "CookAll");
             this.CookMapsOnly = GetParamValueIfNotSpecified(Command, CookMapsOnly, this.CookMapsOnly, "CookMapsOnly");
@@ -522,7 +522,7 @@ namespace AutomationTool
 			this.RunCommandline = ParseParamValueIfNotSpecified(Command, RunCommandline, "addcmdline");
 			this.Package = GetParamValueIfNotSpecified(Command, Package, this.Package, "package");
 			this.Deploy = GetParamValueIfNotSpecified(Command, Deploy, this.Deploy, "deploy");
-			this.IterativeDeploy = GetParamValueIfNotSpecified(Command, IterativeDeploy, this.IterativeDeploy, "iterativedeploy", "iterate");
+			this.IterativeDeploy = GetParamValueIfNotSpecified(Command, IterativeDeploy, this.IterativeDeploy, new string[] {"iterativedeploy", "iterate" } );
 			this.Device = ParseParamValueIfNotSpecified(Command, Device, "device", String.Empty).Trim(new char[] { '\"' });
 
 			// strip the platform prefix the specified device.
