fix for missing extension when checking for iOS executable on PC (needs a .stub)
UE-9471
#ios

diff --git a/Engine/Source/Programs/AutomationTool/IOS/IOSPlatform.Automation.cs b/Engine/Source/Programs/AutomationTool/IOS/IOSPlatform.Automation.cs
--- a/Engine/Source/Programs/AutomationTool/IOS/IOSPlatform.Automation.cs
+++ b/Engine/Source/Programs/AutomationTool/IOS/IOSPlatform.Automation.cs
@@ -130,7 +130,7 @@ public class IOSPlatform : Platform
 		Log("Package {0}", Params.RawProjectPath);
 
 		// ensure the ue4game binary exists, if applicable
-		string FullExePath = CombinePaths(Path.GetDirectoryName (Params.ProjectGameExeFilename), SC.StageExecutables[0]);
+		string FullExePath = CombinePaths(Path.GetDirectoryName(Params.ProjectGameExeFilename), SC.StageExecutables[0] + (UnrealBuildTool.BuildHostPlatform.Current.Platform != UnrealTargetPlatform.Mac ? ".stub" : ""));
 		if (!SC.IsCodeBasedProject && !FileExists_NoExceptions(FullExePath))
 		{
 			Log("Failed to find game binary " + FullExePath);
