Update feature packs as part of promotable builds.

diff --git a/Engine/Source/Programs/AutomationTool/ProjectUtils.cs b/Engine/Source/Programs/AutomationTool/ProjectUtils.cs
--- a/Engine/Source/Programs/AutomationTool/ProjectUtils.cs
+++ b/Engine/Source/Programs/AutomationTool/ProjectUtils.cs
@@ -6,6 +6,7 @@ using System.Linq;
 using System.Text;
 using System.IO;
 using UnrealBuildTool;
+using System.Diagnostics;
 
 namespace AutomationTool
 {
@@ -472,6 +473,8 @@ namespace AutomationTool
             TargetRules.TargetType.Client,
             TargetRules.TargetType.Server,
         };
+
+		[DebuggerDisplay("{GameName}")]
         public class BranchUProject
         {
             public string GameName;
@@ -575,6 +578,11 @@ namespace AutomationTool
         public List<BranchUProject> CodeProjects = new List<BranchUProject>();
         public List<BranchUProject> NonCodeProjects = new List<BranchUProject>();
 
+		public IEnumerable<BranchUProject> AllProjects
+		{
+			get { return CodeProjects.Union(NonCodeProjects); }
+		}
+
         public BranchInfo(List<UnrealTargetPlatform> InHostPlatforms)
         {
             BaseEngineProject = new BranchUProject();
diff --git a/Engine/Source/Programs/AutomationTool/Scripts/GUBP.Automation.cs b/Engine/Source/Programs/AutomationTool/Scripts/GUBP.Automation.cs
--- a/Engine/Source/Programs/AutomationTool/Scripts/GUBP.Automation.cs
+++ b/Engine/Source/Programs/AutomationTool/Scripts/GUBP.Automation.cs
@@ -1382,6 +1382,86 @@ public class GUBP : BuildCommand
         }
     }
 
+	public class MakeFeaturePackNode : HostPlatformNode
+	{
+        BranchInfo.BranchUProject GameProj;
+
+        public MakeFeaturePackNode(GUBP bp, UnrealTargetPlatform InHostPlatform, BranchInfo.BranchUProject InGameProj)
+            : base(InHostPlatform)
+        {
+			GameProj = InGameProj;
+            if (GameProj.GameName != bp.Branch.BaseEngineProject.GameName && GameProj.Properties.Targets.ContainsKey(TargetRules.TargetType.Editor))
+            {
+                AddDependency(EditorGameNode.StaticGetFullName(InHostPlatform, GameProj));
+            }
+			AddDependency(ToolsNode.StaticGetFullName(InHostPlatform)); // for UnrealPak
+            AgentSharingGroup = "FeaturePacks"  + StaticGetHostPlatformSuffix(InHostPlatform);
+        }
+
+		public static bool ShouldMakeFeaturePack(UnrealTargetPlatform InHostPlatform, BranchInfo.BranchUProject InGameProj)
+		{
+			// No obvious way to store this in the project options; it's a property of non-code projects too.
+			if (InHostPlatform == UnrealTargetPlatform.Win64)
+			{
+				if(InGameProj.GameName == "StarterContent" || InGameProj.GameName == "MobileStarterContent" || InGameProj.GameName.StartsWith("FP_"))
+				{
+					return true;
+				}
+				if(InGameProj.GameName.StartsWith("TP_"))
+				{
+					return CommandUtils.FileExists(CommandUtils.CombinePaths(CommandUtils.GetDirectoryName(InGameProj.FilePath), "contents.txt"));
+				}
+			}
+			return false;
+		}
+
+        public static string StaticGetFullName(UnrealTargetPlatform InHostPlatform, BranchInfo.BranchUProject InGameProj)
+        {
+            return InGameProj.GameName + "_MakeFeaturePack" + StaticGetHostPlatformSuffix(InHostPlatform);
+        }
+
+		public override string GetFullName()
+        {
+            return StaticGetFullName(HostPlatform, GameProj);
+        }
+
+		public override void DoBuild(GUBP bp)
+        {
+			string ContentsFileName = CommandUtils.CombinePaths(CommandUtils.GetDirectoryName(GameProj.FilePath), "contents.txt");
+
+			// Make sure we delete the output file. It may be read-only.
+			string OutputFileName = CommandUtils.CombinePaths(CommandUtils.CmdEnv.LocalRoot, "FeaturePacks", Path.GetFileNameWithoutExtension(GameProj.GameName) + ".upack");
+			CommandUtils.DeleteFile(OutputFileName);
+
+			// Get the command line
+			string CmdLine = CommandUtils.MakePathSafeToUseWithCommandLine(OutputFileName) + " " + CommandUtils.MakePathSafeToUseWithCommandLine("-create=" + ContentsFileName);
+			if (GlobalCommandLine.Installed)
+			{
+				CmdLine += " -installed";
+			}
+			if (GlobalCommandLine.UTF8Output)
+			{
+				CmdLine += " -UTF8Output";
+			}
+
+			// Run UnrealPak
+			string UnrealPakExe;
+			if(HostPlatform == UnrealTargetPlatform.Win64)
+			{
+				UnrealPakExe = CombinePaths(CmdEnv.LocalRoot, "Engine/Binaries/Win64/UnrealPak.exe");
+			}
+			else
+			{
+				throw new AutomationException("Unknown path to UnrealPak for host platform ({0})", HostPlatform);
+			}
+			RunAndLog(CmdEnv, UnrealPakExe, CmdLine, Options: ERunOptions.Default | ERunOptions.AllowSpew | ERunOptions.UTF8Output);
+
+			// Add the build products
+			BuildProducts = new List<string>();
+			BuildProducts.Add(OutputFileName);
+		}
+	}
+
     public class GamePlatformMonolithicsNode : CompileNode
     {
         BranchInfo.BranchUProject GameProj;
@@ -1977,6 +2057,13 @@ public class GUBP : BuildCommand
                         }
                     }
                 }
+				foreach(var Proj in bp.Branch.AllProjects)
+				{
+					if(MakeFeaturePackNode.ShouldMakeFeaturePack(HostPlatform, Proj))
+					{
+						AddDependency(MakeFeaturePackNode.StaticGetFullName(HostPlatform, Proj));
+					}
+				}
             }
         }
 		public override bool IsSeparatePromotable()
@@ -5702,6 +5789,14 @@ public class GUBP : BuildCommand
                 }
             }
 			
+			foreach(BranchInfo.BranchUProject Project in Branch.AllProjects)
+			{
+				if(MakeFeaturePackNode.ShouldMakeFeaturePack(HostPlatform, Project))
+				{
+					AddNode(new MakeFeaturePackNode(this, HostPlatform, Project));
+				}
+			}
+
             foreach (var CodeProj in Branch.CodeProjects)
             {
                 var Options = CodeProj.Options(HostPlatform);
