- Added another tag to strip from .ini files [UE-9113]

diff --git a/Engine/Source/Programs/AutomationTool/Utils.cs b/Engine/Source/Programs/AutomationTool/Utils.cs
--- a/Engine/Source/Programs/AutomationTool/Utils.cs
+++ b/Engine/Source/Programs/AutomationTool/Utils.cs
@@ -281,6 +281,7 @@ namespace AutomationTool
 		private static string[] LinesToFilter = new string[]
 		{
 			"KeyStorePassword",
+			"KeyPassword",
 		};
 
 		private static void FilterIniFile(string SourceName, string TargetName)
