build script, put sample make builds behind their own trigger

diff --git a/Engine/Source/Programs/AutomationTool/Scripts/GUBP.Automation.cs b/Engine/Source/Programs/AutomationTool/Scripts/GUBP.Automation.cs
--- a/Engine/Source/Programs/AutomationTool/Scripts/GUBP.Automation.cs
+++ b/Engine/Source/Programs/AutomationTool/Scripts/GUBP.Automation.cs
@@ -372,6 +372,7 @@ public class GUBP : BuildCommand
 		public List<string> DependentPromotions = new List<string>();
         public List<string> AllDependencyBuildProducts = null;
         public List<string> AllDependencies = null;		
+		public List<string> CompletedDependencies = new List<string>();
         public string AgentSharingGroup = "";
         public int ComputedDependentCISFrequencyQuantumShift = -1;
 
@@ -553,6 +554,13 @@ public class GUBP : BuildCommand
                 FullNamesOfPseudosependencies.Add(Node);
             }
         }
+		public void AddCompletedDependency(string Node)
+		{
+			if(!CompletedDependencies.Contains(Node))
+			{
+				CompletedDependencies.Add(Node);
+			}
+		}
         public void RemovePseudodependency(string Node)
         {
             if (FullNamesOfPseudosependencies.Contains(Node))
@@ -1563,7 +1571,7 @@ public class GUBP : BuildCommand
 			if (!WithXp)
 			{
 				Name = InGameProj.GameName + "_" + InTargetPlatform + "_Mono" + StaticGetHostPlatformSuffix(InHostPlatform);
-			}
+        }
 			else
 			{
 				Name = InGameProj.GameName + "_WinXP_Mono" + StaticGetHostPlatformSuffix(InHostPlatform);
@@ -1619,62 +1627,62 @@ public class GUBP : BuildCommand
 					var AllowXp = Target.Rules.GUBP_BuildWindowsXPMonolithics();
 					if (!WithXp || (AllowXp && WithXp))
 					{
-						var Platforms = Target.Rules.GUBP_GetPlatforms_MonolithicOnly(HostPlatform);
-						var AdditionalPlatforms = Target.Rules.GUBP_GetBuildOnlyPlatforms_MonolithicOnly(HostPlatform);
-						var AllPlatforms = Platforms.Union(AdditionalPlatforms);
-						if (AllPlatforms.Contains(TargetPlatform) && Target.Rules.SupportsPlatform(TargetPlatform))
-						{
-							var Configs = Target.Rules.GUBP_GetConfigs_MonolithicOnly(HostPlatform, TargetPlatform);
-							foreach (var Config in Configs)
-							{
+                    var Platforms = Target.Rules.GUBP_GetPlatforms_MonolithicOnly(HostPlatform);
+					var AdditionalPlatforms = Target.Rules.GUBP_GetBuildOnlyPlatforms_MonolithicOnly(HostPlatform);
+					var AllPlatforms = Platforms.Union(AdditionalPlatforms);
+					if (AllPlatforms.Contains(TargetPlatform) && Target.Rules.SupportsPlatform(TargetPlatform))
+                    {
+                        var Configs = Target.Rules.GUBP_GetConfigs_MonolithicOnly(HostPlatform, TargetPlatform);
+                        foreach (var Config in Configs)
+                        {
 								if (WithXp)
 								{
 									Args += " -winxp";
 								}
 
-								if (GUBP.bBuildRocket)
-								{
-									if (HostPlatform == UnrealTargetPlatform.Win64)
-									{
-										if (TargetPlatform == UnrealTargetPlatform.Win32 && Config != UnrealTargetConfiguration.Shipping)
-										{
-											continue;
-										}
-										if (TargetPlatform == UnrealTargetPlatform.Win64 && Config != UnrealTargetConfiguration.Development)
-										{
-											continue;
-										}
-										if (TargetPlatform == UnrealTargetPlatform.Android && Config != UnrealTargetConfiguration.Shipping && Config != UnrealTargetConfiguration.Development)
-										{
-											continue;
-										}
+                            if (GUBP.bBuildRocket)
+                            {
+                                if (HostPlatform == UnrealTargetPlatform.Win64)
+                                {
+                                    if (TargetPlatform == UnrealTargetPlatform.Win32 && Config != UnrealTargetConfiguration.Shipping)
+                                    {
+                                        continue;
+                                    }
+                                    if (TargetPlatform == UnrealTargetPlatform.Win64 && Config != UnrealTargetConfiguration.Development)
+                                    {
+                                        continue;
+                                    }
+                                    if (TargetPlatform == UnrealTargetPlatform.Android && Config != UnrealTargetConfiguration.Shipping && Config != UnrealTargetConfiguration.Development)
+                                    {
+                                        continue;
+                                    }
 										if (TargetPlatform == UnrealTargetPlatform.HTML5 && Config != UnrealTargetConfiguration.Shipping && Config != UnrealTargetConfiguration.Development)
 										{
 											continue;
 										}
-										if (TargetPlatform == UnrealTargetPlatform.Linux && Config != UnrealTargetConfiguration.Shipping && Config != UnrealTargetConfiguration.Development)
-										{
-											continue;
-										}
-									}
-									else if (Config != UnrealTargetConfiguration.Shipping && Config != UnrealTargetConfiguration.Development)
-									{
-										continue;
-									}
-								}
-								if (GameProj.GameName == bp.Branch.BaseEngineProject.GameName)
-								{
-										Agenda.AddTargets(new string[] { Target.TargetName }, TargetPlatform, Config, InAddArgs: Args);
-								}
-								else
-								{
-										Agenda.AddTargets(new string[] { Target.TargetName }, TargetPlatform, Config, GameProj.FilePath, InAddArgs: Args);
-								}
-							}
-						}
+                                    if (TargetPlatform == UnrealTargetPlatform.Linux && Config != UnrealTargetConfiguration.Shipping && Config != UnrealTargetConfiguration.Development)
+                                    {
+                                        continue;
+                                    }
+                                }
+                                else if (Config != UnrealTargetConfiguration.Shipping && Config != UnrealTargetConfiguration.Development)
+                                {
+                                    continue;
+                                }
+                            }
+                            if (GameProj.GameName == bp.Branch.BaseEngineProject.GameName)
+                            {
+                                Agenda.AddTargets(new string[] { Target.TargetName }, TargetPlatform, Config, InAddArgs: Args);
+                            }
+                            else
+                            {
+                                Agenda.AddTargets(new string[] { Target.TargetName }, TargetPlatform, Config, GameProj.FilePath, InAddArgs: Args);
+                            }
+                        }
                     }
                 }
             }
+            }
 
             return Agenda;
         }
@@ -1960,6 +1968,106 @@ public class GUBP : BuildCommand
         }
     }
 
+	public class SharedCookAggregateNode : AggregateNode
+	{
+		public SharedCookAggregateNode(GUBP bp, List<UnrealTargetPlatform> InHostPlatforms, Dictionary<string, List<UnrealTargetPlatform>> NonCodeProjectNames, Dictionary<string, List<TargetRules.GUBPFormalBuild>> NonCodeFormalBuilds)
+		{			
+			foreach (var HostPlatform in InHostPlatforms)
+			{
+				{
+					var Options = bp.Branch.BaseEngineProject.Options(HostPlatform);
+					if (Options.bIsPromotable && !Options.bSeparateGamePromotion)
+					{
+						var Kind = TargetRules.TargetType.Game;						
+						if (bp.Branch.BaseEngineProject.Properties.Targets.ContainsKey(Kind))
+						{
+							var Target = bp.Branch.BaseEngineProject.Properties.Targets[Kind];
+							var Platforms = Target.Rules.GUBP_GetPlatforms_MonolithicOnly(HostPlatform);
+							foreach (var Plat in Platforms)
+							{
+								if (bNoIOSOnPC && Plat == UnrealTargetPlatform.IOS && HostPlatform == UnrealTargetPlatform.Win64)
+								{
+									continue;
+								}
+								if (bp.ActivePlatforms.Contains(Plat))
+								{
+									string CookedPlatform = Platform.Platforms[Plat].GetCookPlatform(Kind == TargetRules.TargetType.Server, Kind == TargetRules.TargetType.Client, "");										
+									
+									foreach (var NonCodeProject in bp.Branch.NonCodeProjects)
+									{										
+										if (!NonCodeProjectNames.ContainsKey(NonCodeProject.GameName) || !NonCodeProjectNames[NonCodeProject.GameName].Contains(Plat))
+										{
+											continue;
+										}
+										if(NonCodeFormalBuilds.ContainsKey(NonCodeProject.GameName))
+										{
+											var PlatList = NonCodeFormalBuilds[NonCodeProject.GameName];
+											foreach (var PlatPair in PlatList)
+											{
+												if(PlatPair.TargetPlatform == Plat)
+												{													
+													AddCompletedDependency(CookNode.StaticGetFullName(HostPlatform, NonCodeProject, CookedPlatform));
+												}
+											}
+										}
+									}
+								}
+							}								
+						}						
+					}					
+				}
+				foreach (var CodeProj in bp.Branch.CodeProjects)
+				{
+					var Options = CodeProj.Options(HostPlatform);
+					if (!Options.bSeparateGamePromotion)
+					{
+						if (Options.bIsPromotable)
+						{
+							foreach (var Kind in BranchInfo.MonolithicKinds)
+							{
+								if (CodeProj.Properties.Targets.ContainsKey(Kind))
+								{
+									var Target = CodeProj.Properties.Targets[Kind];
+									var Platforms = new List<UnrealTargetPlatform>();
+									var PossPlats = Target.Rules.GUBP_GetPlatforms_MonolithicOnly(HostPlatform);
+									var FormalBuilds = Target.Rules.GUBP_GetConfigsForFormalBuilds_MonolithicOnly(HostPlatform);
+									foreach (var FormalBuild in FormalBuilds)
+									{
+										var Platform = FormalBuild.TargetPlatform;
+										if(!Platforms.Contains(Platform) && PossPlats.Contains(Platform))
+										{
+											Platforms.Add(Platform);
+										}
+									}
+									foreach (var Plat in Platforms)
+									{
+										if (bNoIOSOnPC && Plat == UnrealTargetPlatform.IOS && HostPlatform == UnrealTargetPlatform.Win64)
+										{
+											continue;
+										}
+										if (bp.ActivePlatforms.Contains(Plat))
+										{
+											string CookedPlatform = Platform.Platforms[Plat].GetCookPlatform(Kind == TargetRules.TargetType.Server, Kind == TargetRules.TargetType.Client, "");
+											AddCompletedDependency(CookNode.StaticGetFullName(HostPlatform, CodeProj, CookedPlatform));
+										}
+									}
+								}
+							}
+						}
+					}
+				}				
+			}
+		}
+		public static string StaticGetFullName()
+		{
+			return "SharedCookAggregate";
+		}
+		public override string GetFullName()
+		{
+			return StaticGetFullName();
+		}
+	}
+
     public class GameAggregatePromotableNode : AggregatePromotableNode
     {
         BranchInfo.BranchUProject GameProj;		
@@ -1997,12 +2105,12 @@ public class GUBP : BuildCommand
 							if(GameProj.Properties.Targets[TargetRules.TargetType.Game].Rules.GUBP_BuildWindowsXPMonolithics())
 							{
 								AddDependency(GamePlatformMonolithicsNode.StaticGetFullName(HostPlatform, GameProj, Plat, true));
-							}
                     }
                 }
             }
         }
         }
+        }
 
         public static string StaticGetFullName(BranchInfo.BranchUProject InGameProj)
         {
@@ -2076,6 +2184,7 @@ public class GUBP : BuildCommand
         }
     }
 
+
     public class WaitForUserInput : GUBPNode
     {
         protected bool bTriggerWasTriggered;
@@ -2145,6 +2254,23 @@ public class GUBP : BuildCommand
         }
     }
 
+	public class WaitForFormalUserInput : WaitForUserInput
+	{
+		
+		public WaitForFormalUserInput(GUBP bp)
+		{		
+			AddDependency(SharedCookAggregateNode.StaticGetFullName());
+			AddCompletedDependency(SharedCookAggregateNode.StaticGetFullName());
+		}
+		public static string StaticGetFullName()
+		{
+			return "Shared_WaitToMakeBuilds";
+		}
+		public override string GetFullName()
+		{
+			return StaticGetFullName();
+		}
+	}
     public class WaitForPromotionUserInput : WaitForUserInput
     {
         string PromotionLabelPrefix;
@@ -2833,7 +2959,7 @@ public class GUBP : BuildCommand
 			var Options = InGameProj.Options(HostPlatform);
 			if(!Options.bSeparateGamePromotion)
 			{
-				AddPseudodependency(SharedLabelPromotableNode.StaticGetFullName(false));				
+				AddPseudodependency(WaitForFormalUserInput.StaticGetFullName());				
 			}
             if (ClientTargetPlatforms != null)
             {
@@ -3186,11 +3312,6 @@ public class GUBP : BuildCommand
 
             // verify we actually built these
             var WorkingGameProject = InGameProj;
-			var Options = InGameProj.Options(HostPlatform);
-			if(!Options.bSeparateGamePromotion)
-			{
-				AddPseudodependency(SharedLabelPromotableNode.StaticGetFullName(false));
-			}
             if (!WorkingGameProject.Properties.Targets.ContainsKey(TargetRules.TargetType.Editor))
             {
                 // this is a codeless project, use the base project
@@ -3769,6 +3890,32 @@ public class GUBP : BuildCommand
         return Result;
     }
 
+	List<string> GetCompletedOnlyDependencies(string NodeToDo, bool bFlat = false, bool ECOnly = true)
+	{
+		var Result = new List<string>();
+		foreach (var Node in GUBPNodes[NodeToDo].CompletedDependencies)
+		{
+			bool Usable = GUBPNodes[Node].RunInEC() || !ECOnly;
+			if(Usable)
+			{
+				if(!Result.Contains(Node))
+				{
+					Result.Add(Node);
+				}
+			}
+			if (bFlat || !Usable)
+			{
+				foreach (var RNode in GetCompletedOnlyDependencies(Node, bFlat, ECOnly))
+				{
+					if (!Result.Contains(RNode))
+					{
+						Result.Add(RNode);
+					}
+				}
+			}
+		}
+		return Result;
+	}
     List<string> GetECDependencies(string NodeToDo, bool bFlat = false)
     {
         return GetDependencies(NodeToDo, bFlat, true);
@@ -3910,6 +4057,10 @@ public class GUBP : BuildCommand
             {
                 Result = Math.Max(ComputeDependentCISFrequencyQuantumShift(Dep), Result);
             }
+			foreach (var Dep in GUBPNodes[NodeToDo].CompletedDependencies)
+			{
+				Result = Math.Max(ComputeDependentCISFrequencyQuantumShift(Dep), Result);
+			}
             if (Result < 0)
             {
                 throw new AutomationException("Failed to compute shift.");
@@ -4202,6 +4353,10 @@ public class GUBP : BuildCommand
                 {
                     Log("           pdep> {0}", Dep);
                 }
+				foreach (var Dep in GUBPNodes[NodeToDo].CompletedDependencies)
+				{
+					Log("			cdep>{0}", Dep);
+				}
             }
             if (bShowECDependencies)
             {
@@ -4209,7 +4364,12 @@ public class GUBP : BuildCommand
                 {
                     Log("           {0}", Dep);
                 }
+				foreach (var Dep in GetCompletedOnlyDependencies(NodeToDo))
+				{
+					Log("           compDep> {0}", Dep);
+				}
             }
+
 			if(bShowDependednOn)
 			{
 				foreach (var Dep in GUBPNodes[NodeToDo].FullNamesOfDependedOn)
@@ -4446,7 +4606,8 @@ public class GUBP : BuildCommand
             foreach (var NodeToDo in NodesToDo)
             {
                 bool bReady = true;
-                bool bPseudoReady = true;
+				bool bPseudoReady = true;
+				bool bCompReady = true;
                 if (!SubSort && GUBPNodes[NodeToDo].AgentSharingGroup != "")
                 {
                     if (ExaminedAgentGroups.Contains(GUBPNodes[NodeToDo].AgentSharingGroup))
@@ -4518,6 +4679,18 @@ public class GUBP : BuildCommand
                             break;
                         }
                     }
+					foreach (var Dep in GUBPNodes[NodeToDo].CompletedDependencies)
+					{
+						if (!GUBPNodes.ContainsKey(Dep))
+						{
+							throw new AutomationException("Completed Dependency {0} node found.", Dep);
+						}
+						if (NodesToDo.Contains(Dep))
+						{
+							bCompReady = false;
+							break;
+						}
+					}
                 }
                 var Priority = GUBPNodes[NodeToDo].Priority();
 
@@ -4561,6 +4734,10 @@ public class GUBP : BuildCommand
                                             bReady = false;
                                         }
                                     }
+									if (!bCompReady)
+									{
+										bReady = false;
+									}
                                 }
                             }
                         }
@@ -5071,7 +5248,7 @@ public class GUBP : BuildCommand
         bool bChanges = ParseParam("Changes") || ParseParam("AllChanges");
         bool bHistory = ParseParam("History") || bChanges;
         bool bListOnly = ParseParam("ListOnly") || bHistory;
-        bool bSkipTriggers = ParseParam("SkipTriggers");		
+        bool bSkipTriggers = ParseParam("SkipTriggers");
         bFake = ParseParam("fake");
         bool bFakeEC = ParseParam("FakeEC");
         TimeIndex = ParseParamInt("TimeIndex", 0);
@@ -5205,7 +5382,7 @@ public class GUBP : BuildCommand
                 int Minutes = int.Parse(Parts[1]);
 
                 int DeltaMinutes = NowMinutes - Minutes;
-				
+
                 int TimeQuantum = 20;
 				if(BranchOptions.QuantumOverride != 0)
 				{
@@ -5618,10 +5795,10 @@ public class GUBP : BuildCommand
 										AddNode(new GamePlatformMonolithicsNode(this, HostPlatform, Branch.BaseEngineProject, Plat, true));
 									}
                             }
+                            }
                         }
                     }
                 }
-                }
 
                 var CookedAgentSharingGroup = "Shared_CookedTests" + HostPlatformNode.StaticGetHostPlatformSuffix(HostPlatform);
 
@@ -5788,6 +5965,10 @@ public class GUBP : BuildCommand
                     AddNode(new GameAggregateNode(this, HostPlatform, Branch.BaseEngineProject, "AllCookedTests", GameTestNodes));
                 }
             }
+			if (!GUBPNodes.ContainsKey(SharedCookAggregateNode.StaticGetFullName()))
+			{
+				AddNode(new SharedCookAggregateNode(this, HostPlatforms, NonCodeProjectNames, NonCodeFormalBuilds));
+			}
 			
 			foreach(BranchInfo.BranchUProject Project in Branch.AllProjects)
 			{
@@ -6056,6 +6237,7 @@ public class GUBP : BuildCommand
             AddNode(new WaitForTestShared(this));
             AddNode(new WaitForSharedPromotionUserInput(this, true));
             AddNode(new SharedLabelPromotableNode(this, true));
+			AddNode(new WaitForFormalUserInput(this));
         }
         foreach (var HostPlatform in HostPlatforms)
         {
@@ -6690,9 +6872,9 @@ public class GUBP : BuildCommand
                 ECProps.Add(string.Format("SortKey/{0}={1}", NodePair.Key, NodePair.Value));
             }
 			foreach (var NodePair in FullNodeDependedOnBy)
-			{
+            {
 				ECProps.Add(string.Format("DependedOnBy/{0}={1}", NodePair.Key, NodePair.Value));
-			}
+            }
 			foreach (var NodePair in FullNodeDependentPromotions)
 			{
 				ECProps.Add(string.Format("DependentPromotions/{0}={1}", NodePair.Key, NodePair.Value));
@@ -6848,6 +7030,7 @@ public class GUBP : BuildCommand
                     string PreCondition = "";
                     string RunCondition = "";
                     var UncompletedEcDeps = new List<string>();
+					var UncompletedCompletedDeps = new List<string>();
 
                     string MyAgentGroup = GUBPNodes[NodeToDo].AgentSharingGroup;
                     bool bDoNestedJobstep = false;
@@ -6879,6 +7062,18 @@ public class GUBP : BuildCommand
                         }
                     }
                     var PreConditionUncompletedEcDeps = UncompletedEcDeps;
+					var CompletedDeps = GetCompletedOnlyDependencies(NodeToDo);					
+					foreach (var Dep in CompletedDeps)
+					{
+						if (GUBPNodes[Dep].RunInEC() && !NodeIsAlreadyComplete(Dep, LocalOnly) && OrdereredToDo.Contains(Dep)) // if something is already finished, we don't put it into EC
+						{
+							if (OrdereredToDo.IndexOf(Dep) > OrdereredToDo.IndexOf(NodeToDo))
+							{
+								throw new AutomationException("Topological sort error, node {0} has a dependency of {1} which sorted after it.", NodeToDo, Dep);
+							}
+							UncompletedCompletedDeps.Add(Dep);
+						}
+					}
                     if (MyAgentGroup != "")
                     {
                         bDoNestedJobstep = true;
@@ -6936,6 +7131,25 @@ public class GUBP : BuildCommand
                         }
                         PreCondition = PreCondition + ") true;]\"";
                     }
+					if(UncompletedCompletedDeps.Count > 0)
+					{
+						PreCondition = "\"\\$\" . \"[/javascript if(";
+						// these run "parallel", but we add preconditions to serialize them
+						int Index = 0;
+						foreach (var Dep in CompletedDeps)
+						{
+							if (GUBPNodes[Dep].RunInEC())
+							{
+								PreCondition = PreCondition + "getProperty('" + GetJobStep(PreconditionParentPath, Dep) + "/status\') == \'completed\'";
+								Index++;
+								if (Index != CompletedDeps.Count)
+								{
+									PreCondition = PreCondition + " && ";
+								}
+							}
+						}
+						PreCondition = PreCondition + ") true;]\"";
+					}
 
                     if (UncompletedEcDeps.Count > 0)
                     {
@@ -7186,7 +7400,7 @@ public class GUBP : BuildCommand
                     bool WasLocal;
 					try
 					{
-                    GUBPNodes[NodeToDo].BuildProducts = RetrieveFromTempStorage(CmdEnv, NodeStoreName, out WasLocal, GameNameIfAny, StorageRootIfAny);
+						GUBPNodes[NodeToDo].BuildProducts = RetrieveFromTempStorage(CmdEnv, NodeStoreName, out WasLocal, GameNameIfAny, StorageRootIfAny);
 					}
 					catch
 					{
