build script, cook less often as part of CI

diff --git a/Engine/Source/Programs/AutomationTool/Scripts/GUBP.Automation.cs b/Engine/Source/Programs/AutomationTool/Scripts/GUBP.Automation.cs
--- a/Engine/Source/Programs/AutomationTool/Scripts/GUBP.Automation.cs
+++ b/Engine/Source/Programs/AutomationTool/Scripts/GUBP.Automation.cs
@@ -2658,7 +2658,7 @@ public class GUBP : BuildCommand
         }
         public override int CISFrequencyQuantumShift(GUBP bp)
         {
-            return base.CISFrequencyQuantumShift(bp) + 3 + (bIsMassive ? 1 : 0);
+            return base.CISFrequencyQuantumShift(bp) + 4 + (bIsMassive ? 1 : 0);
         }
         public override float Priority()
         {
