Remove unnecessary dependency on template editors being built before making feature packs.

diff --git a/Engine/Source/Programs/AutomationTool/Scripts/GUBP.Automation.cs b/Engine/Source/Programs/AutomationTool/Scripts/GUBP.Automation.cs
--- a/Engine/Source/Programs/AutomationTool/Scripts/GUBP.Automation.cs
+++ b/Engine/Source/Programs/AutomationTool/Scripts/GUBP.Automation.cs
@@ -1398,10 +1398,6 @@ public class GUBP : BuildCommand
             : base(InHostPlatform)
         {
 			GameProj = InGameProj;
-            if (GameProj.GameName != bp.Branch.BaseEngineProject.GameName && GameProj.Properties.Targets.ContainsKey(TargetRules.TargetType.Editor))
-            {
-                AddDependency(EditorGameNode.StaticGetFullName(InHostPlatform, GameProj));
-            }
 			AddDependency(ToolsNode.StaticGetFullName(InHostPlatform)); // for UnrealPak
             AgentSharingGroup = "FeaturePacks"  + StaticGetHostPlatformSuffix(InHostPlatform);
         }
