build script, change releaseMode to Release

diff --git a/Engine/Source/Programs/AutomationTool/Scripts/GUBP.Automation.cs b/Engine/Source/Programs/AutomationTool/Scripts/GUBP.Automation.cs
--- a/Engine/Source/Programs/AutomationTool/Scripts/GUBP.Automation.cs
+++ b/Engine/Source/Programs/AutomationTool/Scripts/GUBP.Automation.cs
@@ -6968,7 +6968,7 @@ public class GUBP : BuildCommand
             if (LastSticky == "" && bHaveECNodes)
             {
                 // if we don't have any sticky nodes and we have other nodes, we run a fake noop just to release the resource 
-                string Args = String.Format("{0}, subprocedure => 'GUBP_UAT_Node', parallel => '0', jobStepName => 'Noop', actualParameter => [{{actualParameterName => 'NodeName', value => 'Noop'}}, {{actualParameterName => 'Sticky', value =>'1' }}], releaseMode => 'release'}});", BaseArgs);
+                string Args = String.Format("{0}, subprocedure => 'GUBP_UAT_Node', parallel => '0', jobStepName => 'Noop', actualParameter => [{{actualParameterName => 'NodeName', value => 'Noop'}}, {{actualParameterName => 'Sticky', value =>'1' }}], releaseMode => 'Release'}});", BaseArgs);
                 StepList.Add(Args);
                 bHasNoop = true;
             }
@@ -7230,7 +7230,7 @@ public class GUBP : BuildCommand
 #endif
                     if (Sticky && NodeToDo == LastSticky)
                     {
-                        Args = Args + ", releaseMode => 'release'";
+                        Args = Args + ", releaseMode => 'Release'";
                     }
                     Args = Args + "});";
                     StepList.Add(Args);
@@ -7271,7 +7271,7 @@ public class GUBP : BuildCommand
                             RelPreCondition = RelPreCondition + "getProperty('" + PreconditionParentPath + "/jobSteps[" + NodeToDo + "]/status') == 'completed'";
                             RelPreCondition = RelPreCondition + ") true;]\"";
                             // we need to release the resource
-                            string RelArgs = String.Format("{0}, subprocedure => 'GUBP_Release_AgentShare', parallel => '1', jobStepName => 'Release_{1}', actualParameter => [{{actualParameterName => 'AgentSharingGroup', valued => '{2}'}}], releaseMode => 'release', precondition => '{3}'",
+                            string RelArgs = String.Format("{0}, subprocedure => 'GUBP_Release_AgentShare', parallel => '1', jobStepName => 'Release_{1}', actualParameter => [{{actualParameterName => 'AgentSharingGroup', valued => '{2}'}}], releaseMode => 'Release', precondition => '{3}'",
                                 BaseArgs, MyAgentGroup, MyAgentGroup, RelPreCondition);
                             StepList.Add(RelArgs);
                         }
