Fixed missing startup map from the maps to cook.

diff --git a/Engine/Source/Programs/AutomationTool/ProjectParams.cs b/Engine/Source/Programs/AutomationTool/ProjectParams.cs
--- a/Engine/Source/Programs/AutomationTool/ProjectParams.cs
+++ b/Engine/Source/Programs/AutomationTool/ProjectParams.cs
@@ -615,6 +615,11 @@ namespace AutomationTool
                 this.MapsToCook = MapsToCook;
             }
 
+            if (String.IsNullOrEmpty(this.MapToRun) == false)
+            {
+                this.MapsToCook.Add(this.MapToRun);
+            }
+            
 			if (ServerConfigsToBuild == null)
 			{
 				if (Command != null)
@@ -1886,10 +1891,10 @@ namespace AutomationTool
 				throw new AutomationException("Don't use both -cook and -cookonthefly.");
 			}
 
-            if (Compressed && !Pak)
+            /*if (Compressed && !Pak)
             {
                 throw new AutomationException("-compressed can only be used with -pak");
-            }
+            }*/
 
             if (CreateChunkInstall && (!Manifests || !Stage))
             {
