build script, save compile/cook/test duration as an ec property separately from full node duration

diff --git a/Engine/Source/Programs/AutomationTool/Scripts/GUBP.Automation.cs b/Engine/Source/Programs/AutomationTool/Scripts/GUBP.Automation.cs
--- a/Engine/Source/Programs/AutomationTool/Scripts/GUBP.Automation.cs
+++ b/Engine/Source/Programs/AutomationTool/Scripts/GUBP.Automation.cs
@@ -729,6 +729,7 @@ public class GUBP : BuildCommand
                 bool ReallyDeleteBuildProducts = DeleteBuildProducts() && !GUBP.bForceIncrementalCompile;
                 Agenda.DoRetries = false; // these would delete build products
                 UE4Build.Build(Agenda, InDeleteBuildProducts: ReallyDeleteBuildProducts, InUpdateVersionFiles: false, InForceUnity: true);
+				
                 PostBuild(bp, UE4Build);
 
                 UE4Build.CheckBuildProducts(UE4Build.BuildProductFiles);
@@ -5095,7 +5096,7 @@ public class GUBP : BuildCommand
             {
                 var Parts = Prop.Split("=".ToCharArray());
                 RunECTool(String.Format("setProperty \"/myWorkflow/{0}\" \"{1}\"", Parts[0], Parts[1]), true);
-            }
+            }			
         }
         catch (Exception Ex)
         {
@@ -5103,6 +5104,19 @@ public class GUBP : BuildCommand
             Log(System.Diagnostics.TraceEventType.Warning, LogUtils.FormatException(Ex));
         }
     }
+	void UpdateECBuildTime(string NodeToDo, double BuildDuration)
+	{
+		try
+		{
+			Log("Updating duration prop for node {0}", NodeToDo);
+			RunECTool(String.Format("setProperty \"/myWorkflow/NodeDuration/{0}\" \"{1}\"", NodeToDo, BuildDuration.ToString()));
+		}
+		catch (Exception Ex)
+		{
+			Log(System.Diagnostics.TraceEventType.Warning, "Failed to UpdateECBuildTime.");
+			Log(System.Diagnostics.TraceEventType.Warning, LogUtils.FormatException(Ex));
+		}
+	}
 
 
     [Help("Runs one, several or all of the GUBP nodes")]
@@ -7421,6 +7435,7 @@ public class GUBP : BuildCommand
                 {
                     SaveStatus(NodeToDo, StartedTempStorageSuffix, NodeStoreName, bSaveSharedTempStorage, GameNameIfAny);
                 }
+				var BuildDuration = 0.0;
                 try
                 {
                     if (!String.IsNullOrEmpty(FakeFail) && FakeFail.Equals(NodeToDo, StringComparison.InvariantCultureIgnoreCase))
@@ -7441,7 +7456,10 @@ public class GUBP : BuildCommand
 							RetrieveFromPermanentStorage(CmdEnv, "Rocket\\Automated", RocketBuild, GUBPNodes[NodeToDo].NodeHostPlatform());
 						}
 						Log("***** Building GUBP Node {0} for {1}", NodeToDo, NodeStoreName);
+						var StartTime = DateTime.UtcNow;
                         GUBPNodes[NodeToDo].DoBuild(this);
+						BuildDuration = (DateTime.UtcNow - StartTime).TotalMilliseconds / 1000;
+						
                     }
                     if (!GUBPNodes[NodeToDo].IsAggregate())
                     {
@@ -7477,6 +7495,7 @@ public class GUBP : BuildCommand
                         UpdateNodeHistory(NodeToDo, CLString);
                         SaveStatus(NodeToDo, FailedTempStorageSuffix, NodeStoreName, bSaveSharedTempStorage, GameNameIfAny, ParseParamValue("MyJobStepId"));
                         UpdateECProps(NodeToDo, CLString);
+						UpdateECBuildTime(NodeToDo, BuildDuration);
                     }
 
                     Log("{0}", ExceptionToString(Ex));
@@ -7529,6 +7548,7 @@ public class GUBP : BuildCommand
                     UpdateNodeHistory(NodeToDo, CLString);
                     SaveStatus(NodeToDo, SucceededTempStorageSuffix, NodeStoreName, bSaveSharedTempStorage, GameNameIfAny);
                     UpdateECProps(NodeToDo, CLString);
+					UpdateECBuildTime(NodeToDo, BuildDuration);
                 }
             }
             foreach (var Product in GUBPNodes[NodeToDo].BuildProducts)
