build script, make feature packs alongside promotables every 2 hr 40m

diff --git a/Engine/Source/Programs/AutomationTool/Scripts/GUBP.Automation.cs b/Engine/Source/Programs/AutomationTool/Scripts/GUBP.Automation.cs
--- a/Engine/Source/Programs/AutomationTool/Scripts/GUBP.Automation.cs
+++ b/Engine/Source/Programs/AutomationTool/Scripts/GUBP.Automation.cs
@@ -1430,6 +1430,10 @@ public class GUBP : BuildCommand
             return StaticGetFullName(HostPlatform, GameProj);
         }
 
+		public override int CISFrequencyQuantumShift(GUBP bp)
+		{
+			return base.CISFrequencyQuantumShift(bp) + 3;
+		}
 		public override void DoBuild(GUBP bp)
         {
 			string ContentsFileName = CommandUtils.CombinePaths(CommandUtils.GetDirectoryName(GameProj.FilePath), "contents.txt");
