Integrate latest Rocket changes from 4.7 branch.

diff --git a/Engine/Source/Programs/AutomationTool/Scripts/GUBP.Automation.cs b/Engine/Source/Programs/AutomationTool/Scripts/GUBP.Automation.cs
--- a/Engine/Source/Programs/AutomationTool/Scripts/GUBP.Automation.cs
+++ b/Engine/Source/Programs/AutomationTool/Scripts/GUBP.Automation.cs
@@ -1395,7 +1395,7 @@ public class GUBP : BuildCommand
 	{
         BranchInfo.BranchUProject GameProj;
 
-        public MakeFeaturePackNode(GUBP bp, UnrealTargetPlatform InHostPlatform, BranchInfo.BranchUProject InGameProj)
+        public MakeFeaturePackNode(UnrealTargetPlatform InHostPlatform, BranchInfo.BranchUProject InGameProj)
             : base(InHostPlatform)
         {
 			GameProj = InGameProj;
@@ -1403,23 +1403,36 @@ public class GUBP : BuildCommand
             AgentSharingGroup = "FeaturePacks"  + StaticGetHostPlatformSuffix(InHostPlatform);
         }
 
-		public static bool ShouldMakeFeaturePack(UnrealTargetPlatform InHostPlatform, BranchInfo.BranchUProject InGameProj)
+		public static bool IsFeaturePack(BranchInfo.BranchUProject InGameProj)
 		{
 			// No obvious way to store this in the project options; it's a property of non-code projects too.
-			if (InHostPlatform == UnrealTargetPlatform.Win64)
+			if(InGameProj.GameName == "StarterContent" || InGameProj.GameName == "MobileStarterContent" || InGameProj.GameName.StartsWith("FP_"))
 			{
-				if(InGameProj.GameName == "StarterContent" || InGameProj.GameName == "MobileStarterContent" || InGameProj.GameName.StartsWith("FP_"))
-				{
-					return true;
-				}
-				if(InGameProj.GameName.StartsWith("TP_"))
-				{
-					return CommandUtils.FileExists(CommandUtils.CombinePaths(CommandUtils.GetDirectoryName(InGameProj.FilePath), "contents.txt"));
-				}
+				return true;
+			}
+			if(InGameProj.GameName.StartsWith("TP_"))
+			{
+				return CommandUtils.FileExists(CommandUtils.CombinePaths(CommandUtils.GetDirectoryName(InGameProj.FilePath), "contents.txt"));
 			}
 			return false;
 		}
 
+		public static UnrealTargetPlatform GetDefaultBuildPlatform(GUBP bp)
+		{
+			if(bp.HostPlatforms.Contains(UnrealTargetPlatform.Win64))
+			{
+				return UnrealTargetPlatform.Win64;
+			}
+			else if(bp.HostPlatforms.Contains(UnrealTargetPlatform.Mac))
+			{
+				return UnrealTargetPlatform.Mac;
+			}
+			else
+			{
+				return bp.HostPlatforms[0];
+			}
+		}
+
         public static string StaticGetFullName(UnrealTargetPlatform InHostPlatform, BranchInfo.BranchUProject InGameProj)
         {
             return InGameProj.GameName + "_MakeFeaturePack" + StaticGetHostPlatformSuffix(InHostPlatform);
@@ -1616,6 +1629,11 @@ public class GUBP : BuildCommand
 
             string Args = "-nobuilduht -skipactionhistory -CopyAppBundleBackToDevice" + bp.RocketUBTArgs();
 
+            if (GUBP.bBuildRocket && (TargetPlatform == UnrealTargetPlatform.Win32 || TargetPlatform == UnrealTargetPlatform.Win64))
+            {
+                Args += " -nodebuginfo";
+            }
+
             foreach (var Kind in BranchInfo.MonolithicKinds)
             {
                 if (GUBP.bBuildRocket && Kind != TargetRules.TargetType.Game)
@@ -2166,14 +2184,16 @@ public class GUBP : BuildCommand
                         }
                     }
                 }
-				foreach(var Proj in bp.Branch.AllProjects)
+            }
+
+			UnrealTargetPlatform FeaturePackPlatform = MakeFeaturePackNode.GetDefaultBuildPlatform(bp);
+			foreach(var Proj in bp.Branch.AllProjects)
+			{
+				if(MakeFeaturePackNode.IsFeaturePack(Proj))
 				{
-					if(MakeFeaturePackNode.ShouldMakeFeaturePack(HostPlatform, Proj))
-					{
-						AddDependency(MakeFeaturePackNode.StaticGetFullName(HostPlatform, Proj));
-					}
+					AddDependency(MakeFeaturePackNode.StaticGetFullName(FeaturePackPlatform, Proj));
 				}
-            }
+			}
         }
 		public override bool IsSeparatePromotable()
 		{
@@ -5984,11 +6004,14 @@ public class GUBP : BuildCommand
 				AddNode(new SharedCookAggregateNode(this, HostPlatforms, NonCodeProjectNames, NonCodeFormalBuilds));
 			}
 			
-			foreach(BranchInfo.BranchUProject Project in Branch.AllProjects)
+			if(HostPlatform == MakeFeaturePackNode.GetDefaultBuildPlatform(this))
 			{
-				if(MakeFeaturePackNode.ShouldMakeFeaturePack(HostPlatform, Project))
+				foreach(BranchInfo.BranchUProject Project in Branch.AllProjects)
 				{
-					AddNode(new MakeFeaturePackNode(this, HostPlatform, Project));
+					if(MakeFeaturePackNode.IsFeaturePack(Project))
+					{
+						AddNode(new MakeFeaturePackNode(HostPlatform, Project));
+					}
 				}
 			}
 
