Move GetEditorCommandletExe into CommandUtils, so that any commandlet code can use it.

diff --git a/Engine/Source/Programs/AutomationTool/CommandletUtils.cs b/Engine/Source/Programs/AutomationTool/CommandletUtils.cs
--- a/Engine/Source/Programs/AutomationTool/CommandletUtils.cs
+++ b/Engine/Source/Programs/AutomationTool/CommandletUtils.cs
@@ -180,6 +180,25 @@ namespace AutomationTool
 			}
 		}
 
+		/// <summary>
+		/// Returns the default path of the editor executable to use for running commandlets.
+		/// </summary>
+		/// <param name="BuildRoot">Root directory for the build</param>
+		/// <param name="HostPlatform">Platform to get the executable for</param>
+		/// <returns>Path to the editor executable</returns>
+		public static string GetEditorCommandletExe(string BuildRoot, UnrealBuildTool.UnrealTargetPlatform HostPlatform)
+		{
+			switch(HostPlatform)
+			{
+				case UnrealBuildTool.UnrealTargetPlatform.Mac:
+					return CommandUtils.CombinePaths(BuildRoot, "Engine/Binaries/Mac/UE4Editor.app/Contents/MacOS/UE4Editor");
+				case UnrealBuildTool.UnrealTargetPlatform.Win64:
+					return CommandUtils.CombinePaths(BuildRoot, "Engine/Binaries/Win64/UE4Editor-Cmd.exe");
+				default:
+					throw new NotImplementedException();
+			}
+		}
+
 		/// <summary>
 		/// Combines a list of parameters into a single commandline param separated with '+':
 		/// For example: Map1+Map2+Map3
