Add CommandUtils.GetGenericPlatformName() to return the generic name for a given target platform (eg. "Windows" on Win32/Win64).

diff --git a/Engine/Source/Programs/AutomationTool/CommandUtils.cs b/Engine/Source/Programs/AutomationTool/CommandUtils.cs
--- a/Engine/Source/Programs/AutomationTool/CommandUtils.cs
+++ b/Engine/Source/Programs/AutomationTool/CommandUtils.cs
@@ -2042,6 +2042,22 @@ namespace AutomationTool
                 Log(System.Diagnostics.TraceEventType.Warning, " Exception was {0}", LogUtils.FormatException(Ex));
             }
         }
+
+		/// <summary>
+		/// Returns the generic name for a given platform (eg. "Windows" for Win32/Win64)
+		/// </summary>
+		/// <param name="Platform">Specific platform</param>
+		public static string GetGenericPlatformName(UnrealBuildTool.UnrealTargetPlatform Platform)
+		{
+			if(Platform == UnrealTargetPlatform.Win32 || Platform == UnrealTargetPlatform.Win64)
+			{
+				return "Windows";
+			}
+			else
+			{
+				return Enum.GetName(typeof(UnrealBuildTool.UnrealTargetPlatform), Platform);
+			}
+		}
 	}
 
 
