build script, save node duration to job step as well as workflow

diff --git a/Engine/Source/Programs/AutomationTool/Scripts/GUBP.Automation.cs b/Engine/Source/Programs/AutomationTool/Scripts/GUBP.Automation.cs
--- a/Engine/Source/Programs/AutomationTool/Scripts/GUBP.Automation.cs
+++ b/Engine/Source/Programs/AutomationTool/Scripts/GUBP.Automation.cs
@@ -5134,6 +5134,7 @@ public class GUBP : BuildCommand
 		{
 			Log("Updating duration prop for node {0}", NodeToDo);
 			RunECTool(String.Format("setProperty \"/myWorkflow/NodeDuration/{0}\" \"{1}\"", NodeToDo, BuildDuration.ToString()));
+			RunECTool(String.Format("setProperty \"/myJobStep/NodeDuration\" \"{0}\"", BuildDuration.ToString()));
 		}
 		catch (Exception Ex)
 		{
