Refactoring Rocket samples to remove dependencies on RocketBuild nodes. Also separates out the case where we want to build samples only, with a node whose build products are just fetched from the network.

diff --git a/Engine/Source/Programs/AutomationTool/Scripts/GUBP.Automation.cs b/Engine/Source/Programs/AutomationTool/Scripts/GUBP.Automation.cs
--- a/Engine/Source/Programs/AutomationTool/Scripts/GUBP.Automation.cs
+++ b/Engine/Source/Programs/AutomationTool/Scripts/GUBP.Automation.cs
@@ -114,10 +114,7 @@ public class GUBP : BuildCommand
     public bool bFake = false;
     public static bool bNoIOSOnPC = false;
     public static bool bBuildRocket = false;
-	public static bool bPostedRocketBuild = false;
 	public static bool bBuildOnlySamples = false;
-	public static List<string> Samples = new List<string>();
-	public static string RocketBuild;	
     public static bool bForceIncrementalCompile = false;
     public static string ECProject = null;
     public string EmailHint;
@@ -152,10 +149,6 @@ public class GUBP : BuildCommand
         if (bBuildRocket)
         {
             Args = " -NoSimplygon " + (bUseRocketInsteadOfBuildRocket ? "-Rocket" : "-BuildRocket");
-			if(bPostedRocketBuild && bUseRocketInsteadOfBuildRocket)
-			{
-				Args += " -PostedRocket";
-			}
         }
         return Args;
     }
@@ -5243,16 +5236,7 @@ public class GUBP : BuildCommand
 		}
 
 		bBuildRocket = ParseParam("BuildRocket");		
-		string Sample = ParseParamValue("Sample");		
-		if (Sample != null)
-		{
-			Samples = new List<string>(Sample.Split('+'));
-			bBuildOnlySamples = true;
-		}
-		else
-		{
-			bBuildOnlySamples = ParseParam("OnlySamples");
-		}
+		bBuildOnlySamples = ParseParam("OnlySamples");
         bForceIncrementalCompile = ParseParam("ForceIncrementalCompile");
         bool bNoAutomatedTesting = ParseParam("NoAutomatedTesting") || BranchOptions.bNoAutomatedTesting;		
         StoreName = ParseParamValue("Store");
@@ -5261,15 +5245,6 @@ public class GUBP : BuildCommand
         if (bBuildRocket)
         {
             StoreSuffix = StoreSuffix + "-Rkt";
-			if (ParseParam("UsePostedRocket"))
-			{
-				bPostedRocketBuild = true;
-				RocketBuild = ParseParamValue("RocketBuild");
-				if (RocketBuild == null)
-				{
-					RocketBuild = FEngineVersionSupport.FromVersionFile(CommandUtils.CombinePaths(CommandUtils.CmdEnv.LocalRoot, @"Engine\Source\Runtime\Launch\Resources\Version.h")).ToString();
-				}
-			}
         }
 		
         if (bPreflightBuild)
@@ -7477,12 +7452,6 @@ public class GUBP : BuildCommand
                     }
                     else
                     {                        
-						if(GUBPNodes[NodeToDo].IsRocketSample() && bPostedRocketBuild)
-						{
-
-							Log("***** Retrieving Build {0} for {1}", RocketBuild, GUBPNodes[NodeToDo].NodeHostPlatform());
-							RetrieveFromPermanentStorage(CmdEnv, "Rocket\\Automated", RocketBuild, GUBPNodes[NodeToDo].NodeHostPlatform());
-						}
 						Log("***** Building GUBP Node {0} for {1}", NodeToDo, NodeStoreName);
 						var StartTime = DateTime.UtcNow;
                         GUBPNodes[NodeToDo].DoBuild(this);
diff --git a/Engine/Source/Programs/AutomationTool/UE4Build.cs b/Engine/Source/Programs/AutomationTool/UE4Build.cs
--- a/Engine/Source/Programs/AutomationTool/UE4Build.cs
+++ b/Engine/Source/Programs/AutomationTool/UE4Build.cs
@@ -100,8 +100,7 @@ namespace AutomationTool
 			}
 		}
 
-		public static string BaseUBTDirectory = "";
-		
+	
 		/// True if UBT is compiled and ready to build!
 		private bool bIsUBTReady = false;
 		
@@ -164,17 +163,7 @@ namespace AutomationTool
 			{
 				AddArgs += " -forcedebuginfo";
 			}
-			if(AddArgs.Contains("PostedRocket"))
-			{
-				if (UnrealBuildTool.BuildHostPlatform.Current.Platform == UnrealBuildTool.UnrealTargetPlatform.Win64)
-				{
-					BaseUBTDirectory = @"Rocket/TempInst/Windows";
-				}
-				if (UnrealBuildTool.BuildHostPlatform.Current.Platform == UnrealBuildTool.UnrealTargetPlatform.Mac)
-				{
-					BaseUBTDirectory = @"Rocket/TempInst/Mac";
-				}
-			}
+
 			PrepareUBT();
 
             string UBTManifest = GetUBTManifest(UprojectPath, AddArgs);
@@ -313,17 +302,6 @@ namespace AutomationTool
 			{
 				AddArgs += " -noxge";
 			}
-			if(AddArgs.Contains("PostedRocket"))
-			{
-				if (UnrealBuildTool.BuildHostPlatform.Current.Platform == UnrealBuildTool.UnrealTargetPlatform.Win64)
-				{
-					BaseUBTDirectory = @"Rocket/TempInst/Windows";
-				}
-				if (UnrealBuildTool.BuildHostPlatform.Current.Platform == UnrealBuildTool.UnrealTargetPlatform.Mac)
-				{
-					BaseUBTDirectory = @"Rocket/TempInst/Mac";
-				}
-			}
 
 			PrepareUBT();
 
@@ -726,7 +704,7 @@ namespace AutomationTool
 		public List<string> FindXGEFiles()
 		{
 			var Result = new List<string>();
-			var Root = CombinePaths(CmdEnv.LocalRoot, BaseUBTDirectory, @"\Engine\Intermediate\Build");			
+			var Root = CombinePaths(CmdEnv.LocalRoot, @"\Engine\Intermediate\Build");			
 			Result.AddRange(FindFiles_NoExceptions("*.xge.xml", false, Root));
 			Result.Sort();
 			return Result;
@@ -1508,7 +1486,7 @@ namespace AutomationTool
 
 		public static string GetUBTExecutable()
 		{
-			return CommandUtils.CombinePaths(CmdEnv.LocalRoot, BaseUBTDirectory, @"Engine/Binaries/DotNET/UnrealBuildTool.exe");
+			return CommandUtils.CombinePaths(CmdEnv.LocalRoot, @"Engine/Binaries/DotNET/UnrealBuildTool.exe");
 		}
 
 		public string UBTExecutable
