build script, store temp storage duration as ec property

diff --git a/Engine/Source/Programs/AutomationTool/Scripts/GUBP.Automation.cs b/Engine/Source/Programs/AutomationTool/Scripts/GUBP.Automation.cs
--- a/Engine/Source/Programs/AutomationTool/Scripts/GUBP.Automation.cs
+++ b/Engine/Source/Programs/AutomationTool/Scripts/GUBP.Automation.cs
@@ -4545,8 +4545,8 @@ public class GUBP : BuildCommand
         }
         string RecordOfSuccess = CombinePaths(CommandUtils.CmdEnv.LocalRoot, "Engine", "Saved", "Logs", NodeToDo + Suffix +".log");
         CreateDirectory(Path.GetDirectoryName(RecordOfSuccess));
-        WriteAllText(RecordOfSuccess, Contents);
-        StoreToTempStorage(CmdEnv, NodeStoreName + Suffix, new List<string> { RecordOfSuccess }, !bSaveSharedTempStorage, GameNameIfAny);
+        WriteAllText(RecordOfSuccess, Contents);		
+		StoreToTempStorage(CmdEnv, NodeStoreName + Suffix, new List<string> { RecordOfSuccess }, !bSaveSharedTempStorage, GameNameIfAny);		
     }
 	string GetPropertyFromStep(string PropertyPath)
 	{
@@ -7052,9 +7052,9 @@ public class GUBP : BuildCommand
                     
                     string NodeParentPath = ParentPath;
 					string PreconditionParentPath;
-					if (GUBPNodes[NodeToDo].GetFullName().Contains("MakeBuild") && GUBPNodes[NodeToDo].FullNamesOfPseudosependencies.Contains(SharedLabelPromotableNode.StaticGetFullName(false)) && !bGraphSubset)
+					if (GUBPNodes[NodeToDo].GetFullName().Contains("MakeBuild") && GUBPNodes[NodeToDo].FullNamesOfPseudosependencies.Contains(WaitForFormalUserInput.StaticGetFullName()) && !bGraphSubset)
                     {
-						RemovePseudodependencyFromNode(NodeToDo, SharedLabelPromotableNode.StaticGetFullName(false));
+						RemovePseudodependencyFromNode(NodeToDo, WaitForFormalUserInput.StaticGetFullName());
 						PreconditionParentPath = GetPropertyFromStep("/myWorkflow/ParentJob");
 						UncompletedEcDeps = GetECDependencies(NodeToDo);
 					}
@@ -7460,8 +7460,15 @@ public class GUBP : BuildCommand
                     }
                     if (!GUBPNodes[NodeToDo].IsAggregate())
                     {
+						var StoreDuration = 0.0;
+						var StartTime = DateTime.UtcNow;
                         StoreToTempStorage(CmdEnv, NodeStoreName, GUBPNodes[NodeToDo].BuildProducts, !bSaveSharedTempStorage, GameNameIfAny, StorageRootIfAny);
-
+						StoreDuration = (DateTime.UtcNow - StartTime).TotalMilliseconds / 1000;
+						Log("Took {0} seconds to store build products", StoreDuration);
+						if (IsBuildMachine)
+						{
+							RunECTool(String.Format("setProperty \"/myJobStep/StoreDuration\" \"{0}\"", StoreDuration.ToString()));
+						}
                         if (ParseParam("StompCheck"))
                         {
                             foreach (var Dep in GUBPNodes[NodeToDo].AllDependencies)
