Remove IsRocketSample() callback. Not used any more.

diff --git a/Engine/Source/Programs/AutomationTool/Scripts/GUBP.Automation.cs b/Engine/Source/Programs/AutomationTool/Scripts/GUBP.Automation.cs
--- a/Engine/Source/Programs/AutomationTool/Scripts/GUBP.Automation.cs
+++ b/Engine/Source/Programs/AutomationTool/Scripts/GUBP.Automation.cs
@@ -422,10 +422,6 @@ public class GUBP : BuildCommand
 		{
 			return false;
 		}
-		public virtual bool IsRocketSample()
-		{
-			return false;
-		}
 		public virtual string NodeHostPlatform()
 		{
 			return "";
