[INTEGRATE] Change 2449441 by [REDACTED_PERSON][REDACTED_HANDLE] on 2015/02/17 18:57:48

	Only stage CEF3 if bUsesCEF3 is set to true in the game's target.cs file. Fix for UE-9983.

diff --git a/Engine/Source/Programs/AutomationTool/ProjectParams.cs b/Engine/Source/Programs/AutomationTool/ProjectParams.cs
--- a/Engine/Source/Programs/AutomationTool/ProjectParams.cs
+++ b/Engine/Source/Programs/AutomationTool/ProjectParams.cs
@@ -277,6 +277,7 @@ namespace AutomationTool
             this.Compressed = InParams.Compressed;
             this.UseDebugParamForEditorExe = InParams.UseDebugParamForEditorExe;
             this.bUsesSteam = InParams.bUsesSteam;
+			this.bUsesCEF3 = InParams.bUsesCEF3;
 			this.bUsesSlate = InParams.bUsesSlate;
 			this.bUsesSlateEditorStyle = InParams.bUsesSlateEditorStyle;
             this.bDebugBuildsActuallyUseDebugCRT = InParams.bDebugBuildsActuallyUseDebugCRT;
@@ -1099,6 +1100,11 @@ namespace AutomationTool
         /// </summary>
         public bool bUsesSteam;
 
+        /// <summary>
+        /// Whether the project uses CEF3
+        /// </summary>
+        public bool bUsesCEF3;
+
 		/// <summary>
 		/// Whether the project uses visual Slate UI (as opposed to the low level windowing/messaging which is alway used)
 		/// </summary>
@@ -1332,6 +1338,7 @@ namespace AutomationTool
 			var Properties = ProjectUtils.GetProjectProperties(RawProjectPath);
 
 			bUsesSteam = Properties.bUsesSteam;
+			bUsesCEF3 = Properties.bUsesCEF3;
 			bUsesSlate = Properties.bUsesSlate;
 			bUsesSlateEditorStyle = Properties.bUsesSlateEditorStyle;
             bDebugBuildsActuallyUseDebugCRT = Properties.bDebugBuildsActuallyUseDebugCRT;
@@ -1436,6 +1443,7 @@ namespace AutomationTool
 						bUsesSlate = TargetData.Rules.bUsesSlate;
 						bUsesSlateEditorStyle = TargetData.Rules.bUsesSlateEditorStyle;
 						bUsesSteam = TargetData.Rules.bUsesSteam;
+						bUsesCEF3 = TargetData.Rules.bUsesCEF3;
 						ProjectType = ValidTarget;
 						break;
 					}
@@ -1462,6 +1470,7 @@ namespace AutomationTool
 				bUsesSlate = TargetData.Rules.bUsesSlate;
 				bUsesSlateEditorStyle = TargetData.Rules.bUsesSlateEditorStyle;
 				bUsesSteam = TargetData.Rules.bUsesSteam;
+				bUsesCEF3 = TargetData.Rules.bUsesCEF3;
 				ProjectType = TargetRules.TargetType.Program;
 				ProgramTarget = TargetData.TargetName;
 				GameTarget = TargetData.TargetName;
@@ -1987,6 +1996,7 @@ namespace AutomationTool
 				CommandUtils.Log("SkipStage={0}", SkipStage);
 				CommandUtils.Log("Stage={0}", Stage);
 				CommandUtils.Log("bUsesSteam={0}", bUsesSteam);
+				CommandUtils.Log("bUsesCEF3={0}", bUsesCEF3);
 				CommandUtils.Log("bUsesSlate={0}", bUsesSlate);
                 CommandUtils.Log("bDebugBuildsActuallyUseDebugCRT={0}", bDebugBuildsActuallyUseDebugCRT);
 				CommandUtils.Log("Project Params **************");
diff --git a/Engine/Source/Programs/AutomationTool/ProjectUtils.cs b/Engine/Source/Programs/AutomationTool/ProjectUtils.cs
--- a/Engine/Source/Programs/AutomationTool/ProjectUtils.cs
+++ b/Engine/Source/Programs/AutomationTool/ProjectUtils.cs
@@ -36,6 +36,11 @@ namespace AutomationTool
 		/// </summary>
 		public bool bUsesSteam;
 
+		/// <summary>
+		/// Whether the project uses CEF3
+		/// </summary>
+		public bool bUsesCEF3;
+
 		/// <summary>
 		/// Whether the project uses visual Slate UI (as opposed to the low level windowing/messaging which is always used)
 		/// </summary>
@@ -398,6 +403,7 @@ namespace AutomationTool
 					}
 
 					Properties.bUsesSteam |= Rules.bUsesSteam;
+					Properties.bUsesCEF3 |= Rules.bUsesCEF3;
 					Properties.bUsesSlate |= Rules.bUsesSlate;
                     Properties.bDebugBuildsActuallyUseDebugCRT |= Rules.bDebugBuildsActuallyUseDebugCRT;
 					Properties.bUsesSlateEditorStyle |= Rules.bUsesSlateEditorStyle;
@@ -523,6 +529,7 @@ namespace AutomationTool
                 CommandUtils.Log("      FilePath          : " + FilePath);
                 CommandUtils.Log("      bIsCodeBasedProject  : " + (Properties.bIsCodeBasedProject ? "YES" : "NO"));
                 CommandUtils.Log("      bUsesSteam  : " + (Properties.bUsesSteam ? "YES" : "NO"));
+                CommandUtils.Log("      bUsesCEF3   : " + (Properties.bUsesCEF3 ? "YES" : "NO"));
                 CommandUtils.Log("      bUsesSlate  : " + (Properties.bUsesSlate ? "YES" : "NO"));
                 foreach (var HostPlatform in InHostPlatforms)
                 {
@@ -539,6 +546,7 @@ namespace AutomationTool
                         CommandUtils.Log("            TargetName          : " + ThisTarget.Value.TargetName);
                         CommandUtils.Log("              Type          : " + ThisTarget.Key.ToString());
                         CommandUtils.Log("              bUsesSteam  : " + (ThisTarget.Value.Rules.bUsesSteam ? "YES" : "NO"));
+                        CommandUtils.Log("              bUsesCEF3   : " + (ThisTarget.Value.Rules.bUsesCEF3 ? "YES" : "NO"));
                         CommandUtils.Log("              bUsesSlate  : " + (ThisTarget.Value.Rules.bUsesSlate ? "YES" : "NO"));
                         if (Array.IndexOf(MonolithicKinds, ThisTarget.Key) >= 0)
                         {
diff --git a/Engine/Source/Programs/AutomationTool/Win/WinPlatform.Automation.cs b/Engine/Source/Programs/AutomationTool/Win/WinPlatform.Automation.cs
--- a/Engine/Source/Programs/AutomationTool/Win/WinPlatform.Automation.cs
+++ b/Engine/Source/Programs/AutomationTool/Win/WinPlatform.Automation.cs
@@ -89,8 +89,11 @@ public abstract class BaseWinPlatform : Platform
 		SC.StageFiles(StagedFileType.NonUFS, CombinePaths(SC.ProjectRoot, "Content/Splash"), "Splash.bmp", false, null, null, true);
 
 		// CEF3 files
-		SC.StageFiles(StagedFileType.NonUFS, CombinePaths(SC.LocalRoot, "Engine/Binaries/ThirdParty/CEF3", SC.PlatformDir), "*", true, null, null, true);
-		StageExecutable("exe", SC, CommandUtils.CombinePaths(SC.LocalRoot, "Engine/Binaries", SC.PlatformDir), "UnrealCEFSubProcess.", false, null, null, true);
+		if (Params.bUsesCEF3)
+		{
+			SC.StageFiles(StagedFileType.NonUFS, CombinePaths(SC.LocalRoot, "Engine/Binaries/ThirdParty/CEF3", SC.PlatformDir), "*", true, null, null, true);
+			StageExecutable("exe", SC, CommandUtils.CombinePaths(SC.LocalRoot, "Engine/Binaries", SC.PlatformDir), "UnrealCEFSubProcess.", false, null, null, true);
+		}
 
         if (Params.StageNonMonolithic)
         {
diff --git a/Engine/Source/Programs/UnrealBuildTool/System/RulesCompiler.cs b/Engine/Source/Programs/UnrealBuildTool/System/RulesCompiler.cs
--- a/Engine/Source/Programs/UnrealBuildTool/System/RulesCompiler.cs
+++ b/Engine/Source/Programs/UnrealBuildTool/System/RulesCompiler.cs
@@ -357,6 +357,11 @@ namespace UnrealBuildTool
         /// </summary>
         public bool bUsesSteam;
 
+        /// <summary>
+        /// Whether the target uses CEF3
+        /// </summary>
+        public bool bUsesCEF3;
+
 		/// <summary>
 		/// Whether the project uses visual Slate UI (as opposed to the low level windowing/messaging which is always used)
 		/// </summary>
