[INTEGRATE] Change 2449448 by [REDACTED_PERSON][REDACTED_HANDLE] on 2015/02/17 19:00:41

	Conditionally stage CEF3 for Mac too. (UE-9983)

diff --git a/Engine/Source/Programs/AutomationTool/Mac/MacPlatform.Automation.cs b/Engine/Source/Programs/AutomationTool/Mac/MacPlatform.Automation.cs
--- a/Engine/Source/Programs/AutomationTool/Mac/MacPlatform.Automation.cs
+++ b/Engine/Source/Programs/AutomationTool/Mac/MacPlatform.Automation.cs
@@ -81,9 +81,12 @@ public class MacPlatform : Platform
 		SC.StageFiles(StagedFileType.NonUFS, CombinePaths(SC.ProjectRoot, "Content/Splash"), "Splash.bmp", false, null, null, true);
 
 		// CEF3 files
-		SC.StageFiles(StagedFileType.NonUFS, CombinePaths(SC.LocalRoot, "Engine/Binaries/ThirdParty/CEF3/Mac/"), "*", true, null, null, true);
-		string UnrealCEFSubProcessPath = CombinePaths("Engine/Binaries", SC.PlatformDir, "UnrealCEFSubProcess.app");
-		StageAppBundle(SC, CombinePaths(SC.LocalRoot, "Engine/Binaries", SC.PlatformDir, "UnrealCEFSubProcess.app"), UnrealCEFSubProcessPath);
+		if(Params.bUsesCEF3)
+		{
+			SC.StageFiles(StagedFileType.NonUFS, CombinePaths(SC.LocalRoot, "Engine/Binaries/ThirdParty/CEF3/Mac/"), "*", true, null, null, true);
+			string UnrealCEFSubProcessPath = CombinePaths("Engine/Binaries", SC.PlatformDir, "UnrealCEFSubProcess.app");
+			StageAppBundle(SC, CombinePaths(SC.LocalRoot, "Engine/Binaries", SC.PlatformDir, "UnrealCEFSubProcess.app"), UnrealCEFSubProcessPath);
+		}
 	}
 
 	public override void Package(ProjectParams Params, DeploymentContext SC, int WorkingCL)
