Assign unique guid in Desktop.Automation.csproj; looks like it was copied from Android.Automation.csproj. Prevents prompt to save every time, and it being dumped at the bottom of the solution.

diff --git a/Engine/Source/Programs/AutomationTool/Desktop/Desktop.Automation.csproj b/Engine/Source/Programs/AutomationTool/Desktop/Desktop.Automation.csproj
--- a/Engine/Source/Programs/AutomationTool/Desktop/Desktop.Automation.csproj
+++ b/Engine/Source/Programs/AutomationTool/Desktop/Desktop.Automation.csproj
@@ -1,10 +1,10 @@
-<?xml version="1.0" encoding="utf-8"?>
+﻿<?xml version="1.0" encoding="utf-8"?>
 <Project ToolsVersion="4.0" DefaultTargets="Build" xmlns="http://schemas.microsoft.com/developer/msbuild/2003">
   <Import Project="$(MSBuildExtensionsPath)\$(MSBuildToolsVersion)\Microsoft.Common.props" Condition="Exists('$(MSBuildExtensionsPath)\$(MSBuildToolsVersion)\Microsoft.Common.props')" />
   <PropertyGroup>
     <Configuration Condition=" '$(Configuration)' == '' ">Debug</Configuration>
     <Platform Condition=" '$(Platform)' == '' ">AnyCPU</Platform>
-    <ProjectGuid>{32E3A47E-359A-4F78-AE6C-703B87DE0B97}</ProjectGuid>
+    <ProjectGuid>{A9CD1A06-764A-4C53-86CE-2F34B94F0283}</ProjectGuid>
     <OutputType>Library</OutputType>
     <AppDesignerFolder>Properties</AppDesignerFolder>
     <RootNamespace>Desktop.Automation</RootNamespace>
@@ -72,4 +72,4 @@
     </ProjectReference>
   </ItemGroup>
   <Import Project="$(MSBuildToolsPath)\Microsoft.CSharp.targets" />
-</Project>
+</Project>
\ No newline at end of file
