build script, batch sign files rather than sign individually

diff --git a/Engine/Source/Programs/AutomationTool/CommandUtils.cs b/Engine/Source/Programs/AutomationTool/CommandUtils.cs
--- a/Engine/Source/Programs/AutomationTool/CommandUtils.cs
+++ b/Engine/Source/Programs/AutomationTool/CommandUtils.cs
@@ -2318,13 +2318,15 @@ namespace AutomationTool
 				}
 				else
 				{
+					List<string> FilesToSign = new List<string>();
 					foreach (string File in Files)
 					{
 						if (!(Path.GetDirectoryName(File).Replace("\\", "/")).Contains("Binaries/XboxOne"))
 						{
-							SignSingleExecutableIfEXEOrDLL(File);
-						}
+							FilesToSign.Add(File);
+						}						
 					}
+					SignMultipleFilesIfEXEOrDLL(FilesToSign);
 				}
 			}
 			else
@@ -2332,6 +2334,115 @@ namespace AutomationTool
 				CommandUtils.Log("Skipping signing {0} files due to -nosign.", Files.Count);
 			}
 		}
+
+		public static void SignListFilesIfEXEOrDLL(string FilesToSign)
+		{
+			string SignToolName = null;
+			if (WindowsPlatform.Compiler == WindowsCompiler.VisualStudio2013)
+			{
+				SignToolName = "C:/Program Files (x86)/Windows Kits/8.1/bin/x86/SignTool.exe";
+			}
+			else if (WindowsPlatform.Compiler == WindowsCompiler.VisualStudio2012)
+			{
+				SignToolName = "C:/Program Files (x86)/Windows Kits/8.0/bin/x86/SignTool.exe";
+			}
+
+
+			if (!File.Exists(SignToolName))
+			{
+				throw new AutomationException("SignTool not found at '{0}' (are you missing the Windows SDK?)", SignToolName);
+			}
+
+			// Code sign the executable
+			string TimestampServer = "http://timestamp.verisign.com/scripts/timestamp.dll";
+
+			string SpecificStoreArg = bUseMachineStoreInsteadOfUserStore ? " /sm" : "";	
+			//@TODO: Verbosity choosing
+			//  /v will spew lots of info
+			//  /q does nothing on success and minimal output on failure
+			string CodeSignArgs = String.Format("sign{0} /a /n \"{1}\" /t {2} /d /v {3}", SpecificStoreArg, SigningIdentity, TimestampServer, FilesToSign);
+
+			DateTime StartTime = DateTime.Now;
+
+			int NumTrials = 0;
+			for (; ; )
+			{
+				ProcessResult Result = CommandUtils.Run(SignToolName, CodeSignArgs, null, CommandUtils.ERunOptions.AllowSpew);
+				++NumTrials;
+
+				if (Result.ExitCode != 1)
+				{
+					if (Result.ExitCode == 2)
+					{
+						CommandUtils.Log(TraceEventType.Error, String.Format("Signtool returned a warning."));
+					}
+					// Success!
+					break;
+				}
+				else
+				{
+					// Keep retrying until we run out of time
+					TimeSpan RunTime = DateTime.Now - StartTime;
+					if (RunTime > CodeSignTimeOut)
+					{
+						throw new AutomationException("Failed to sign executables {0} times over a period of {1}", NumTrials, RunTime);
+					}
+				}
+			}
+		}
+
+		public static void SignMultipleFilesIfEXEOrDLL(List<string> Files, bool bIgnoreExtension = false)
+		{
+			if (UnrealBuildTool.Utils.IsRunningOnMono)
+			{
+				CommandUtils.Log(TraceEventType.Information, String.Format("Can't sign we are running under mono."));
+				return;
+			}
+			List<string> FinalFiles = new List<string>();
+			foreach (string Filename in Files)
+			{
+				// Make sure the file isn't read-only
+				FileInfo TargetFileInfo = new FileInfo(Filename);
+
+				// Executable extensions
+				List<string> Extensions = new List<string>();
+				Extensions.Add(".dll");
+				Extensions.Add(".exe");
+
+				bool IsExecutable = bIgnoreExtension;
+
+				foreach (var Ext in Extensions)
+				{
+					if (TargetFileInfo.FullName.EndsWith(Ext, StringComparison.InvariantCultureIgnoreCase))
+					{
+						IsExecutable = true;
+						break;
+					}
+				}
+				if (IsExecutable && CommandUtils.FileExists(Filename))
+				{
+					FinalFiles.Add(Filename);
+				}
+			}			
+
+			StringBuilder FilesToSignBuilder = new StringBuilder();
+			List<string> FinalListSignStrings = new List<string>();
+			foreach(string File in FinalFiles)
+			{
+				FilesToSignBuilder.Append("\"" + File + "\" ");				
+				if(FilesToSignBuilder.Length > 1900)
+				{
+					string AddFilesToFinalList = FilesToSignBuilder.ToString();
+					FinalListSignStrings.Add(AddFilesToFinalList);
+					FilesToSignBuilder.Clear();
+				}
+			}
+			foreach(string FilesToSign in FinalListSignStrings)
+			{
+				SignListFilesIfEXEOrDLL(FilesToSign);
+			}
+
+		}
 	}
 
 }
