Move FileFilter into AutomationTool project, so it can be used elsewhere.

diff --git a/Engine/Source/Programs/AutomationTool/AutomationTool.csproj b/Engine/Source/Programs/AutomationTool/AutomationTool.csproj
--- a/Engine/Source/Programs/AutomationTool/AutomationTool.csproj
+++ b/Engine/Source/Programs/AutomationTool/AutomationTool.csproj
@@ -104,6 +104,7 @@
     <Compile Include="CommandletUtils.cs" />
     <Compile Include="DeploymentContext.cs" />
     <Compile Include="Distiller.cs" />
+    <Compile Include="FileFilter.cs" />
     <Compile Include="HelpUtils.cs" />
     <Compile Include="HostPlatform.cs" />
     <Compile Include="LinuxHostPlatform.cs" />
diff --git a/Engine/Source/Programs/AutomationTool/FileFilter.cs b/Engine/Source/Programs/AutomationTool/FileFilter.cs
new file mode 100644
--- /dev/null
+++ b/Engine/Source/Programs/AutomationTool/FileFilter.cs
@@ -0,0 +1,144 @@
+﻿// Copyright 1998-2015 Epic Games, Inc. All Rights Reserved.
+
+using System;
+using System.Collections.Generic;
+using System.Diagnostics;
+using System.IO;
+using System.Linq;
+using System.Text;
+using System.Text.RegularExpressions;
+
+/// <summary>
+/// Stores a set of filter rules, which can be used to filter files in or out of a set
+/// </summary>
+public class FileFilter
+{
+	List<Tuple<Regex, bool>> Rules;
+
+	/// <summary>
+	/// Default constructor
+	/// </summary>
+	public FileFilter()
+	{
+		Rules = new List<Tuple<Regex,bool>>();
+	}
+
+	/// <summary>
+	/// Copy constructor
+	/// </summary>
+	/// <param name="Other">Filter to copy from</param>
+	public FileFilter(FileFilter Other)
+	{
+		Rules = new List<Tuple<Regex,bool>>(Other.Rules);
+	}
+
+	/// <summary>
+	/// Constructs a file filter from a p4-style filespec. Exclude lines are prefixed with a - character.
+	/// </summary>
+	/// <param name="Lines">Lines to parse rules from</param>
+	public FileFilter(IEnumerable<string> Lines) : this()
+	{
+		foreach(string Line in Lines)
+		{
+			if(Line.StartsWith("-"))
+			{
+				Exclude(Line.Substring(1).TrimStart());
+			}
+			else
+			{
+				Include(Line);
+			}
+		}
+	}
+
+	/// <summary>
+	/// Adds an include rule to the filter
+	/// </summary>
+	/// <param name="Pattern">Pattern to match. See CreateRegex() for details.</param>
+	public void Include(string Pattern)
+	{
+		Rules.Add(new Tuple<Regex,bool>(CreateRegex(Pattern), true));
+	}
+
+	/// <summary>
+	/// Adds several exclude rules to the filter
+	/// </summary>
+	/// <param name="Patterns">Patterns to match. See CreateRegex() for details.</param>
+	public void Include(params string[] Patterns)
+	{
+		foreach(string Pattern in Patterns)
+		{
+			Include(Pattern);
+		}
+	}
+
+	/// <summary>
+	/// Adds an exclude rule to the filter
+	/// </summary>
+	/// <param name="Pattern">Mask to match. See CreateRegex() for details.</param>
+	public void Exclude(string Pattern)
+	{
+		Rules.Add(new Tuple<Regex, bool>(CreateRegex(Pattern), false));
+	}
+
+	/// <summary>
+	/// Adds several exclude rules to the filter
+	/// </summary>
+	/// <param name="Patterns">Patterns to match. See CreateRegex() for details.</param>
+	public void Exclude(params string[] Patterns)
+	{
+		foreach(string Pattern in Patterns)
+		{
+			Exclude(Pattern);
+		}
+	}
+
+	/// <summary>
+	/// Determines whether the given file matches the filter
+	/// </summary>
+	/// <param name="File">File to match</param>
+	/// <returns>True if the file passes the filter</returns>
+	public bool Matches(string FileName)
+	{
+		bool bIsMatch = false;
+		foreach(Tuple<Regex, bool> Rule in Rules)
+		{
+			if(bIsMatch != Rule.Item2 && Rule.Item1.IsMatch(FileName))
+			{
+				bIsMatch = Rule.Item2;
+			}
+		}
+		return bIsMatch;
+	}
+
+	/// <summary>
+	/// Creates a regex to match the given p4-style pattern, which matches forward and back slashes as path separators, and match ?, *, and ... as wildcards.
+	/// </summary>
+	/// <param name="Pattern">Pattern to convert into a regex</param>
+	/// <returns>The resulting regex</returns>
+	public static Regex CreateRegex(string Pattern)
+	{
+		// Convert the pattern into a p4-style path specification.
+		string FinalPattern = "^" + Regex.Escape(Pattern.Replace('\\', '/')) + "$";
+		FinalPattern = FinalPattern.Replace("/", "[/\\\\]");
+		FinalPattern = FinalPattern.Replace("\\?", ".");
+		FinalPattern = FinalPattern.Replace("\\*\\.$", "[^/\\\\.]*$"); // No extension pattern: *.
+		FinalPattern = FinalPattern.Replace("\\*", "[^/\\\\]*");
+		FinalPattern = FinalPattern.Replace("\\.\\.\\.", ".*");
+
+		// Remove sentinels when we're just looking for substrings
+		if(FinalPattern.StartsWith("^.*"))
+		{
+			FinalPattern = FinalPattern.Substring(3);
+		}
+		if(FinalPattern.EndsWith(".*$"))
+		{
+			FinalPattern = FinalPattern.Substring(0, FinalPattern.Length - 3);
+		}
+
+		// Create the regex
+		return new Regex(FinalPattern, RegexOptions.Compiled | RegexOptions.IgnoreCase);
+	}
+}
+
+
