Fix path to HTML5LaunchHelper, and remove UnrealPak on Win64 from the HTML5 binaries list. It will not compile on Mac.

diff --git a/Engine/Source/Programs/AutomationTool/Scripts/BuildCommonTools.Automation.cs b/Engine/Source/Programs/AutomationTool/Scripts/BuildCommonTools.Automation.cs
--- a/Engine/Source/Programs/AutomationTool/Scripts/BuildCommonTools.Automation.cs
+++ b/Engine/Source/Programs/AutomationTool/Scripts/BuildCommonTools.Automation.cs
@@ -151,9 +151,7 @@ public class BuildCommonTools : BuildCommand
 		// HTML5 binaries
 		if (Platforms.Contains(UnrealBuildTool.UnrealTargetPlatform.HTML5))
 		{
-			Agenda.AddTarget("UnrealPak", UnrealBuildTool.UnrealTargetPlatform.Win64, UnrealBuildTool.UnrealTargetConfiguration.Development);
-			
-			Agenda.DotNetProjects.Add(@"Engine/Source/Programs/HTML5LaunchHelper/HTML5LaunchHelper.csproj");
+			Agenda.DotNetProjects.Add(@"Engine/Source/Programs/HTML5/HTML5LaunchHelper/HTML5LaunchHelper.csproj");
 			ExtraBuildProducts.Add(CommandUtils.CombinePaths(CommandUtils.CmdEnv.LocalRoot, @"Engine/Binaries/DotNET/HTML5/HTML5LaunchHelper.exe"));
 		}
 		
