Fix output path for HTML5LaunchHelper too.

diff --git a/Engine/Source/Programs/AutomationTool/Scripts/BuildCommonTools.Automation.cs b/Engine/Source/Programs/AutomationTool/Scripts/BuildCommonTools.Automation.cs
--- a/Engine/Source/Programs/AutomationTool/Scripts/BuildCommonTools.Automation.cs
+++ b/Engine/Source/Programs/AutomationTool/Scripts/BuildCommonTools.Automation.cs
@@ -152,7 +152,7 @@ public class BuildCommonTools : BuildCommand
 		if (Platforms.Contains(UnrealBuildTool.UnrealTargetPlatform.HTML5))
 		{
 			Agenda.DotNetProjects.Add(@"Engine/Source/Programs/HTML5/HTML5LaunchHelper/HTML5LaunchHelper.csproj");
-			ExtraBuildProducts.Add(CommandUtils.CombinePaths(CommandUtils.CmdEnv.LocalRoot, @"Engine/Binaries/DotNET/HTML5/HTML5LaunchHelper.exe"));
+			ExtraBuildProducts.Add(CommandUtils.CombinePaths(CommandUtils.CmdEnv.LocalRoot, @"Engine/Binaries/DotNET/HTML5LaunchHelper.exe"));
 		}
 		
 		return Agenda;
