UE-3968 - Project deployment can overwrite existing directories on device
- all projects are now deployed to a subdirectory of UE4Game when launching

Related to OBB support changes. (Futher changes to follow)
- File name for OBB is now of format main.1.<details> rather than main.00001.<details> due to Google Play file naming requirements.

diff --git a/Engine/Source/Programs/AutomationTool/Android/AndroidPlatform.Automation.cs b/Engine/Source/Programs/AutomationTool/Android/AndroidPlatform.Automation.cs
--- a/Engine/Source/Programs/AutomationTool/Android/AndroidPlatform.Automation.cs
+++ b/Engine/Source/Programs/AutomationTool/Android/AndroidPlatform.Automation.cs
@@ -15,6 +15,8 @@ public class AndroidPlatform : Platform
 {
 	private const int DeployMaxParallelCommands = 6;
 
+    private const string TargetAndroidLocation = "obb/";
+
 	public AndroidPlatform()
 		: base(UnrealTargetPlatform.Android)
 	{
@@ -75,7 +77,7 @@ public class AndroidPlatform : Platform
 		if (PackageVersion.Length > 0)
 		{
 			int IntVersion = int.Parse(PackageVersion);
-			PackageVersion = IntVersion.ToString("00000");
+			PackageVersion = IntVersion.ToString("0");
 		}
 
 		string ObbName = string.Format("main.{0}.{1}.obb", PackageVersion, PackageName);
@@ -90,7 +92,7 @@ public class AndroidPlatform : Platform
 	{
         string ObbName = GetFinalObbName(ApkName);
         string PackageName = GetPackageInfo(ApkName, false);
-		return "obb/" + PackageName + "/" + Path.GetFileName(ObbName);
+        return TargetAndroidLocation + PackageName + "/" + Path.GetFileName(ObbName);
 	}
 
     private static string GetStorageQueryCommand()
@@ -110,80 +112,6 @@ public class AndroidPlatform : Platform
 		return Path.Combine(Path.GetDirectoryName(ApkName), "Install_" + Params.ShortProjectName + "_" + Params.ClientConfigsToBuild[0].ToString() + Architecture + GPUArchitecture + (Utils.IsRunningOnMono ? ".command" : ".bat"));
 	}
 
-//     public static Dictionary<string, string> SetupEnv(string ProjectPath)
-//     {
-//         if(configCacheIni == null)
-//         {
-//             CommandUtils.LogWarning(Path.GetDirectoryName(ProjectPath));
-//             configCacheIni = new UnrealBuildTool.ConfigCacheIni("Engine", Path.GetDirectoryName(ProjectPath));
-// 
-//             string[] EnvVarNames = {"ANDROID_HOME", "NDKHOME", "ANT_HOME"};
-// 
-//             string path;
-//             if(configCacheIni.GetPath("/Script/AndroidPlatformEditor.AndroidSDKSettings", "SDKPath", out path))
-//             {
-//                 AndroidEnv.Add("ANDROID_HOME", path);
-//             }
-//             else // if we fail then try and get it from the real environment
-//             {
-//                 AndroidEnv.Add("ANDROID_HOME", Environment.GetEnvironmentVariable("ANDROID_HOME"));
-//             }
-// 
-//             if (configCacheIni.GetPath("/Script/AndroidPlatformEditor.AndroidSDKSettings", "NDKPath", out path))
-//             {
-//                 AndroidEnv.Add("NDKHOME", path);
-//             }
-//             else // if we fail then try and get it from the real environment
-//             {
-//                 AndroidEnv.Add("NDKHOME", Environment.GetEnvironmentVariable("NDKHOME"));
-//             }
-// 
-//             if (configCacheIni.GetPath("/Script/AndroidPlatformEditor.AndroidSDKSettings", "ANTPath", out path))
-//             {
-//                 AndroidEnv.Add("ANT_HOME", path);
-//             }
-//             else // if we fail then try and get it from the real environment
-//             {
-//                 AndroidEnv.Add("ANT_HOME", Environment.GetEnvironmentVariable("ANT_HOME"));
-//             }
-// 
-//             // If we are on Mono and we are still missing a key then go and find it from the .bash_profile
-//             if (Utils.IsRunningOnMono && EnvVarNames.All( s => AndroidEnv.ContainsKey(s)))
-//             {
-//                 string BashProfilePath = Path.Combine(Environment.GetFolderPath(Environment.SpecialFolder.Personal), ".bash_profile");
-//                 if (File.Exists(BashProfilePath))
-//                 {
-//                     string[] BashProfileContents = File.ReadAllLines(BashProfilePath);
-//                     foreach (string Line in BashProfileContents)
-//                     {
-//                         foreach (var key in EnvVarNames)
-//                         {
-//                             if (AndroidEnv.ContainsKey(key))
-//                             {
-//                                 continue;
-//                             }
-//                             
-//                             if (Line.StartsWith("export " + key + "="))
-//                             {
-//                                 string PathVar = Line.Split('=')[1].Replace("\"", "");
-//                                 AndroidEnv.Add(key, PathVar);
-//                             }
-//                         }
-//                     }
-//                 }                    
-//             }
-// 
-//             // Set for the process
-//             foreach(var kvp in AndroidEnv)
-//             {
-//                 Environment.SetEnvironmentVariable(kvp.Key, kvp.Value);
-//             }
-// 
-//         }
-// 
-//         return AndroidEnv;
-//     }
-
 	public override void Package(ProjectParams Params, DeploymentContext SC, int WorkingCL)
 	{
 		string[] Architectures = UnrealBuildTool.AndroidToolChain.GetAllArchitectures();
@@ -288,9 +216,9 @@ public class AndroidPlatform : Platform
 						"if [ $? -eq 0 ]; then",
 						"\techo",
 						"\techo Removing old data. Failures here are usually fine - indicating the files were not on the device.",
-						"\t$ADB $DEVICE shell 'rm -r $EXTERNAL_STORAGE/" + Params.ShortProjectName + "'",
+                        "\t$ADB $DEVICE shell 'rm -r $EXTERNAL_STORAGE/UE4Game/" + Params.ShortProjectName + "'",
 						"\t$ADB $DEVICE shell 'rm -r $EXTERNAL_STORAGE/UE4Game/UE4CommandLine.txt" + "'",
-						"\t$ADB $DEVICE shell 'rm -r $EXTERNAL_STORAGE/obb/" + PackageName + "'",
+						"\t$ADB $DEVICE shell 'rm -r $EXTERNAL_STORAGE/" + TargetAndroidLocation + PackageName + "'",
 						bPackageDataInsideApk ? "" : "\techo",
 						bPackageDataInsideApk ? "" : "\techo Installing new data. Failures here indicate storage problems \\(missing SD card or bad permissions\\) and are fatal.",
 						bPackageDataInsideApk ? "" : "\tSTORAGE=$(echo \"`$ADB $DEVICE shell 'echo $EXTERNAL_STORAGE'`\" | cat -v | tr -d '^M')",
@@ -329,9 +257,9 @@ public class AndroidPlatform : Platform
 						"@echo Installing existing application. Failures here indicate a problem with the device (connection or storage permissions) and are fatal.",
 						"%ADB% %DEVICE% install " + Path.GetFileName(ApkName),
 						"@if \"%ERRORLEVEL%\" NEQ \"0\" goto Error",
-						"%ADB% %DEVICE% shell rm -r %STORAGE%/" + Params.ShortProjectName,
+                        "%ADB% %DEVICE% shell rm -r %STORAGE%/UE4Game/" + Params.ShortProjectName,
 						"%ADB% %DEVICE% shell rm -r %STORAGE%/UE4Game/UE4CommandLine.txt", // we need to delete the commandline in UE4Game or it will mess up loading
-						"%ADB% %DEVICE% shell rm -r %STORAGE%/obb/" + PackageName,
+						"%ADB% %DEVICE% shell rm -r %STORAGE%/" + TargetAndroidLocation + PackageName,
 						bPackageDataInsideApk ? "" : "@echo.",
 						bPackageDataInsideApk ? "" : "@echo Installing new data. Failures here indicate storage problems (missing SD card or bad permissions) and are fatal.",
 						bPackageDataInsideApk ? "" : "%ADB% %DEVICE% push " + Path.GetFileName(ObbName) + " %STORAGE%/" + DeviceObbName,
@@ -577,8 +505,8 @@ public class AndroidPlatform : Platform
 		{
 			// cache some strings
 			string BaseCommandline = "push";
-			string RemoteDir = StorageLocation + "/" /*+ Params.ProjectDirPrefix + "/"*/ + Params.ShortProjectName;
-            string UE4GameRemoteDir = StorageLocation + "/" /*+ Params.ProjectDirPrefix + "/"*/ + Params.ShortProjectName;
+			string RemoteDir = StorageLocation + "/UE4Game/" + Params.ShortProjectName;
+            string UE4GameRemoteDir = StorageLocation + "/UE4Game/" + Params.ShortProjectName;
 
 			// make sure device is at a clean state
 			RunAdbCommand(Params, "shell rm -r " + RemoteDir);
@@ -704,7 +632,7 @@ public class AndroidPlatform : Platform
 		{
 			// cache some strings
 			string BaseCommandline = "push";
-            string RemoteDir = StorageLocation + "/" /*+ Params.ProjectDirPrefix + "/"*/ + Params.ShortProjectName;
+            string RemoteDir = StorageLocation + "/UE4Game/" + Params.ShortProjectName;
 
 			string FinalRemoteDir = RemoteDir;
 			/*
diff --git a/Engine/Source/Programs/AutomationTool/ProjectParams.cs b/Engine/Source/Programs/AutomationTool/ProjectParams.cs
--- a/Engine/Source/Programs/AutomationTool/ProjectParams.cs
+++ b/Engine/Source/Programs/AutomationTool/ProjectParams.cs
@@ -1667,18 +1667,7 @@ namespace AutomationTool
 			get { return ProjectUtils.GetShortProjectName(RawProjectPath); }
 		}
 
-        /// <summary>
-        /// Get a directory pre-fix for the game's deploy location. Defaults to UE4Game is not supplied
-        /// </summary>
-
-        public string ProjectDirPrefix
-        {
-            get { return String.IsNullOrEmpty(ProjectDirPrefixValue) ? "UE4Game" : ProjectDirPrefixValue;  }
-            set { ProjectDirPrefixValue = value; }
-        }
-        private string ProjectDirPrefixValue;
-
-		/// <summary>
+  		/// <summary>
 		/// True if this project contains source code.
 		/// </summary>	
 		public bool IsCodeBasedProject
diff --git a/Engine/Source/Runtime/Core/Private/Android/AndroidFile.cpp b/Engine/Source/Runtime/Core/Private/Android/AndroidFile.cpp
--- a/Engine/Source/Runtime/Core/Private/Android/AndroidFile.cpp
+++ b/Engine/Source/Runtime/Core/Private/Android/AndroidFile.cpp
@@ -62,6 +62,13 @@ extern "C" void Java_com_epicgames_ue4_GameActivity_nativeSetObbInfo(JNIEnv* jen
 	jenv->ReleaseStringUTFChars(PackageName, JavaChars);
 }
 
+// Constructs the base path for any files which are not in OBB/pak data
+const FString &GetFileBasePath()
+{
+	static FString BasePath = GFilePathBase + FString("/UE4Game/") + FApp::GetGameName() + FString("/");
+	return BasePath;
+}
+
 /**
 	Android file handle implementation for partial (i.e. parcels) files.
 	This can handle all the variations of accessing data for assets and files.
@@ -342,7 +349,7 @@ public:
 	void Read()
 	{
 		// Local filepaths are directly in the deployment directory.
-		static const FString BasePath = GFilePathBase + FString("/") + FApp::GetGameName() + FString("/");
+		static const FString &BasePath = GetFileBasePath();
 		const FString ManfiestPath = BasePath + ManifestFileName;
 
 		ManifestEntries.Empty();
@@ -419,7 +426,7 @@ public:
 	{
 		
 		// Local filepaths are directly in the deployment directory.
-		static const FString BasePath = GFilePathBase + FString("/") + FApp::GetGameName() + FString("/");
+		static const FString &BasePath = GetFileBasePath();
 		const FString ManfiestPath = BasePath + ManifestFileName;
 
 
@@ -845,8 +852,8 @@ public:
 			// See <http://developer.android.com/google/play/expansion-files.html>
 			FString OBBDir1 = GFilePathBase + FString(TEXT("/Android/obb/") + GPackageName);
 			FString OBBDir2 = GFilePathBase + FString(TEXT("/obb/") + GPackageName);
-			FString MainOBBName = FString::Printf(TEXT("main.%05d.%s.obb"), GPackageVersion, *GPackageName);
-			FString PatchOBBName = FString::Printf(TEXT("patch.%05d.%s.obb"), GPackageVersion, *GPackageName);
+			FString MainOBBName = FString::Printf(TEXT("main.%d.%s.obb"), GPackageVersion, *GPackageName);
+			FString PatchOBBName = FString::Printf(TEXT("patch.%d.%s.obb"), GPackageVersion, *GPackageName);
 			if (FileExists(*(OBBDir1 / MainOBBName), true))
 			{
 				MountOBB(*(OBBDir1 / MainOBBName));
@@ -888,7 +895,6 @@ public:
 			// For local files we need to check if it's a plain
 			// file, as opposed to directories.
 			result = S_ISREG(FileInfo.st_mode);
-			return result;
 		}
 		else
 		{
@@ -896,7 +902,7 @@ public:
 			result = IsResource(AssetPath) || IsAsset(AssetPath);
 		}
 #if LOG_ANDROID_FILE
-		FPlatformMisc::LowLevelOutputDebugStringf(TEXT("FAndroidPlatformFile::FileExists('%s') => %s"), Filename, result ? TEXT("TRUE") : TEXT("FALSE"));
+		FPlatformMisc::LowLevelOutputDebugStringf(TEXT("FAndroidPlatformFile::FileExists('%s') => %s\nResolved as %s"), Filename, result ? TEXT("TRUE") : TEXT("FALSE"), *LocalPath);
 #endif
 		return result;
 	}
@@ -1598,7 +1604,7 @@ private:
 				}
 
 				// Local filepaths are directly in the deployment directory.
-				static FString BasePath = GFilePathBase + FString("/") + FApp::GetGameName() + FString("/");
+				static FString BasePath = GetFileBasePath();
 				LocalPath = BasePath + AndroidPath;
 
 				// Asset paths are relative to the base directory.
