More efficient file filter class, for distilling files in/out of builds. Rules are merged into a tree with nodes for each path fragments and rule numbers in each leaf, allowing tree traversal in an arbitrary order to how they are applied. Includes utility methods for filtering against directory trees, recursing only as deep as necessary to guarantee all rules are satisfied.

diff --git a/Engine/Source/Programs/AutomationTool/FileFilter.cs b/Engine/Source/Programs/AutomationTool/FileFilter.cs
--- a/Engine/Source/Programs/AutomationTool/FileFilter.cs
+++ b/Engine/Source/Programs/AutomationTool/FileFilter.cs
@@ -9,18 +9,28 @@ using System.Text;
 using System.Text.RegularExpressions;
 
 /// <summary>
-/// Stores a set of filter rules, which can be used to filter files in or out of a set
+/// Stores a set of rules, similar to p4 path specifications, which can be used to efficiently filter files in or out of a set.
+/// The rules are merged into a tree with nodes for each path fragments and rule numbers in each leaf, allowing tree traversal in an arbitrary order 
+/// to how they are applied.
 /// </summary>
 public class FileFilter
 {
-	List<Tuple<Regex, bool>> Rules;
+	/// <summary>
+	/// Root node for the tree.
+	/// </summary>
+	FileFilterNode RootNode;
+
+	/// <summary>
+	/// Terminating nodes for each rule added to the filter
+	/// </summary>
+	List<FileFilterNode> Rules = new List<FileFilterNode>();
 
 	/// <summary>
 	/// Default constructor
 	/// </summary>
 	public FileFilter()
 	{
-		Rules = new List<Tuple<Regex,bool>>();
+		RootNode = new FileFilterNode(null, "");
 	}
 
 	/// <summary>
@@ -28,19 +38,24 @@ public class FileFilter
 	/// </summary>
 	/// <param name="Other">Filter to copy from</param>
 	public FileFilter(FileFilter Other)
+		: this()
 	{
-		Rules = new List<Tuple<Regex,bool>>(Other.Rules);
+		foreach (FileFilterNode OtherRule in Other.Rules)
+		{
+			AddRule(OtherRule.ToString(), OtherRule.bInclude);
+		}
 	}
 
 	/// <summary>
 	/// Constructs a file filter from a p4-style filespec. Exclude lines are prefixed with a - character.
 	/// </summary>
 	/// <param name="Lines">Lines to parse rules from</param>
-	public FileFilter(IEnumerable<string> Lines) : this()
+	public FileFilter(IEnumerable<string> Lines)
+		: this()
 	{
-		foreach(string Line in Lines)
+		foreach (string Line in Lines)
 		{
-			if(Line.StartsWith("-"))
+			if (Line.StartsWith("-"))
 			{
 				Exclude(Line.Substring(1).TrimStart());
 			}
@@ -57,7 +72,7 @@ public class FileFilter
 	/// <param name="Pattern">Pattern to match. See CreateRegex() for details.</param>
 	public void Include(string Pattern)
 	{
-		Rules.Add(new Tuple<Regex,bool>(CreateRegex(Pattern), true));
+		AddRule(Pattern, true);
 	}
 
 	/// <summary>
@@ -66,7 +81,16 @@ public class FileFilter
 	/// <param name="Patterns">Patterns to match. See CreateRegex() for details.</param>
 	public void Include(params string[] Patterns)
 	{
-		foreach(string Pattern in Patterns)
+		Include(Patterns);
+	}
+
+	/// <summary>
+	/// Adds several exclude rules to the filter
+	/// </summary>
+	/// <param name="Patterns">Patterns to match. See CreateRegex() for details.</param>
+	public void Include(IEnumerable<string> Patterns)
+	{
+		foreach (string Pattern in Patterns)
 		{
 			Include(Pattern);
 		}
@@ -78,7 +102,7 @@ public class FileFilter
 	/// <param name="Pattern">Mask to match. See CreateRegex() for details.</param>
 	public void Exclude(string Pattern)
 	{
-		Rules.Add(new Tuple<Regex, bool>(CreateRegex(Pattern), false));
+		AddRule(Pattern, false);
 	}
 
 	/// <summary>
@@ -87,58 +111,376 @@ public class FileFilter
 	/// <param name="Patterns">Patterns to match. See CreateRegex() for details.</param>
 	public void Exclude(params string[] Patterns)
 	{
-		foreach(string Pattern in Patterns)
+		Exclude(Patterns);
+	}
+
+	/// <summary>
+	/// Adds several exclude rules to the filter
+	/// </summary>
+	/// <param name="Patterns">Patterns to match. See CreateRegex() for details.</param>
+	public void Exclude(IEnumerable<string> Patterns)
+	{
+		foreach (string Pattern in Patterns)
 		{
 			Exclude(Pattern);
 		}
 	}
 
+	/// <summary>
+	/// Adds an include or exclude rule to the filter
+	/// </summary>
+	/// <param name="Pattern">The pattern which the rule should match</param>
+	/// <param name="bInclude">Whether to include or exclude files matching this rule</param>
+	void AddRule(string Pattern, bool bInclude)
+	{
+		string NormalizedPattern = Pattern.Replace('\\', '/');
+
+		// We don't want a slash at the start, but if there was not one specified, it's not anchored to the root of the tree.
+		if (NormalizedPattern.StartsWith("/"))
+		{
+			NormalizedPattern = NormalizedPattern.Substring(1);
+		}
+		else
+		{
+			NormalizedPattern = ".../" + NormalizedPattern;
+		}
+
+		// All directories indicate a wildcard match
+		if (NormalizedPattern.EndsWith("/"))
+		{
+			NormalizedPattern += "...";
+		}
+
+		// Replace any directory wildcards mid-string
+		for (int Idx = NormalizedPattern.IndexOf("..."); Idx != -1; Idx = NormalizedPattern.IndexOf("...", Idx))
+		{
+			if (Idx > 0 && NormalizedPattern[Idx - 1] != '/')
+			{
+				NormalizedPattern = NormalizedPattern.Insert(Idx, "*/");
+				Idx++;
+			}
+
+			Idx += 3;
+
+			if (Idx < NormalizedPattern.Length && NormalizedPattern[Idx] != '/')
+			{
+				NormalizedPattern = NormalizedPattern.Insert(Idx, "/*");
+				Idx += 2;
+			}
+		}
+
+		// Split the pattern into fragments
+		string[] BranchPatterns = NormalizedPattern.Split('/');
+
+		// Add it into the tree
+		FileFilterNode LastNode = RootNode;
+		foreach (string BranchPattern in BranchPatterns)
+		{
+			FileFilterNode NextNode = LastNode.Branches.FirstOrDefault(x => x.Pattern == BranchPattern);
+			if (NextNode == null)
+			{
+				NextNode = new FileFilterNode(LastNode, BranchPattern);
+				LastNode.Branches.Add(NextNode);
+			}
+			LastNode = NextNode;
+		}
+
+		// We've reached the end of the pattern, so mark it as a leaf node
+		Rules.Add(LastNode);
+		LastNode.RuleNumber = Rules.Count - 1;
+		LastNode.bInclude = bInclude;
+
+		// Update the maximums along that path
+		for (FileFilterNode UpdateNode = LastNode; UpdateNode != null; UpdateNode = UpdateNode.Parent)
+		{
+			if (bInclude)
+			{
+				UpdateNode.MaxIncludeRuleNumber = LastNode.RuleNumber;
+			}
+			else
+			{
+				UpdateNode.MaxExcludeRuleNumber = LastNode.RuleNumber;
+			}
+		}
+	}
+
 	/// <summary>
 	/// Determines whether the given file matches the filter
 	/// </summary>
-	/// <param name="File">File to match</param>
+	/// <param name="FileName">File to match</param>
 	/// <returns>True if the file passes the filter</returns>
 	public bool Matches(string FileName)
 	{
-		bool bIsMatch = false;
-		foreach(Tuple<Regex, bool> Rule in Rules)
+		string[] Tokens = FileName.TrimStart('/', '\\').Split('/', '\\');
+
+		FileFilterNode MatchingNode = FindMatchingNode(RootNode, Tokens, 0, RootNode);
+
+		return MatchingNode.bInclude;
+	}
+
+	/// <summary>
+	/// Determines whether it's possible for anything within the given folder name to match the filter. Useful to early out of recursive file searches.
+	/// </summary>
+	/// <param name="FolderName">File to match</param>
+	/// <returns>True if the file passes the filter</returns>
+	public bool PossiblyMatches(string FolderName)
+	{
+		string[] Tokens = FolderName.Trim('/', '\\').Split('/', '\\');
+
+		FileFilterNode MatchingNode = FindMatchingNode(RootNode, Tokens, 0, RootNode);
+
+		return MatchingNode.bInclude || HighestPossibleIncludeMatch(RootNode, Tokens, 0, MatchingNode.RuleNumber) > MatchingNode.RuleNumber;
+	}
+
+	/// <summary>
+	/// Applies the filter to each element in a sequence, and returns the list of files that match
+	/// </summary>
+	/// <param name="FileNames">List of filenames</param>
+	/// <returns>List of filenames which match the filter</returns>
+	public IEnumerable<string> ApplyTo(IEnumerable<string> FileNames)
+	{
+		return FileNames.Where(x => Matches(x));
+	}
+
+	/// <summary>
+	/// Finds a list of files within a given directory which match the filter.
+	/// </summary>
+	/// <param name="FolderName">File to match</param>
+	/// <returns>True if the file passes the filter</returns>
+	public IEnumerable<string> ApplyToDirectory(string DirectoryName)
+	{
+		List<string> MatchingFileNames = new List<string>();
+		FindMatchesFromDirectory(new DirectoryInfo(DirectoryName), "", MatchingFileNames);
+		return MatchingFileNames.ToArray();
+	}
+
+	/// <summary>
+	/// Finds a list of files within a given directory which match the filter.
+	/// </summary>
+	/// <param name="FolderName">File to match</param>
+	/// <returns>True if the file passes the filter</returns>
+	void FindMatchesFromDirectory(DirectoryInfo CurrentDirectory, string NamePrefix, List<string> MatchingFileNames)
+	{
+		foreach (FileInfo NextFile in CurrentDirectory.EnumerateFiles())
+		{
+			string FileName = NamePrefix + NextFile.Name;
+			if (Matches(FileName))
+			{
+				MatchingFileNames.Add(FileName);
+			}
+		}
+		foreach (DirectoryInfo NextDirectory in CurrentDirectory.EnumerateDirectories())
 		{
-			if(bIsMatch != Rule.Item2 && Rule.Item1.IsMatch(FileName))
+			string NextNamePrefix = NamePrefix + NextDirectory.Name;
+			if (PossiblyMatches(NextNamePrefix))
 			{
-				bIsMatch = Rule.Item2;
+				FindMatchesFromDirectory(NextDirectory, NextNamePrefix + "/", MatchingFileNames);
 			}
 		}
-		return bIsMatch;
 	}
 
 	/// <summary>
-	/// Creates a regex to match the given p4-style pattern, which matches forward and back slashes as path separators, and match ?, *, and ... as wildcards.
+	/// Finds the node which matches a given list of tokens.
 	/// </summary>
-	/// <param name="Pattern">Pattern to convert into a regex</param>
-	/// <returns>The resulting regex</returns>
-	public static Regex CreateRegex(string Pattern)
+	/// <param name="CurrentNode"></param>
+	/// <param name="Tokens"></param>
+	/// <param name="TokenIdx"></param>
+	/// <param name="CurrentBestNode"></param>
+	/// <returns></returns>
+	FileFilterNode FindMatchingNode(FileFilterNode CurrentNode, string[] Tokens, int TokenIdx, FileFilterNode CurrentBestNode)
 	{
-		// Convert the pattern into a p4-style path specification.
-		string FinalPattern = "^" + Regex.Escape(Pattern.Replace('\\', '/')) + "$";
-		FinalPattern = FinalPattern.Replace("/", "[/\\\\]");
-		FinalPattern = FinalPattern.Replace("\\?", ".");
-		FinalPattern = FinalPattern.Replace("\\*\\.$", "[^/\\\\.]*$"); // No extension pattern: *.
-		FinalPattern = FinalPattern.Replace("\\*", "[^/\\\\]*");
-		FinalPattern = FinalPattern.Replace("\\.\\.\\.", ".*");
+		// If we've matched all the input tokens, check if this rule is better than any other we've seen
+		if (TokenIdx == Tokens.Length)
+		{
+			return (CurrentNode.RuleNumber > CurrentBestNode.RuleNumber) ? CurrentNode : CurrentBestNode;
+		}
 
-		// Remove sentinels when we're just looking for substrings
-		if(FinalPattern.StartsWith("^.*"))
+		// If there is no rule under the current node which is better than the current best node, early out
+		if (CurrentBestNode.bInclude)
 		{
-			FinalPattern = FinalPattern.Substring(3);
+			if (CurrentNode.MaxExcludeRuleNumber < CurrentBestNode.RuleNumber)
+			{
+				return CurrentBestNode;
+			}
 		}
-		if(FinalPattern.EndsWith(".*$"))
+		else
 		{
-			FinalPattern = FinalPattern.Substring(0, FinalPattern.Length - 3);
+			if (CurrentNode.MaxIncludeRuleNumber < CurrentBestNode.RuleNumber)
+			{
+				return CurrentBestNode;
+			}
 		}
 
-		// Create the regex
-		return new Regex(FinalPattern, RegexOptions.Compiled | RegexOptions.IgnoreCase);
+		// Test all the branches for one that matches
+		FileFilterNode BestNode = CurrentBestNode;
+		foreach (FileFilterNode Branch in CurrentNode.Branches)
+		{
+			if (Branch.Pattern == "...")
+			{
+				for (int NextTokenIdx = Tokens.Length; NextTokenIdx >= TokenIdx; NextTokenIdx--)
+				{
+					BestNode = FindMatchingNode(Branch, Tokens, NextTokenIdx, BestNode);
+				}
+			}
+			else
+			{
+				if (Branch.IsMatch(Tokens[TokenIdx]))
+				{
+					BestNode = FindMatchingNode(Branch, Tokens, TokenIdx + 1, BestNode);
+				}
+			}
+		}
+		return BestNode;
+	}
+
+	/// <summary>
+	/// Returns the highest possible rule number which can match the given list of input tokens, assuming that the list of input tokens is incomplete.
+	/// </summary>
+	/// <param name="CurrentNode">The current node being checked</param>
+	/// <param name="Tokens">The tokens to match</param>
+	/// <param name="TokenIdx">Current token index</param>
+	/// <param name="CurrentBestRuleNumber">The highest rule number seen so far. Used to optimize tree traversals.</param>
+	/// <returns>New highest rule number</returns>
+	int HighestPossibleIncludeMatch(FileFilterNode CurrentNode, string[] Tokens, int TokenIdx, int CurrentBestRuleNumber)
+	{
+		// If we've matched all the input tokens, check if this rule is better than any other we've seen
+		if (TokenIdx == Tokens.Length)
+		{
+			return Math.Max(CurrentBestRuleNumber, CurrentNode.MaxIncludeRuleNumber);
+		}
+
+		// Test all the branches for one that matches
+		int BestRuleNumber = CurrentBestRuleNumber;
+		if (CurrentNode.MaxIncludeRuleNumber > BestRuleNumber)
+		{
+			foreach (FileFilterNode Branch in CurrentNode.Branches)
+			{
+				if (Branch.Pattern == "...")
+				{
+					if (Branch.MaxIncludeRuleNumber > CurrentBestRuleNumber)
+					{
+						BestRuleNumber = Branch.MaxIncludeRuleNumber;
+					}
+				}
+				else
+				{
+					if (Branch.IsMatch(Tokens[TokenIdx]))
+					{
+						BestRuleNumber = HighestPossibleIncludeMatch(Branch, Tokens, TokenIdx + 1, BestRuleNumber);
+					}
+				}
+			}
+		}
+		return BestRuleNumber;
 	}
 }
 
+/// <summary>
+/// Node within a filter tree. Each node matches a single path fragment - a folder or file name, with an include or exclude rule consisting of a sequence of nodes. 
+/// </summary>
+class FileFilterNode
+{
+	/// <summary>
+	/// Node which this is parented to. Null for the root node.
+	/// </summary>
+	public readonly FileFilterNode Parent;
+
+	/// <summary>
+	/// Pattern to match for this node. May contain * or ? wildcards.
+	/// </summary>
+	public readonly string Pattern;
+
+	/// <summary>
+	/// Highest include rule number matched by this node or any child nodes.
+	/// </summary>
+	public int MaxIncludeRuleNumber;
+
+	/// <summary>
+	/// Highest exclude rule number matched by this node or any child nodes.
+	/// </summary>
+	public int MaxExcludeRuleNumber;
+
+	/// <summary>
+	/// Child branches of this node, distinct by the the pattern for each.
+	/// </summary>
+	public List<FileFilterNode> Branches = new List<FileFilterNode>();
+
+	/// <summary>
+	/// If a path matches the sequence of nodes ending in this node, the number of the rule which added it. -1 if this was not the last node in a rule.
+	/// </summary>
+	public int RuleNumber;
+
+	/// <summary>
+	/// True if the rule terminating in this node was an include node.
+	/// </summary>
+	public bool bInclude;
+
+	/// <summary>
+	/// Default constructor.
+	/// </summary>
+	public FileFilterNode(FileFilterNode InParent, string InPattern)
+	{
+		Parent = InParent;
+		Pattern = InPattern;
+		RuleNumber = -1;
+		bInclude = false;
+	}
+
+	/// <summary>
+	/// Determine if the given token matches the pattern in this node
+	/// </summary>
+	public bool IsMatch(string Token)
+	{
+		if(Pattern.EndsWith("."))
+		{
+			return !Token.Contains('.') && IsMatch(Token, 0, Pattern.Substring(0, Pattern.Length - 1), 0);
+		}
+		else
+		{
+			return IsMatch(Token, 0, Pattern, 0);
+		}
+	}
 
+	/// <summary>
+	/// Determine if a given token matches a pattern, starting from the given positions within each string.
+	/// </summary>
+	static bool IsMatch(string Token, int TokenIdx, string Pattern, int PatternIdx)
+	{
+		for (; ; )
+		{
+			if (PatternIdx == Pattern.Length)
+			{
+				return (TokenIdx == Token.Length);
+			}
+			else if (Pattern[PatternIdx] == '*')
+			{
+				for (int NextTokenIdx = Token.Length; NextTokenIdx >= TokenIdx; NextTokenIdx--)
+				{
+					if (IsMatch(Token, NextTokenIdx, Pattern, PatternIdx + 1))
+					{
+						return true;
+					}
+				}
+				return false;
+			}
+			else if (TokenIdx < Token.Length && (Char.ToLower(Token[TokenIdx]) == Char.ToLower(Pattern[PatternIdx]) || Pattern[PatternIdx] == '?'))
+			{
+				TokenIdx++;
+				PatternIdx++;
+			}
+			else
+			{
+				return false;
+			}
+		}
+	}
+
+	/// <summary>
+	/// Debugger visualization
+	/// </summary>
+	/// <returns>Path to this node</returns>
+	public override string ToString()
+	{
+		return (Parent == null) ? "/" : Parent.ToString().TrimEnd('/') + "/" + Pattern;
+	}
+}
