Changed localization automation script so that it downloads translations from OneSky based on a project's languages, rather than the parent project group's listed languages.

diff --git a/Engine/Source/Programs/AutomationTool/Scripts/Localisation.Automation.cs b/Engine/Source/Programs/AutomationTool/Scripts/Localisation.Automation.cs
--- a/Engine/Source/Programs/AutomationTool/Scripts/Localisation.Automation.cs
+++ b/Engine/Source/Programs/AutomationTool/Scripts/Localisation.Automation.cs
@@ -20,7 +20,7 @@ class Localise : BuildCommand
         //Export
         if (file != null)
         {
-            ExportFileToDirectory(file, new DirectoryInfo(CmdEnv.LocalRoot + "/Engine/Content/Localization/" + ProjectName), projectGroup.EnabledCultures);
+            ExportFileToDirectory(file, new DirectoryInfo(CmdEnv.LocalRoot + "/Engine/Content/Localization/" + ProjectName), project.EnabledCultures);
         }
     }
 
