Add a -Timestamp parameter for UAT to add timestamps to log messages.

diff --git a/Engine/Source/Programs/AutomationTool/Log.cs b/Engine/Source/Programs/AutomationTool/Log.cs
--- a/Engine/Source/Programs/AutomationTool/Log.cs
+++ b/Engine/Source/Programs/AutomationTool/Log.cs
@@ -14,6 +14,7 @@ namespace AutomationTool
 
 	public class LogUtils
 	{
+		private static Stopwatch Timer;
 		private static AutomationFileTraceListener FileLog;
 
 		/// <summary>
@@ -22,6 +23,7 @@ namespace AutomationTool
 		/// <param name="CommandLine">Command line.</param>
 		public static void InitLogging(string[] CommandLine)
 		{
+			Timer = (CommandUtils.ParseParam(CommandLine, "-Timestamps"))? Stopwatch.StartNew() : null;
 			var VerbosityLevel = CommandUtils.ParseParam(CommandLine, "-Verbose") ? TraceEventType.Verbose : TraceEventType.Information;
 			var Filter = new VerbosityFilter(VerbosityLevel);
 			Trace.Listeners.Add(new AutomationConsoleTraceListener());
@@ -70,7 +72,12 @@ namespace AutomationTool
 		/// <returns>Formatted message</returns>
 		public static string FormatMessage(string Source, TraceEventType Verbosity, string Message)
 		{
-			var FormattedMessage = Source + ": ";
+			string FormattedMessage = "";
+			if(Timer != null)
+			{
+				FormattedMessage += String.Format("[{0:hh\\:mm\\:ss\\.fff}] ", Timer.Elapsed);
+			}
+			FormattedMessage += Source + ": ";
 			switch (Verbosity)
 			{
 				case TraceEventType.Error:
