build script, don't try and kill processes on mac #codereview Michael.Trepka

diff --git a/Engine/Source/Programs/AutomationTool/Program.cs b/Engine/Source/Programs/AutomationTool/Program.cs
--- a/Engine/Source/Programs/AutomationTool/Program.cs
+++ b/Engine/Source/Programs/AutomationTool/Program.cs
@@ -130,7 +130,7 @@ namespace AutomationTool
 			Environment.ExitCode = ReturnCode;
 
 			// Try to kill process before app domain exits to leave the other KillAll call to extreme edge cases
-			if (ShouldKillProcesses)
+			if (ShouldKillProcesses && !Utils.IsRunningOnMono)
 			{
 				ProcessManager.KillAll();
 			}
@@ -156,7 +156,7 @@ namespace AutomationTool
 		{
 			// Kill all spawned processes (Console instead of Log because logging is closed at this time anyway)
 			Console.WriteLine("Domain_ProcessExit");
-			if (ShouldKillProcesses)
+			if (ShouldKillProcesses && !Utils.IsRunningOnMono)
 			{			
 				ProcessManager.KillAll();
 			}
