CIS fix, properly check for empy project path

diff --git a/Engine/Source/Programs/AutomationTool/ProjectUtils.cs b/Engine/Source/Programs/AutomationTool/ProjectUtils.cs
--- a/Engine/Source/Programs/AutomationTool/ProjectUtils.cs
+++ b/Engine/Source/Programs/AutomationTool/ProjectUtils.cs
@@ -215,17 +215,20 @@ namespace AutomationTool
 			Properties.RawProjectPath = RawProjectPath;
 
 			// detect if the project is content only, but has non-default build settings
-			List<string> ExtraSearchPaths = null;
-			if (RequiresTempTarget(RawProjectPath, ClientTargetPlatforms))
+			if (!string.IsNullOrEmpty(RawProjectPath))
 			{
-				GenerateTempTarget(RawProjectPath);
-				Properties.bWasGenerated = true;
-				ExtraSearchPaths = new List<string>();
-				ExtraSearchPaths.Add(CommandUtils.CombinePaths(Path.GetDirectoryName(RawProjectPath), "Intermediate", "Source"));
-			}
-			else if (File.Exists(Path.Combine(Path.GetDirectoryName(RawProjectPath), "Intermediate", "Source", Path.GetFileNameWithoutExtension(RawProjectPath) + ".Target.cs")))
-			{
-				File.Delete(Path.Combine(Path.GetDirectoryName(RawProjectPath), "Intermediate", "Source", Path.GetFileNameWithoutExtension(RawProjectPath) + ".Target.cs"));
+				List<string> ExtraSearchPaths = null;
+				if (RequiresTempTarget(RawProjectPath, ClientTargetPlatforms))
+				{
+					GenerateTempTarget(RawProjectPath);
+					Properties.bWasGenerated = true;
+					ExtraSearchPaths = new List<string>();
+					ExtraSearchPaths.Add(CommandUtils.CombinePaths(Path.GetDirectoryName(RawProjectPath), "Intermediate", "Source"));
+				}
+				else if (File.Exists(Path.Combine(Path.GetDirectoryName(RawProjectPath), "Intermediate", "Source", Path.GetFileNameWithoutExtension(RawProjectPath) + ".Target.cs")))
+				{
+					File.Delete(Path.Combine(Path.GetDirectoryName(RawProjectPath), "Intermediate", "Source", Path.GetFileNameWithoutExtension(RawProjectPath) + ".Target.cs"));
+				}
 			}
 
 			if (CommandUtils.CmdEnv.HasCapabilityToCompile)
