fix for generating Target.cs file when one already exists

diff --git a/Engine/Source/Programs/AutomationTool/ProjectUtils.cs b/Engine/Source/Programs/AutomationTool/ProjectUtils.cs
--- a/Engine/Source/Programs/AutomationTool/ProjectUtils.cs
+++ b/Engine/Source/Programs/AutomationTool/ProjectUtils.cs
@@ -157,6 +157,14 @@ namespace AutomationTool
 
 		private static bool RequiresTempTarget(string RawProjectPath, List<UnrealTargetPlatform> ClientTargetPlatforms)
 		{
+			// check to see if we already have a Target.cs file
+			// @todo: Target.cs may not be the same name as the project in the future, so look at a better way to determine this
+			if (File.Exists(Path.Combine(Path.GetDirectoryName(RawProjectPath), "Source", Path.GetFileNameWithoutExtension(RawProjectPath) + ".Target.cs")))
+			{
+				return false;
+			}
+
+			// no Target file, now check to see if build settings have changed
 			if (ClientTargetPlatforms == null || ClientTargetPlatforms.Count < 1)
 			{
 				foreach (UnrealTargetPlatform TargetPlatformType in Enum.GetValues(typeof(UnrealTargetPlatform)))
@@ -214,8 +222,8 @@ namespace AutomationTool
 			var Properties = new ProjectProperties();
 			Properties.RawProjectPath = RawProjectPath;
 
-			List<string> ExtraSearchPaths = null;
 			// detect if the project is content only, but has non-default build settings
+			List<string> ExtraSearchPaths = null;
 			if (!string.IsNullOrEmpty(RawProjectPath))
 			{
 				if (RequiresTempTarget(RawProjectPath, ClientTargetPlatforms))
