added check for additional plugins enabled as part of the requirements to build a content project
cleaned up the code a little bit to share the code better

diff --git a/Engine/Source/Programs/AutomationTool/ProjectUtils.cs b/Engine/Source/Programs/AutomationTool/ProjectUtils.cs
--- a/Engine/Source/Programs/AutomationTool/ProjectUtils.cs
+++ b/Engine/Source/Programs/AutomationTool/ProjectUtils.cs
@@ -165,28 +165,51 @@ namespace AutomationTool
 			}
 
 			// no Target file, now check to see if build settings have changed
+			List<UnrealTargetPlatform> TargetPlatforms = ClientTargetPlatforms;
 			if (ClientTargetPlatforms == null || ClientTargetPlatforms.Count < 1)
 			{
+				// No client target platforms, add all in
+				TargetPlatforms = new List<UnrealTargetPlatform>();
 				foreach (UnrealTargetPlatform TargetPlatformType in Enum.GetValues(typeof(UnrealTargetPlatform)))
 				{
 					if (TargetPlatformType != UnrealTargetPlatform.Unknown)
 					{
-						IUEBuildPlatform BuildPlat = UEBuildPlatform.GetBuildPlatform(TargetPlatformType, true);
-						if (BuildPlat != null && !(BuildPlat as UEBuildPlatform).HasDefaultBuildConfig(TargetPlatformType, Path.GetDirectoryName(RawProjectPath)))
-						{
-							return true;
-						}
+						TargetPlatforms.Add(TargetPlatformType);
 					}
 				}
 			}
-			else
+
+			// check the target platforms for any differences in build settings or additional plugins
+			foreach (UnrealTargetPlatform TargetPlatformType in TargetPlatforms)
 			{
-				foreach (UnrealTargetPlatform TargetPlatformType in ClientTargetPlatforms)
+				IUEBuildPlatform BuildPlat = UEBuildPlatform.GetBuildPlatform(TargetPlatformType, true);
+				if (BuildPlat != null && !(BuildPlat as UEBuildPlatform).HasDefaultBuildConfig(TargetPlatformType, Path.GetDirectoryName(RawProjectPath)))
 				{
-					IUEBuildPlatform BuildPlat = UEBuildPlatform.GetBuildPlatform(TargetPlatformType, true);
-					if (BuildPlat != null && !(BuildPlat as UEBuildPlatform).HasDefaultBuildConfig(TargetPlatformType, Path.GetDirectoryName(RawProjectPath)))
+					return true;
+				}
+				// find if there are any plugins
+				List<string> PluginList = new List<string>();
+				// Use the project settings to update the plugin list for this target
+				PluginList = UProjectInfo.GetEnabledPlugins(RawProjectPath, PluginList, TargetPlatformType);
+				if (PluginList.Count > 0)
+				{
+					foreach (var PluginName in PluginList)
 					{
-						return true;
+						// check the plugin info for this plugin itself
+						foreach (var Plugin in Plugins.AllPlugins)
+						{
+							if (Plugin.Name == PluginName)
+							{
+								foreach (var Module in Plugin.Modules)
+								{
+									if (Module.Platforms.Count > 0 && Module.Platforms.Contains(TargetPlatformType))
+									{
+										return true;
+									}
+								}
+								break;
+							}
+						}
 					}
 				}
 			}
