Allow suppressing P4 output for several more commands.

diff --git a/Engine/Source/Programs/AutomationTool/P4Utils.cs b/Engine/Source/Programs/AutomationTool/P4Utils.cs
--- a/Engine/Source/Programs/AutomationTool/P4Utils.cs
+++ b/Engine/Source/Programs/AutomationTool/P4Utils.cs
@@ -1017,10 +1017,10 @@ namespace AutomationTool
 		/// <param name="FromCL">Changelist to unshelve.</param>
 		/// <param name="ToCL">Changelist where the checked out files should be added.</param>
 		/// <param name="CommandLine">Commandline for the command.</param>
-		public void Unshelve(int FromCL, int ToCL, string CommandLine = "")
+		public void Unshelve(int FromCL, int ToCL, string CommandLine = "", bool SpewIsVerbose = false)
 		{
 			CheckP4Enabled();
-			LogP4("unshelve " + String.Format("-s {0} ", FromCL) + String.Format("-c {0} ", ToCL) + CommandLine);
+			LogP4("unshelve " + String.Format("-s {0} ", FromCL) + String.Format("-c {0} ", ToCL) + CommandLine, SpewIsVerbose: SpewIsVerbose);
 		}
 
         /// <summary>
@@ -1165,10 +1165,10 @@ namespace AutomationTool
 		/// Reverts all files from the specified changelist.
 		/// </summary>
 		/// <param name="CL">Changelist to revert.</param>
-		public void RevertAll(int CL)
+		public void RevertAll(int CL, bool SpewIsVerbose = false)
 		{
 			CheckP4Enabled();
-			LogP4("revert " + String.Format("-c {0} //...", CL));
+			LogP4("revert " + String.Format("-c {0} //...", CL), SpewIsVerbose: SpewIsVerbose);
 		}
 
 		/// <summary>
@@ -1384,16 +1384,16 @@ namespace AutomationTool
 		/// </summary>
 		/// <param name="CL">Changelist to delete.</param>
 		/// <param name="RevertFiles">Indicates whether files in that changelist should be reverted.</param>
-		public void DeleteChange(int CL, bool RevertFiles = true)
+		public void DeleteChange(int CL, bool RevertFiles = true, bool SpewIsVerbose = false)
 		{
 			CheckP4Enabled();
 			if (RevertFiles)
 			{
-				RevertAll(CL);
+				RevertAll(CL, SpewIsVerbose: SpewIsVerbose);
 			}
 
 			string CmdOutput;
-			if (LogP4Output(out CmdOutput, String.Format("change -d {0}", CL)))
+			if (LogP4Output(out CmdOutput, String.Format("change -d {0}", CL), SpewIsVerbose: SpewIsVerbose))
 			{
 				string EndStr = " deleted.";
 				string ChangeStr = "Change ";
