build script, fix for code/noncode projects, add failure email mining at end of node

diff --git a/Engine/Build/Target.cs.template b/Engine/Build/Target.cs.template
--- a/Engine/Build/Target.cs.template
+++ b/Engine/Build/Target.cs.template
@@ -93,6 +93,12 @@ public class {GAME_NAME}Target : TargetRules
 			OutCPPEnvironmentConfiguration.Definitions.Add("UE4GAME=1");
 		}
 	}
+	public override GUBPProjectOptions GUBP_IncludeProjectInPromotedBuild_EditorTypeOnly(UnrealTargetPlatform HostPlatform)
+    {
+        var Result = new GUBPProjectOptions();
+        Result.bTestWithShared = true;
+        return Result;
+    }
 
 	public override List<UnrealTargetPlatform> GUBP_GetPlatforms_MonolithicOnly(UnrealTargetPlatform HostPlatform)
 	{
@@ -101,15 +107,15 @@ public class {GAME_NAME}Target : TargetRules
 		switch (HostPlatform)
 		{
 			case UnrealTargetPlatform.Mac:
-				Platforms = new List<UnrealTargetPlatform> { HostPlatform, UnrealTargetPlatform.IOS };
+				Platforms = new List<UnrealTargetPlatform> { UnrealTargetPlatform.IOS };
 				break;
 
 			case UnrealTargetPlatform.Linux:
-				Platforms = new List<UnrealTargetPlatform> { HostPlatform };
+				Platforms = new List<UnrealTargetPlatform>();
 				break;
 
 			case UnrealTargetPlatform.Win64:
-				Platforms = new List<UnrealTargetPlatform> { HostPlatform, UnrealTargetPlatform.Win32, UnrealTargetPlatform.IOS, UnrealTargetPlatform.XboxOne, UnrealTargetPlatform.PS4, UnrealTargetPlatform.Android, UnrealTargetPlatform.Linux, UnrealTargetPlatform.HTML5 };
+				Platforms = new List<UnrealTargetPlatform>();
 				break;
 
 			default:
@@ -124,9 +130,21 @@ public class {GAME_NAME}Target : TargetRules
 	{
 		return new List<UnrealTargetConfiguration> { UnrealTargetConfiguration.Development, UnrealTargetConfiguration.Shipping, UnrealTargetConfiguration.Test };
 	}
-
-	public override bool GUBP_BuildWindowsXPMonolithics()
-	{
-		return true;
+	public override List<GUBPFormalBuild> GUBP_GetConfigsForFormalBuilds_MonolithicOnly(UnrealTargetPlatform HostPlatform)
+    {
+		if (HostPlatform == UnrealTargetPlatform.Win64)
+		{
+			return new List<GUBPFormalBuild>();			
+		}
+		else
+		{
+			return new List<GUBPFormalBuild> 
+            { 
+                new GUBPFormalBuild(UnrealTargetPlatform.IOS, UnrealTargetConfiguration.Development),
+				new GUBPFormalBuild(UnrealTargetPlatform.IOS, UnrealTargetConfiguration.Shipping),
+				new GUBPFormalBuild(UnrealTargetPlatform.IOS, UnrealTargetConfiguration.Test),
+            };
+		}
 	}
+
 }
diff --git a/Engine/Source/Programs/AutomationTool/ProjectUtils.cs b/Engine/Source/Programs/AutomationTool/ProjectUtils.cs
--- a/Engine/Source/Programs/AutomationTool/ProjectUtils.cs
+++ b/Engine/Source/Programs/AutomationTool/ProjectUtils.cs
@@ -623,6 +623,11 @@ namespace AutomationTool
                     var Target = Properties.Targets[TargetRules.TargetType.Editor];
                     Options = Target.Rules.GUBP_IncludeProjectInPromotedBuild_EditorTypeOnly(HostPlatform);
                 }
+				else if(Properties.Targets.ContainsKey(TargetRules.TargetType.Game))
+				{
+					var Target = Properties.Targets[TargetRules.TargetType.Game];
+					Options = Target.Rules.GUBP_IncludeProjectInPromotedBuild_EditorTypeOnly(HostPlatform);
+				}
                 return Options;
             }
             public void Dump(List<UnrealTargetPlatform> InHostPlatforms)
diff --git a/Engine/Source/Programs/AutomationTool/Scripts/GUBP.Automation.cs b/Engine/Source/Programs/AutomationTool/Scripts/GUBP.Automation.cs
--- a/Engine/Source/Programs/AutomationTool/Scripts/GUBP.Automation.cs
+++ b/Engine/Source/Programs/AutomationTool/Scripts/GUBP.Automation.cs
@@ -4886,96 +4886,10 @@
 			GUBPNodesHistory.Add(Node, History);
         }
     }
-
-    bool HashSetEqual(HashSet<string> A, HashSet<string> B)
-    {
-        if (A.Count != B.Count)
-        {
-            return false;
-        }
-        foreach (var Elem in A)
-        {
-            if (!B.Contains(Elem))
-            {
-                return false;
-            }
-        }
-        foreach (var Elem in B)
-        {
-            if (!A.Contains(Elem))
-            {
-                return false;
-            }
-        }
-        return true;
-    }
-
-    int FindLastNonDuplicateFail(string NodeToDo, string CLString)
-    {
-        var History = GUBPNodesHistory[NodeToDo];
-        int Result = P4Env.Changelist;
-
-        string GameNameIfAny = GUBPNodes[NodeToDo].GameNameIfAnyForTempStorage();
-        string NodeStore = StoreName + "-" + GUBPNodes[NodeToDo].GetFullName() + FailedTempStorageSuffix;
-
-        var BackwardsFails = new List<int>(History.AllFailed);
-        BackwardsFails.Add(P4Env.Changelist);
-        BackwardsFails.Sort();
-        BackwardsFails.Reverse();
-        HashSet<string> CurrentErrors = null;
-        foreach (var CL in BackwardsFails)
-        {
-            if (CL > P4Env.Changelist)
-            {
-                continue;
-            }
-            if (CL <= History.LastSucceeded)
-            {
-                break;
-            }
-            var ThisNodeStore = NodeStore.Replace(CLString, String.Format("{0}", CL));
-            DeleteLocalTempStorage(CmdEnv, ThisNodeStore, true); // these all clash locally, which is fine we just retrieve them from shared
-
-            List<string> Files = null;
-            try
-            {
-                bool WasLocal;
-                Files = RetrieveFromTempStorage(CmdEnv, ThisNodeStore, out WasLocal, GameNameIfAny); // this will fail on our CL if we didn't fail or we are just setting up the branch
-            }
-            catch (Exception)
-            {
-            }
-            if (Files == null)
-            {
-                continue;
-            }
-            if (Files.Count != 1)
-            {
-                throw new AutomationException("Unexpected number of files for fail record {0}", Files.Count);
-            }
-            string ErrorFile = Files[0];
-            var ThisErrors = ECJobPropsUtils.ErrorsFromProps(ErrorFile);
-            if (CurrentErrors == null)
-            {
-                CurrentErrors = ThisErrors;
-            }
-            else
-            {
-                if (CurrentErrors.Count == 0 || !HashSetEqual(CurrentErrors, ThisErrors))
-                {
-                    break;
-                }
-                Result = CL;
-            }
-        }
-        return Result;
-    }
-    List<string> GetECPropsForNode(string NodeToDo, string CLString, out string EMails, bool OnlyLateUpdates = false)
-    {
-        var StartTime = DateTime.UtcNow;
-
-        var ECProps = new List<string>();
-        EMails = "";
+	void GetFailureEmails(string NodeToDo, string CLString, bool OnlyLateUpdates = false)
+	{
+		var StartTime = DateTime.UtcNow;        
+        string EMails = "";
         string FailCauserEMails = "";
         string EMailNote = "";
         bool SendSuccessForGreenAfterRed = false;
@@ -4983,11 +4897,8 @@
         if (GUBPNodesHistory.ContainsKey(NodeToDo))
         {
             var History = GUBPNodesHistory[NodeToDo];
+			RunECTool(String.Format("setProperty \"/myWorkflow/LastGreen/{0}\" \"{1}\"", NodeToDo, History.LastSucceeded), true);
+			RunECTool(String.Format("setProperty \"/myWorkflow/LastGreen/{0}\" \"{1}\"", NodeToDo, History.FailedString), true);
 
-            ECProps.Add(string.Format("LastGreen/{0}={1}", NodeToDo, History.LastSucceeded));
-            ECProps.Add(string.Format("RedsSince/{0}={1}", NodeToDo, History.FailedString));
-            ECProps.Add(string.Format("InProgress/{0}={1}", NodeToDo, History.InProgressString));
-
-
             if (History.LastSucceeded > 0 && History.LastSucceeded < P4Env.Changelist)
             {
@@ -5050,14 +4961,11 @@
         }
         else
         {
-            ECProps.Add(string.Format("LastGreen/{0}=0", NodeToDo));
-            ECProps.Add(string.Format("RedsSince/{0}=", NodeToDo));
-            ECProps.Add(string.Format("InProgress/{0}=", NodeToDo));
+			RunECTool(String.Format("setProperty \"/myWorkflow/LastGreen/{0}\" \"{1}\"", NodeToDo, "0"));
+			RunECTool(String.Format("setProperty \"/myWorkflow/RedsSince/{0}\" \"{1}\"", NodeToDo, ""));
         }
-
-        ECProps.Add(string.Format("FailCausers/{0}={1}", NodeToDo, FailCauserEMails));
-        ECProps.Add(string.Format("EmailNotes/{0}={1}", NodeToDo, EMailNote));
-
+		RunECTool(String.Format("setProperty \"/myWorkflow/FailCausers/{0}\" \"{1}\"", NodeToDo, FailCauserEMails));
+		RunECTool(String.Format("setProperty \"/myWorkflow/EmailNotes/{0}\" \"{1}\"", NodeToDo, EMailNote));
         {
             var AdditonalEmails = "";
 
@@ -5065,7 +4973,7 @@
             if (ParseParam("CIS") && !GUBPNodes[NodeToDo].SendSuccessEmail() && !GUBPNodes[NodeToDo].TriggerNode())
             {
 				Causers = FailCauserEMails;
-                }
+           }
             string AddEmails = ParseParamValue("AddEmails");
             if (!String.IsNullOrEmpty(AddEmails))
             {
@@ -5073,32 +4981,131 @@
             }
 
             EMails = GetEMailListForNode(this, NodeToDo, AdditonalEmails, Causers);
+			RunECTool(String.Format("setProperty \"/myWorkflow/FailEmails/{0}\" \"{1}\"", NodeToDo, EMails));            
+        }
+		if (GUBPNodes[NodeToDo].SendSuccessEmail() || SendSuccessForGreenAfterRed)
+		{
+			RunECTool(String.Format("setProperty \"/myWorkflow/SendSuccessEmail/{0}\" \"{1}\"", NodeToDo, "1"));
+		}
+		else
+		{
+			RunECTool(String.Format("setProperty \"/myWorkflow/SendSuccessEmail/{0}\" \"{1}\"", NodeToDo, "0"));
+		}
+	}
 
-            ECProps.Add("FailEmails/" + NodeToDo + "=" + EMails);
+    bool HashSetEqual(HashSet<string> A, HashSet<string> B)
+    {
+        if (A.Count != B.Count)
+        {
+            return false;
         }
-        if (GUBPNodes[NodeToDo].SendSuccessEmail() || SendSuccessForGreenAfterRed)
+        foreach (var Elem in A)
         {
-            ECProps.Add("SendSuccessEmail/" + NodeToDo + "=1");
+            if (!B.Contains(Elem))
+            {
+                return false;
+            }
         }
-        else
+        foreach (var Elem in B)
         {
-            ECProps.Add("SendSuccessEmail/" + NodeToDo + "=0");
+            if (!A.Contains(Elem))
+            {
+                return false;
+            }
         }
-        if (!OnlyLateUpdates)
+        return true;
+    }
+
+    int FindLastNonDuplicateFail(string NodeToDo, string CLString)
+    {
+        var History = GUBPNodesHistory[NodeToDo];
+        int Result = P4Env.Changelist;
+
+        string GameNameIfAny = GUBPNodes[NodeToDo].GameNameIfAnyForTempStorage();
+        string NodeStore = StoreName + "-" + GUBPNodes[NodeToDo].GetFullName() + FailedTempStorageSuffix;
+
+        var BackwardsFails = new List<int>(History.AllFailed);
+        BackwardsFails.Add(P4Env.Changelist);
+        BackwardsFails.Sort();
+        BackwardsFails.Reverse();
+        HashSet<string> CurrentErrors = null;
+        foreach (var CL in BackwardsFails)
         {
+            if (CL > P4Env.Changelist)
+            {
+                continue;
+            }
+            if (CL <= History.LastSucceeded)
+            {
+                break;
+            }
+            var ThisNodeStore = NodeStore.Replace(CLString, String.Format("{0}", CL));
+            DeleteLocalTempStorage(CmdEnv, ThisNodeStore, true); // these all clash locally, which is fine we just retrieve them from shared
+
+            List<string> Files = null;
+            try
+            {
+                bool WasLocal;
+                Files = RetrieveFromTempStorage(CmdEnv, ThisNodeStore, out WasLocal, GameNameIfAny); // this will fail on our CL if we didn't fail or we are just setting up the branch
+            }
+            catch (Exception)
+            {
+            }
+            if (Files == null)
+            {
+                continue;
+            }
+            if (Files.Count != 1)
+            {
+                throw new AutomationException("Unexpected number of files for fail record {0}", Files.Count);
+            }
+            string ErrorFile = Files[0];
+            var ThisErrors = ECJobPropsUtils.ErrorsFromProps(ErrorFile);
+            if (CurrentErrors == null)
+            {
+                CurrentErrors = ThisErrors;
+            }
+            else
+            {
+                if (CurrentErrors.Count == 0 || !HashSetEqual(CurrentErrors, ThisErrors))
+                {
+                    break;
+                }
+                Result = CL;
+            }
+        }
+        return Result;
+    }
+    List<string> GetECPropsForNode(string NodeToDo, string CLString, out string EMails, bool OnlyLateUpdates = false)
+    {
+        var StartTime = DateTime.UtcNow;
+
+        var ECProps = new List<string>();
+        EMails = "";		
+		var AdditonalEmails = "";
+		string Causers = "";				
+		string AddEmails = ParseParamValue("AddEmails");
+		if (!String.IsNullOrEmpty(AddEmails))
+		{
+			AdditonalEmails = GUBPNode.MergeSpaceStrings(AddEmails, AdditonalEmails);
+		}
+		EMails = GetEMailListForNode(this, NodeToDo, AdditonalEmails, Causers);
+		ECProps.Add("FailEmails/" + NodeToDo + "=" + EMails);
+	
+		if (!OnlyLateUpdates)
+		{
 			string AgentReq = GUBPNodes[NodeToDo].ECAgentString();
-			if(ParseParamValue("AgentOverride") != "" && !GUBPNodes[NodeToDo].GetFullName().Contains("OnMac"))
+			if (ParseParamValue("AgentOverride") != "" && !GUBPNodes[NodeToDo].GetFullName().Contains("OnMac"))
 			{
 				AgentReq = ParseParamValue("AgentOverride");
 			}
-            ECProps.Add(string.Format("AgentRequirementString/{0}={1}", NodeToDo, AgentReq));
-            ECProps.Add(string.Format("RequiredMemory/{0}={1}", NodeToDo, GUBPNodes[NodeToDo].AgentMemoryRequirement(this)));
-            ECProps.Add(string.Format("Timeouts/{0}={1}", NodeToDo, GUBPNodes[NodeToDo].TimeoutInMinutes()));
-            ECProps.Add(string.Format("JobStepPath/{0}={1}", NodeToDo, GetJobStepPath(NodeToDo)));
-        }
+			ECProps.Add(string.Format("AgentRequirementString/{0}={1}", NodeToDo, AgentReq));
+			ECProps.Add(string.Format("RequiredMemory/{0}={1}", NodeToDo, GUBPNodes[NodeToDo].AgentMemoryRequirement(this)));
+			ECProps.Add(string.Format("Timeouts/{0}={1}", NodeToDo, GUBPNodes[NodeToDo].TimeoutInMinutes()));
+			ECProps.Add(string.Format("JobStepPath/{0}={1}", NodeToDo, GetJobStepPath(NodeToDo)));
+		}
+		
         var BuildDuration = (DateTime.UtcNow - StartTime).TotalMilliseconds;
-        Log("Took {0}s to get P4 history for node {1}", BuildDuration / 1000, NodeToDo);
-
         return ECProps;
     }
 
@@ -5107,7 +5114,7 @@
         try
         {
             Log("Updating node props for node {0}", NodeToDo);
-            string EMails;
+			string EMails = "";
             var Props = GetECPropsForNode(NodeToDo, CLString, out EMails, true);
             foreach (var Prop in Props)
             {
@@ -5675,7 +5682,8 @@
                     }
                     else if (Proj.Properties.bIsCodeBasedProject)
                     {
-                        throw new AutomationException("{0} was listed as a codeless project by GUBP_NonCodeProjects_BaseEditorTypeOnly, however it is a code based project.", Codeless.Key);
+						Branch.NonCodeProjects.Add(Proj);
+						NonCodeProjectNames.Add(Codeless.Key, Codeless.Value);                        
                     }
                     else
                     {
@@ -6019,25 +6027,28 @@
 
 				if (!bNoAutomatedTesting && HostPlatform == UnrealTargetPlatform.Win64) //temp hack till automated testing works on other platforms than Win64
                 {
-                    var EditorTests = CodeProj.Properties.Targets[TargetRules.TargetType.Editor].Rules.GUBP_GetEditorTests_EditorTypeOnly(HostPlatform);
-                    var EditorTestNodes = new List<string>();
-                    string AgentSharingGroup = "";
-                    if (EditorTests.Count > 1)
-                    {
-                        AgentSharingGroup = AgentShareName + "_EditorTests" + HostPlatformNode.StaticGetHostPlatformSuffix(HostPlatform);
-                    }
-                    foreach (var Test in EditorTests)
-                    {
-                        EditorTestNodes.Add(AddNode(new UATTestNode(this, HostPlatform, CodeProj, Test.Key, Test.Value, AgentSharingGroup)));
-						if (!Options.bTestWithShared)
+					if (CodeProj.Properties.Targets.ContainsKey(TargetRules.TargetType.Editor))
+					{
+						var EditorTests = CodeProj.Properties.Targets[TargetRules.TargetType.Editor].Rules.GUBP_GetEditorTests_EditorTypeOnly(HostPlatform);
+						var EditorTestNodes = new List<string>();
+						string AgentSharingGroup = "";
+						if (EditorTests.Count > 1)
 						{
-							RemovePseudodependencyFromNode((UATTestNode.StaticGetFullName(HostPlatform, CodeProj, Test.Key)), WaitForTestShared.StaticGetFullName());
+							AgentSharingGroup = AgentShareName + "_EditorTests" + HostPlatformNode.StaticGetHostPlatformSuffix(HostPlatform);
 						}
-                    }
-                    if (EditorTestNodes.Count > 0)
-                    {                        
-                        AddNode(new GameAggregateNode(this, HostPlatform, CodeProj, "AllEditorTests", EditorTestNodes, 0.0f));
-                    }
+						foreach (var Test in EditorTests)
+						{
+							EditorTestNodes.Add(AddNode(new UATTestNode(this, HostPlatform, CodeProj, Test.Key, Test.Value, AgentSharingGroup)));
+							if (!Options.bTestWithShared)
+							{
+								RemovePseudodependencyFromNode((UATTestNode.StaticGetFullName(HostPlatform, CodeProj, Test.Key)), WaitForTestShared.StaticGetFullName());
+							}
+						}
+						if (EditorTestNodes.Count > 0)
+						{
+							AddNode(new GameAggregateNode(this, HostPlatform, CodeProj, "AllEditorTests", EditorTestNodes, 0.0f));
+						}
+					}
                 }
 
                 var CookedAgentSharingGroup = AgentShareName + "_CookedTests" + HostPlatformNode.StaticGetHostPlatformSuffix(HostPlatform);
@@ -6812,7 +6823,7 @@
             var BuildDuration = (DateTime.UtcNow - StartTime).TotalMilliseconds;
             Log("Took {0}s to cache completion for {1} nodes", BuildDuration / 1000, NodesToDo.Count);
         }
-        if (CLString != "" && StoreName.Contains(CLString) && !ParseParam("NoHistory"))
+        /*if (CLString != "" && StoreName.Contains(CLString) && !ParseParam("NoHistory"))
         {
             Log("******* Updating history");
             var StartTime = DateTime.UtcNow;
@@ -6825,7 +6836,7 @@
             }
             var BuildDuration = (DateTime.UtcNow - StartTime).TotalMilliseconds;
             Log("Took {0}s to get history for {1} nodes", BuildDuration / 1000, NodesToDo.Count);
-        }
+        }*/
 
         var OrdereredToDo = TopologicalSort(NodesToDo, ExplicitTrigger, LocalOnly);
 
@@ -7514,6 +7525,10 @@
                     if (SaveSuccessRecords)
                     {
                         UpdateNodeHistory(NodeToDo, CLString);
+						if (IsBuildMachine)
+						{
+							GetFailureEmails(NodeToDo, CLString);
+						}
                         SaveStatus(NodeToDo, FailedTempStorageSuffix, NodeStoreName, bSaveSharedTempStorage, GameNameIfAny, ParseParamValue("MyJobStepId"));
                         UpdateECProps(NodeToDo, CLString);
 						UpdateECBuildTime(NodeToDo, BuildDuration);
@@ -7567,6 +7582,10 @@
                 if (SaveSuccessRecords) 
                 {
                     UpdateNodeHistory(NodeToDo, CLString);
+					if (IsBuildMachine)
+					{
+						GetFailureEmails(NodeToDo, CLString);
+					}
                     SaveStatus(NodeToDo, SucceededTempStorageSuffix, NodeStoreName, bSaveSharedTempStorage, GameNameIfAny);
                     UpdateECProps(NodeToDo, CLString);
 					UpdateECBuildTime(NodeToDo, BuildDuration);
diff --git a/Engine/Source/UE4Editor.Target.cs b/Engine/Source/UE4Editor.Target.cs
--- a/Engine/Source/UE4Editor.Target.cs
+++ b/Engine/Source/UE4Editor.Target.cs
@@ -120,15 +120,15 @@ public class UE4EditorTarget : TargetRules
         }
 
         //TAPPY CHICKEN
-        /*List<UnrealTargetPlatform> TappyChickenPlats = null;
+        List<UnrealTargetPlatform> TappyChickenPlats = null;
         if (HostPlatform == UnrealTargetPlatform.Mac)
         {
-            TappyChickenPlats = new List<UnrealTargetPlatform> { HostPlatform, UnrealTargetPlatform.IOS };
+            TappyChickenPlats = new List<UnrealTargetPlatform> {};
         }
         else
         {
             TappyChickenPlats = new List<UnrealTargetPlatform> { HostPlatform, UnrealTargetPlatform.Android, UnrealTargetPlatform.HTML5 };
-        }*/
+        }
         //TP_PUZZLE
         List<UnrealTargetPlatform> TPPuzzlePlats = null;
         if (HostPlatform == UnrealTargetPlatform.Mac)
@@ -184,7 +184,7 @@ public class UE4EditorTarget : TargetRules
 
         NonCodeProjectNames.Add("BlackJack", DesktopAndMobilePlats);
         NonCodeProjectNames.Add("MemoryGame", DesktopAndMobilePlats);
-        //NonCodeProjectNames.Add("TappyChicken", TappyChickenPlats);
+        NonCodeProjectNames.Add("TappyChicken", TappyChickenPlats);
         NonCodeProjectNames.Add("SwingNinja", DesktopAndMobilePlats);
         NonCodeProjectNames.Add("SunTemple", DesktopAndMobilePlats);
 
@@ -273,17 +273,14 @@ public class UE4EditorTarget : TargetRules
                     new GUBPFormalBuild(UnrealTargetPlatform.PS4, UnrealTargetConfiguration.Development, false),
             };
 
-        /*var TappyChickenBuildSettings = new List<GUBPFormalBuild>
+        var TappyChickenBuildSettings = new List<GUBPFormalBuild>
             {
                     new GUBPFormalBuild(UnrealTargetPlatform.Android, UnrealTargetConfiguration.Shipping, true),
                     new GUBPFormalBuild(UnrealTargetPlatform.Android, UnrealTargetConfiguration.Test, true),                    
-                    new GUBPFormalBuild(UnrealTargetPlatform.Android, UnrealTargetConfiguration.Development, true),          
-                    new GUBPFormalBuild(UnrealTargetPlatform.IOS, UnrealTargetConfiguration.Shipping, true),
-                    new GUBPFormalBuild(UnrealTargetPlatform.IOS, UnrealTargetConfiguration.Test, true),
-                    new GUBPFormalBuild(UnrealTargetPlatform.IOS, UnrealTargetConfiguration.Development, true),          			
+                    new GUBPFormalBuild(UnrealTargetPlatform.Android, UnrealTargetConfiguration.Development, true),        			
 					new GUBPFormalBuild(UnrealTargetPlatform.HTML5, UnrealTargetConfiguration.Shipping, false),
 					new GUBPFormalBuild(UnrealTargetPlatform.HTML5, UnrealTargetConfiguration.Test, false),
-            };*/
+            };
         #endregion
 
         #region Template Configurations
@@ -339,7 +336,7 @@ public class UE4EditorTarget : TargetRules
         var NonCodeProjectNames = new Dictionary<string, List<GUBPFormalBuild>>();
 
         //Add Samples to the list with its corresponding settings
-        //NonCodeProjectNames.Add("TappyChicken", TappyChickenBuildSettings);    
+        NonCodeProjectNames.Add("TappyChicken", TappyChickenBuildSettings);    
         NonCodeProjectNames.Add("BlackJack", BlackJackBuildSettings);
         NonCodeProjectNames.Add("MemoryGame", MemoryGameBuildSettings);
         NonCodeProjectNames.Add("SwingNinja", SwingNinjaBuildSettings);
