build script, roll back email change temporarily

diff --git a/Engine/Source/Programs/AutomationTool/Scripts/GUBP.Automation.cs b/Engine/Source/Programs/AutomationTool/Scripts/GUBP.Automation.cs
--- a/Engine/Source/Programs/AutomationTool/Scripts/GUBP.Automation.cs
+++ b/Engine/Source/Programs/AutomationTool/Scripts/GUBP.Automation.cs
@@ -4886,112 +4886,6 @@ public class GUBP : BuildCommand
 			GUBPNodesHistory.Add(Node, History);
         }
     }
-	void GetFailureEmails(string NodeToDo, string CLString, bool OnlyLateUpdates = false)
-	{
-		var StartTime = DateTime.UtcNow;        
-        string EMails = "";
-        string FailCauserEMails = "";
-        string EMailNote = "";
-        bool SendSuccessForGreenAfterRed = false;
-        int NumPeople = 0;
-        if (GUBPNodesHistory.ContainsKey(NodeToDo))
-        {
-            var History = GUBPNodesHistory[NodeToDo];
-			RunECTool(String.Format("setProperty \"/myWorkflow/LastGreen/{0}\" \"{1}\"", NodeToDo, History.LastSucceeded), true);
-			RunECTool(String.Format("setProperty \"/myWorkflow/LastGreen/{0}\" \"{1}\"", NodeToDo, History.FailedString), true);
-
-            if (History.LastSucceeded > 0 && History.LastSucceeded < P4Env.Changelist)
-            {
-                int LastNonDuplicateFail = P4Env.Changelist;
-                try
-                {
-                    if (OnlyLateUpdates)
-                    {
-                        LastNonDuplicateFail = FindLastNonDuplicateFail(NodeToDo, CLString);
-                        if (LastNonDuplicateFail < P4Env.Changelist)
-                        {
-                            Log("*** Red-after-red spam reduction, changed CL {0} to CL {1} because the errors didn't change.", P4Env.Changelist, LastNonDuplicateFail);
-                        }
-                    }
-                }
-                catch (Exception Ex)
-                {
-                    LastNonDuplicateFail = P4Env.Changelist;
-                    Log(System.Diagnostics.TraceEventType.Warning, "Failed to FindLastNonDuplicateFail.");
-                    Log(System.Diagnostics.TraceEventType.Warning, LogUtils.FormatException(Ex));
-                }
-
-                var ChangeRecords = GetChanges(History.LastSucceeded, LastNonDuplicateFail, History.LastSucceeded);
-                foreach (var Record in ChangeRecords)
-                {
-                    FailCauserEMails = GUBPNode.MergeSpaceStrings(FailCauserEMails, Record.UserEmail);
-                }
-                if (!String.IsNullOrEmpty(FailCauserEMails))
-                {
-                    NumPeople++;
-                    foreach (var AChar in FailCauserEMails.ToCharArray())
-                    {
-                        if (AChar == ' ')
-                        {
-                            NumPeople++;
-                        }
-                    }
-                    if (NumPeople > 50)
-                    {
-                        EMailNote = String.Format("This step has been broken for more than 50 changes. It last succeeded at CL {0}. ", History.LastSucceeded);
-                    }
-                }
-            }
-            else if (History.LastSucceeded <= 0)
-            {
-                EMailNote = String.Format("This step has been broken for more than a few days, so there is no record of it ever succeeding. ");
-            }
-            if (EMailNote != "" && !String.IsNullOrEmpty(History.FailedString))
-            {
-                EMailNote += String.Format("It has failed at CLs {0}. ", History.FailedString);
-            }
-            if (EMailNote != "" && !String.IsNullOrEmpty(History.InProgressString))
-            {
-                EMailNote += String.Format("These CLs are being built right now {0}. ", History.InProgressString);
-            }
-            if (History.LastSucceeded > 0 && History.LastSucceeded < P4Env.Changelist && History.LastFailed > History.LastSucceeded && History.LastFailed < P4Env.Changelist)
-            {
-                SendSuccessForGreenAfterRed = ParseParam("CIS");
-            }
-        }
-        else
-        {
-			RunECTool(String.Format("setProperty \"/myWorkflow/LastGreen/{0}\" \"{1}\"", NodeToDo, "0"));
-			RunECTool(String.Format("setProperty \"/myWorkflow/RedsSince/{0}\" \"{1}\"", NodeToDo, ""));
-        }
-		RunECTool(String.Format("setProperty \"/myWorkflow/FailCausers/{0}\" \"{1}\"", NodeToDo, FailCauserEMails));
-		RunECTool(String.Format("setProperty \"/myWorkflow/EmailNotes/{0}\" \"{1}\"", NodeToDo, EMailNote));
-        {
-            var AdditonalEmails = "";
-
-			string Causers = "";
-            if (ParseParam("CIS") && !GUBPNodes[NodeToDo].SendSuccessEmail() && !GUBPNodes[NodeToDo].TriggerNode())
-            {
-				Causers = FailCauserEMails;
-           }
-            string AddEmails = ParseParamValue("AddEmails");
-            if (!String.IsNullOrEmpty(AddEmails))
-            {
-                AdditonalEmails = GUBPNode.MergeSpaceStrings(AddEmails, AdditonalEmails);
-            }
-
-            EMails = GetEMailListForNode(this, NodeToDo, AdditonalEmails, Causers);
-			RunECTool(String.Format("setProperty \"/myWorkflow/FailEmails/{0}\" \"{1}\"", NodeToDo, EMails));            
-        }
-		if (GUBPNodes[NodeToDo].SendSuccessEmail() || SendSuccessForGreenAfterRed)
-		{
-			RunECTool(String.Format("setProperty \"/myWorkflow/SendSuccessEmail/{0}\" \"{1}\"", NodeToDo, "1"));
-		}
-		else
-		{
-			RunECTool(String.Format("setProperty \"/myWorkflow/SendSuccessEmail/{0}\" \"{1}\"", NodeToDo, "0"));
-		}
-	}
 
     bool HashSetEqual(HashSet<string> A, HashSet<string> B)
     {
@@ -6823,7 +6717,7 @@ public class GUBP : BuildCommand
             var BuildDuration = (DateTime.UtcNow - StartTime).TotalMilliseconds;
             Log("Took {0}s to cache completion for {1} nodes", BuildDuration / 1000, NodesToDo.Count);
         }
-        /*if (CLString != "" && StoreName.Contains(CLString) && !ParseParam("NoHistory"))
+        if (CLString != "" && StoreName.Contains(CLString) && !ParseParam("NoHistory"))
         {
             Log("******* Updating history");
             var StartTime = DateTime.UtcNow;
@@ -6836,7 +6730,7 @@ public class GUBP : BuildCommand
             }
             var BuildDuration = (DateTime.UtcNow - StartTime).TotalMilliseconds;
             Log("Took {0}s to get history for {1} nodes", BuildDuration / 1000, NodesToDo.Count);
-        }*/
+        }
 
         var OrdereredToDo = TopologicalSort(NodesToDo, ExplicitTrigger, LocalOnly);
 
@@ -7524,11 +7418,7 @@ public class GUBP : BuildCommand
                 {
                     if (SaveSuccessRecords)
                     {
-                        UpdateNodeHistory(NodeToDo, CLString);
-						if (IsBuildMachine)
-						{
-							GetFailureEmails(NodeToDo, CLString);
-						}
+                        UpdateNodeHistory(NodeToDo, CLString);						
                         SaveStatus(NodeToDo, FailedTempStorageSuffix, NodeStoreName, bSaveSharedTempStorage, GameNameIfAny, ParseParamValue("MyJobStepId"));
                         UpdateECProps(NodeToDo, CLString);
 						UpdateECBuildTime(NodeToDo, BuildDuration);
@@ -7581,11 +7471,7 @@ public class GUBP : BuildCommand
                 }
                 if (SaveSuccessRecords) 
                 {
-                    UpdateNodeHistory(NodeToDo, CLString);
-					if (IsBuildMachine)
-					{
-						GetFailureEmails(NodeToDo, CLString);
-					}
+                    UpdateNodeHistory(NodeToDo, CLString);					
                     SaveStatus(NodeToDo, SucceededTempStorageSuffix, NodeStoreName, bSaveSharedTempStorage, GameNameIfAny);
                     UpdateECProps(NodeToDo, CLString);
 					UpdateECBuildTime(NodeToDo, BuildDuration);
