build script, fixes for noncode projects as code projects

diff --git a/Engine/Build/Target.cs.template b/Engine/Build/Target.cs.template
--- a/Engine/Build/Target.cs.template
+++ b/Engine/Build/Target.cs.template
@@ -97,6 +97,7 @@ public class {GAME_NAME}Target : TargetRules
     {
         var Result = new GUBPProjectOptions();
         Result.bTestWithShared = true;
+		Result.bIsNonCode = true;
         return Result;
     }
 
diff --git a/Engine/Source/Programs/AutomationTool/Scripts/GUBP.Automation.cs b/Engine/Source/Programs/AutomationTool/Scripts/GUBP.Automation.cs
--- a/Engine/Source/Programs/AutomationTool/Scripts/GUBP.Automation.cs
+++ b/Engine/Source/Programs/AutomationTool/Scripts/GUBP.Automation.cs
@@ -2173,7 +2173,10 @@ public class GUBP : BuildCommand
                         }
                         else if (Options.bTestWithShared)
                         {
-                            AddDependency(EditorGameNode.StaticGetFullName(HostPlatform, CodeProj)); // if we are just testing, we will still include the editor stuff
+							if (!Options.bIsNonCode)
+							{
+								AddDependency(EditorGameNode.StaticGetFullName(HostPlatform, CodeProj)); // if we are just testing, we will still include the editor stuff
+							}
                         }
                     }
                 }
@@ -6020,7 +6023,7 @@ public class GUBP : BuildCommand
                     AgentShareName = "Shared";
                 }
 
-				if (!BranchOptions.ExcludePlatformsForEditor.Contains(HostPlatform))
+				if (!BranchOptions.ExcludePlatformsForEditor.Contains(HostPlatform) && !Options.bIsNonCode)
 				{
 					AddNode(new EditorGameNode(this, HostPlatform, CodeProj));
 				}
diff --git a/Engine/Source/Programs/UnrealBuildTool/System/RulesCompiler.cs b/Engine/Source/Programs/UnrealBuildTool/System/RulesCompiler.cs
--- a/Engine/Source/Programs/UnrealBuildTool/System/RulesCompiler.cs
+++ b/Engine/Source/Programs/UnrealBuildTool/System/RulesCompiler.cs
@@ -772,6 +772,7 @@ namespace UnrealBuildTool
             public bool bTestWithShared = false;
             public bool bIsMassive = false;
             public bool bCustomWorkflowForPromotion = false;
+			public bool bIsNonCode = false;
         }
         public virtual GUBPProjectOptions GUBP_IncludeProjectInPromotedBuild_EditorTypeOnly(UnrealTargetPlatform HostPlatform)
         {
