Making sure AutomationTool does not let UBT delete junk files multiple times in the same AutomationTool run. Added the option to completely disable junk removal in UAT.

diff --git a/Engine/Source/Programs/AutomationTool/Automation.cs b/Engine/Source/Programs/AutomationTool/Automation.cs
--- a/Engine/Source/Programs/AutomationTool/Automation.cs
+++ b/Engine/Source/Programs/AutomationTool/Automation.cs
@@ -113,7 +113,7 @@ namespace AutomationTool
 		public static CommandLineArg Installed = new CommandLineArg("-Installed");
 		public static CommandLineArg UTF8Output = new CommandLineArg("-UTF8Output");
 		public static CommandLineArg NoAutoSDK = new CommandLineArg("-NoAutoSDK");
-
+		public static CommandLineArg IgnoreJunk = new CommandLineArg("-ignorejunk");
 		/// <summary>
 		/// Force initialize static members by calling this.
 		/// </summary>
@@ -140,6 +140,7 @@ AutomationTool.exe [-verbose] [-compileonly] [-p4] Command0 [-Arg0 -Arg1 -Arg2 
 	[Help("submit", "Allows UAT command to submit changes")]
 	[Help("nosubmit", "Prevents any submit attempts")]
 	[Help("nokill", "Does not kill any spawned processes on exit")]
+	[Help("ignorejunk", "Prevents UBT from cleaning junk files")]
 	static class Automation
 	{
 		#region Command line parsing
diff --git a/Engine/Source/Programs/AutomationTool/UBTUtils.cs b/Engine/Source/Programs/AutomationTool/UBTUtils.cs
--- a/Engine/Source/Programs/AutomationTool/UBTUtils.cs
+++ b/Engine/Source/Programs/AutomationTool/UBTUtils.cs
@@ -8,6 +8,11 @@ namespace AutomationTool
 
 	public partial class CommandUtils
 	{
+		/// <summary>
+		/// True if UBT has deleted junk files
+		/// </summary>
+		private static bool bJunkDeleted = false;
+
 		/// <summary>
 		/// Runs UBT with the specified commandline. Automatically creates a logfile. When 
 		/// no LogName is specified, the executable name is used as logfile base name.
@@ -31,6 +36,16 @@ namespace AutomationTool
 				CommandLine += " -nocreatestub";
 			}
 			CommandLine += " -NoHotReloadFromIDE";
+			if (bJunkDeleted || GlobalCommandLine.IgnoreJunk)
+			{
+				// UBT has already deleted junk files, make sure it doesn't do it again
+				CommandLine += " -ignorejunk";
+			}
+			else
+			{
+				// UBT will delete junk on first run
+				bJunkDeleted = true;
+			}
 
 			CommandUtils.RunAndLog(Env, UBTExecutable, CommandLine, LogName, EnvVars: EnvVars);
 		}
