fix for Plugins not being found when -nocompile is provided
fix for generating Target.cs when -nocodeproject is provided (removed -nocodeproject as it is no longer needed due to other fixes)
default plugins no longer cause a target file to be generated

diff --git a/Engine/Source/Editor/MainFrame/Private/Frame/MainFrameActions.cpp b/Engine/Source/Editor/MainFrame/Private/Frame/MainFrameActions.cpp
--- a/Engine/Source/Editor/MainFrame/Private/Frame/MainFrameActions.cpp
+++ b/Engine/Source/Editor/MainFrame/Private/Frame/MainFrameActions.cpp
@@ -286,16 +286,6 @@ const TCHAR* GetUATCompilationFlags()
 		: TEXT("-nocompileeditor");
 }
 
-/**
- * Gets -nocodeproject flag for UAT if it's needed.
- */
-const TCHAR* GetUATNoCodeProjectFlag()
-{
-	IHotReloadInterface& HotReloadSupport = FModuleManager::LoadModuleChecked<IHotReloadInterface>("HotReload");
-	// If there is no game module loaded, source code is not available.
-	return !HotReloadSupport.IsAnyGameModuleLoaded() ? TEXT(" -nocodeproject") : TEXT("");
-}
-
 void FMainFrameActionCallbacks::CookContent(const FName InPlatformInfoName)
 {
 	const PlatformInfo::FPlatformInfo* const PlatformInfo = PlatformInfo::FindPlatformInfo(InPlatformInfoName);
@@ -335,10 +325,9 @@ void FMainFrameActionCallbacks::CookContent(const FName InPlatformInfoName)
 	}
 
 	FString ProjectPath = FPaths::IsProjectFilePathSet() ? FPaths::ConvertRelativePathToFull(FPaths::GetProjectFilePath()) : FPaths::RootDir() / FApp::GetGameName() / FApp::GetGameName() + TEXT(".uproject");
-	FString CommandLine = FString::Printf(TEXT("BuildCookRun %s%s%s%s -nop4 -project=\"%s\" -cook -ue4exe=%s %s"),
+	FString CommandLine = FString::Printf(TEXT("BuildCookRun %s%s%s -nop4 -project=\"%s\" -cook -ue4exe=%s %s"),
 		FRocketSupport::IsRocket() ? TEXT("-rocket ") : TEXT(""),
 		GetUATCompilationFlags(),
-		GetUATNoCodeProjectFlag(),
 		FApp::IsEngineInstalled() ? TEXT(" -installed") : TEXT(""),
 		*ProjectPath,
 		*FUnrealEdMisc::Get().GetExecutableForCommandlets(),
@@ -610,10 +599,9 @@ void FMainFrameActionCallbacks::PackageProject( const FName InPlatformInfoName )
 	Configuration = Configuration.Replace(TEXT("PPBC_"), TEXT(""));
 
 	FString ProjectPath = FPaths::IsProjectFilePathSet() ? FPaths::ConvertRelativePathToFull(FPaths::GetProjectFilePath()) : FPaths::RootDir() / FApp::GetGameName() / FApp::GetGameName() + TEXT(".uproject");
-	FString CommandLine = FString::Printf(TEXT("BuildCookRun %s%s%s%s -nop4 -project=\"%s\" -cook -stage -archive -archivedirectory=\"%s\" -package -clientconfig=%s -ue4exe=%s %s -utf8output"),
+	FString CommandLine = FString::Printf(TEXT("BuildCookRun %s%s%s -nop4 -project=\"%s\" -cook -stage -archive -archivedirectory=\"%s\" -package -clientconfig=%s -ue4exe=%s %s -utf8output"),
 		FRocketSupport::IsRocket() ? TEXT("-rocket ") : TEXT(""),
 		GetUATCompilationFlags(),
-		GetUATNoCodeProjectFlag(),
 		FApp::IsEngineInstalled() ? TEXT(" -installed") : TEXT(""),
 		*ProjectPath,
 		*PackagingSettings->StagingDirectory.Path,
diff --git a/Engine/Source/Programs/AutomationTool/Automation.cs b/Engine/Source/Programs/AutomationTool/Automation.cs
--- a/Engine/Source/Programs/AutomationTool/Automation.cs
+++ b/Engine/Source/Programs/AutomationTool/Automation.cs
@@ -105,7 +105,6 @@ namespace AutomationTool
 		public static CommandLineArg Preprocess = new CommandLineArg("-Preprocess");
 		public static CommandLineArg NoCompile = new CommandLineArg("-NoCompile");
 		public static CommandLineArg NoCompileEditor = new CommandLineArg("-NoCompileEditor");
-		public static CommandLineArg NoCodeProject = new CommandLineArg("-NoCodeProject");
 		public static CommandLineArg Help = new CommandLineArg("-Help");
 		public static CommandLineArg List = new CommandLineArg("-List");
 		public static CommandLineArg Rocket = new CommandLineArg("-Rocket");
diff --git a/Engine/Source/Programs/AutomationTool/ProjectUtils.cs b/Engine/Source/Programs/AutomationTool/ProjectUtils.cs
--- a/Engine/Source/Programs/AutomationTool/ProjectUtils.cs
+++ b/Engine/Source/Programs/AutomationTool/ProjectUtils.cs
@@ -179,41 +179,59 @@ namespace AutomationTool
 				}
 			}
 
+			// Change the working directory to be the Engine/Source folder. We are running from Engine/Binaries/DotNET
+			string oldCWD = Directory.GetCurrentDirectory();
+			string EngineSourceDirectory = Path.Combine(UnrealBuildTool.Utils.GetExecutingAssemblyDirectory(), "..", "..", "..", "Engine", "Source");
+			if (!Directory.Exists(EngineSourceDirectory)) // only set the directory if it exists, this should only happen if we are launching the editor from an artist sync
+			{
+				EngineSourceDirectory = Path.Combine(UnrealBuildTool.Utils.GetExecutingAssemblyDirectory(), "..", "..", "..", "Engine", "Binaries");
+			}
+			Directory.SetCurrentDirectory(EngineSourceDirectory);
+
+			bool RetVal = false;
 			// check the target platforms for any differences in build settings or additional plugins
 			foreach (UnrealTargetPlatform TargetPlatformType in TargetPlatforms)
 			{
 				IUEBuildPlatform BuildPlat = UEBuildPlatform.GetBuildPlatform(TargetPlatformType, true);
 				if (BuildPlat != null && !(BuildPlat as UEBuildPlatform).HasDefaultBuildConfig(TargetPlatformType, Path.GetDirectoryName(RawProjectPath)))
 				{
-					return true;
+					RetVal = true;
+					break;
 				}
+
 				// find if there are any plugins
 				List<string> PluginList = new List<string>();
 				// Use the project settings to update the plugin list for this target
 				PluginList = UProjectInfo.GetEnabledPlugins(RawProjectPath, PluginList, TargetPlatformType);
-				if (PluginList.Count > 0)
+				if (PluginList.Count > 0 && !RetVal)
 				{
 					foreach (var PluginName in PluginList)
 					{
 						// check the plugin info for this plugin itself
 						foreach (var Plugin in Plugins.AllPlugins)
 						{
-							if (Plugin.Name == PluginName)
+							if (Plugin.Name == PluginName && !Plugin.bEnabledByDefault && !RetVal)
 							{
 								foreach (var Module in Plugin.Modules)
 								{
-									if (Module.Platforms.Count > 0 && Module.Platforms.Contains(TargetPlatformType))
+									if (Module.Platforms.Count > 0 && Module.Platforms.Contains(TargetPlatformType) && !RetVal)
 									{
-										return true;
+										RetVal = true;
+										break;
 									}
 								}
 								break;
 							}
 						}
+						if (RetVal)
+						{
+							break;
+						}
 					}
 				}
 			}
-			return false;
+			Directory.SetCurrentDirectory(oldCWD);
+			return RetVal;
 		}
 
 		private static void GenerateTempTarget(string RawProjectPath)
@@ -466,11 +484,6 @@ namespace AutomationTool
 		{
 			CommandUtils.Log("Compiling targets DLL: {0}", TargetsDllFilename);
 
-			if (!DoNotCompile && GlobalCommandLine.NoCodeProject)
-			{
-				//throw new AutomationException("Building is not supported when -nocodeproject flag is provided.");
-			}
-
 			var ReferencedAssemblies = new List<string>() 
 					{ 
 						"System.dll", 
diff --git a/Engine/Source/Programs/AutomationTool/ScriptCompiler.cs b/Engine/Source/Programs/AutomationTool/ScriptCompiler.cs
--- a/Engine/Source/Programs/AutomationTool/ScriptCompiler.cs
+++ b/Engine/Source/Programs/AutomationTool/ScriptCompiler.cs
@@ -70,7 +70,7 @@ namespace AutomationTool
 			}
 			// Register all the classes inside UBT
 			Log.TraceVerbose("Registering UBT Classes.");
-			UnrealBuildTool.UnrealBuildTool.RegisterAllUBTClasses(GlobalCommandLine.NoCodeProject);
+			UnrealBuildTool.UnrealBuildTool.RegisterAllUBTClasses();
 			Environment.CurrentDirectory = OldCWD;
 
 			// Compile only if not disallowed.
diff --git a/Engine/Source/Programs/AutomationTool/UE4Build.cs b/Engine/Source/Programs/AutomationTool/UE4Build.cs
--- a/Engine/Source/Programs/AutomationTool/UE4Build.cs
+++ b/Engine/Source/Programs/AutomationTool/UE4Build.cs
@@ -1135,11 +1135,6 @@ namespace AutomationTool
 		/// <param name="InForceNoXGE">If true will force XGE off</param>
 		public void Build(BuildAgenda Agenda, bool? InDeleteBuildProducts = null, bool InUpdateVersionFiles = true, bool InForceNoXGE = false, bool InForceNonUnity = false, bool InForceUnity = false, Dictionary<UnrealBuildTool.UnrealTargetPlatform, Dictionary<string, string>> PlatformEnvVars = null)
 		{
-			if (GlobalCommandLine.NoCodeProject)
-			{
-				throw new AutomationException("Building is not supported when -nocodeproject flag is provided.");
-			}
-
 			if (!CmdEnv.HasCapabilityToCompile)
 			{
 				throw new AutomationException("You are attempting to compile on a machine that does not have a supported compiler!");
diff --git a/Engine/Source/Programs/UnrealBuildTool/System/Plugins.cs b/Engine/Source/Programs/UnrealBuildTool/System/Plugins.cs
--- a/Engine/Source/Programs/UnrealBuildTool/System/Plugins.cs
+++ b/Engine/Source/Programs/UnrealBuildTool/System/Plugins.cs
@@ -447,10 +447,13 @@ namespace UnrealBuildTool
 				Plugins.FindPluginsIn( EnginePluginsDirectory, PluginInfo.LoadedFromType.Engine, ref AllPluginsVar );
 
 				// Game plugins
-				foreach( var GameProjectFolder in RulesCompiler.AllGameFolders )
+				if (RulesCompiler.AllGameFolders != null)
 				{
-					var GamePluginsDirectory = Path.Combine( GameProjectFolder, "Plugins" );
-					Plugins.FindPluginsIn( GamePluginsDirectory, PluginInfo.LoadedFromType.GameProject, ref AllPluginsVar );
+					foreach (var GameProjectFolder in RulesCompiler.AllGameFolders)
+					{
+						var GamePluginsDirectory = Path.Combine(GameProjectFolder, "Plugins");
+						Plugins.FindPluginsIn(GamePluginsDirectory, PluginInfo.LoadedFromType.GameProject, ref AllPluginsVar);
+					}
 				}
 
 				// Also keep track of which modules map to which plugins
diff --git a/Engine/Source/Programs/UnrealBuildTool/System/Utils.cs b/Engine/Source/Programs/UnrealBuildTool/System/Utils.cs
--- a/Engine/Source/Programs/UnrealBuildTool/System/Utils.cs
+++ b/Engine/Source/Programs/UnrealBuildTool/System/Utils.cs
@@ -833,7 +833,7 @@ namespace UnrealBuildTool
 		/// assembly path in case shadow copying is enabled.
 		/// </summary>
 		/// <returns>Absolute path to the executing assembly including the assembly filename.</returns>
-		internal static string GetExecutingAssemblyLocation()
+		public static string GetExecutingAssemblyLocation()
 		{
 			return new Uri(System.Reflection.Assembly.GetExecutingAssembly().CodeBase).LocalPath;
 		}
@@ -844,7 +844,7 @@ namespace UnrealBuildTool
 		/// assembly directory in case shadow copying is enabled.
 		/// </summary>
 		/// <returns>Absolute path to the directory containing the executing assembly.</returns>
-		internal static string GetExecutingAssemblyDirectory()
+		public static string GetExecutingAssemblyDirectory()
 		{
 			return Path.GetDirectoryName(GetExecutingAssemblyLocation());
 		}
