Disable -js-console on HTML5 firefox device as it causes v36 to lock up on launch.

diff --git a/Engine/Source/Programs/AutomationTool/HTML5/HTML5Platform.Automation.cs b/Engine/Source/Programs/AutomationTool/HTML5/HTML5Platform.Automation.cs
--- a/Engine/Source/Programs/AutomationTool/HTML5/HTML5Platform.Automation.cs
+++ b/Engine/Source/Programs/AutomationTool/HTML5/HTML5Platform.Automation.cs
@@ -299,7 +299,7 @@ public class HTML5Platform : Platform
 		}
 		else if (LowerBrowserPath.Contains("firefox"))
 		{
-			args += String.Format("-no-remote -profile \"{0}\" -jsconsole", Path.Combine(ProfileDirectory, "firefox"));
+			args += String.Format("-no-remote -profile \"{0}\"", Path.Combine(ProfileDirectory, "firefox"));
 		}
 		//else if (browserPath.Contains("Safari")) {}
 
