Add a FindGameChecked() utility function, rather than having to check manually each time.

diff --git a/Engine/Source/Programs/AutomationTool/ProjectUtils.cs b/Engine/Source/Programs/AutomationTool/ProjectUtils.cs
--- a/Engine/Source/Programs/AutomationTool/ProjectUtils.cs
+++ b/Engine/Source/Programs/AutomationTool/ProjectUtils.cs
@@ -772,6 +772,15 @@ namespace AutomationTool
             }
             return null;
         }
+		public BranchUProject FindGameChecked(string GameName)
+		{
+			BranchUProject Project = FindGame(GameName);
+			if(Project == null)
+			{
+				throw new AutomationException("Cannot find project '{0}' in branch", GameName);
+			}
+			return Project;
+		}
         public SingleTargetProperties FindProgram(string ProgramName)
         {
             foreach (var Proj in BaseEngineProject.Properties.Programs)
