Remove bBuildRocket from GUBP_AlwaysBuildWithTools; doesn't seem to be needed any more (and if it is, we need to fix it in a way that doesn't make "rocket" distinctions)

diff --git a/Engine/Source/Programs/AutomationTool/ProjectUtils.cs b/Engine/Source/Programs/AutomationTool/ProjectUtils.cs
--- a/Engine/Source/Programs/AutomationTool/ProjectUtils.cs
+++ b/Engine/Source/Programs/AutomationTool/ProjectUtils.cs
@@ -690,7 +690,7 @@ namespace AutomationTool
                         bool bInternalToolOnly;
                         bool SeparateNode;
 						bool CrossCompile;
-                        bool Tool = ThisTarget.Rules.GUBP_AlwaysBuildWithTools(HostPlatform, false, out bInternalToolOnly, out SeparateNode, out CrossCompile);
+                        bool Tool = ThisTarget.Rules.GUBP_AlwaysBuildWithTools(HostPlatform, out bInternalToolOnly, out SeparateNode, out CrossCompile);
 
                         CommandUtils.Log("            TargetName                    : " + ThisTarget.TargetName);
                         CommandUtils.Log("              Build With Editor           : " + (ThisTarget.Rules.GUBP_AlwaysBuildWithBaseEditor() ? "YES" : "NO"));
diff --git a/Engine/Source/Programs/AutomationTool/Scripts/GUBP.Automation.cs b/Engine/Source/Programs/AutomationTool/Scripts/GUBP.Automation.cs
--- a/Engine/Source/Programs/AutomationTool/Scripts/GUBP.Automation.cs
+++ b/Engine/Source/Programs/AutomationTool/Scripts/GUBP.Automation.cs
@@ -1019,7 +1019,7 @@ public class GUBP : BuildCommand
                 bool bInternalOnly;
                 bool SeparateNode;
 				bool CrossCompile;
-                if (ProgramTarget.Rules.GUBP_AlwaysBuildWithTools(HostPlatform, GUBP.bBuildRocket, out bInternalOnly, out SeparateNode, out CrossCompile) && ProgramTarget.Rules.SupportsPlatform(HostPlatform) && !bInternalOnly && !SeparateNode)
+                if (ProgramTarget.Rules.GUBP_AlwaysBuildWithTools(HostPlatform, out bInternalOnly, out SeparateNode, out CrossCompile) && ProgramTarget.Rules.SupportsPlatform(HostPlatform) && !bInternalOnly && !SeparateNode)
                 {
                     foreach (var Plat in ProgramTarget.Rules.GUBP_ToolPlatforms(HostPlatform))
                     {
@@ -1072,7 +1072,7 @@ public class GUBP : BuildCommand
 				bool bInternalOnly;
 				bool SeparateNode;
 				bool CrossCompile;
-				if (ProgramTarget.Rules.GUBP_AlwaysBuildWithTools(HostPlatform, GUBP.bBuildRocket, out bInternalOnly, out SeparateNode, out CrossCompile) && ProgramTarget.Rules.SupportsPlatform(HostPlatform) && !bInternalOnly && !SeparateNode && CrossCompile)
+				if (ProgramTarget.Rules.GUBP_AlwaysBuildWithTools(HostPlatform, out bInternalOnly, out SeparateNode, out CrossCompile) && ProgramTarget.Rules.SupportsPlatform(HostPlatform) && !bInternalOnly && !SeparateNode && CrossCompile)
 				{
 					foreach (var Config in ProgramTarget.Rules.GUBP_ToolConfigs(HostPlatform))
 					{
@@ -1201,7 +1201,7 @@ public class GUBP : BuildCommand
                 bool bInternalOnly;
                 bool SeparateNode;
 				bool CrossCompile;
-                if (ProgramTarget.Rules.GUBP_AlwaysBuildWithTools(HostPlatform, GUBP.bBuildRocket, out bInternalOnly, out SeparateNode, out CrossCompile) && ProgramTarget.Rules.SupportsPlatform(HostPlatform) && bInternalOnly && !SeparateNode)
+                if (ProgramTarget.Rules.GUBP_AlwaysBuildWithTools(HostPlatform, out bInternalOnly, out SeparateNode, out CrossCompile) && ProgramTarget.Rules.SupportsPlatform(HostPlatform) && bInternalOnly && !SeparateNode)
                 {
                     foreach (var Plat in ProgramTarget.Rules.GUBP_ToolPlatforms(HostPlatform))
                     {
@@ -5614,7 +5614,7 @@ public class GUBP : BuildCommand
 					bool SeparateNode;
 					bool CrossCompile;
 
-					if (ProgramTarget.Rules.GUBP_AlwaysBuildWithTools(HostPlatform, GUBP.bBuildRocket, out bInternalOnly, out SeparateNode, out CrossCompile) && ProgramTarget.Rules.SupportsPlatform(HostPlatform) && SeparateNode)
+					if (ProgramTarget.Rules.GUBP_AlwaysBuildWithTools(HostPlatform, out bInternalOnly, out SeparateNode, out CrossCompile) && ProgramTarget.Rules.SupportsPlatform(HostPlatform) && SeparateNode)
 					{
 						if (bInternalOnly)
 						{
@@ -5638,7 +5638,7 @@ public class GUBP : BuildCommand
 						bool SeparateNode;
 						bool CrossCompile;
 
-						if(ProgramTarget.Rules.GUBP_AlwaysBuildWithTools(HostPlatform, GUBP.bBuildRocket, out bInternalNodeOnly, out SeparateNode, out CrossCompile) && ProgramTarget.Rules.SupportsPlatform(HostPlatform) && SeparateNode)
+						if(ProgramTarget.Rules.GUBP_AlwaysBuildWithTools(HostPlatform, out bInternalNodeOnly, out SeparateNode, out CrossCompile) && ProgramTarget.Rules.SupportsPlatform(HostPlatform) && SeparateNode)
 						{
 							if(bInternalNodeOnly)
 							{
diff --git a/Engine/Source/Programs/BuildPatchTool/BuildPatchTool.Target.cs b/Engine/Source/Programs/BuildPatchTool/BuildPatchTool.Target.cs
--- a/Engine/Source/Programs/BuildPatchTool/BuildPatchTool.Target.cs
+++ b/Engine/Source/Programs/BuildPatchTool/BuildPatchTool.Target.cs
@@ -50,7 +50,7 @@ public class BuildPatchToolTarget : TargetRules
 		OutLinkEnvironmentConfiguration.bBuildAdditionalConsoleApplication = false;
 		OutLinkEnvironmentConfiguration.bHasExports = false;
 	}
-    public override bool GUBP_AlwaysBuildWithTools(UnrealTargetPlatform InHostPlatform, bool bBuildingRocket, out bool bInternalToolOnly, out bool SeparateNode, out bool CrossCompile)
+    public override bool GUBP_AlwaysBuildWithTools(UnrealTargetPlatform InHostPlatform, out bool bInternalToolOnly, out bool SeparateNode, out bool CrossCompile)
     {
         bInternalToolOnly = true;
         SeparateNode = true;
diff --git a/Engine/Source/Programs/CrashReporter/CrashReportClient/CrashReportClient.Target.cs b/Engine/Source/Programs/CrashReporter/CrashReportClient/CrashReportClient.Target.cs
--- a/Engine/Source/Programs/CrashReporter/CrashReportClient/CrashReportClient.Target.cs
+++ b/Engine/Source/Programs/CrashReporter/CrashReportClient/CrashReportClient.Target.cs
@@ -97,7 +97,7 @@ public class CrashReportClientTarget : TargetRules
 
 		OutCPPEnvironmentConfiguration.Definitions.Add( "USE_CHECKS_IN_SHIPPING=1" );
 	}
-    public override bool GUBP_AlwaysBuildWithTools(UnrealTargetPlatform InHostPlatform, bool bBuildingRocket, out bool bInternalToolOnly, out bool SeparateNode, out bool CrossCompile)
+    public override bool GUBP_AlwaysBuildWithTools(UnrealTargetPlatform InHostPlatform, out bool bInternalToolOnly, out bool SeparateNode, out bool CrossCompile)
 	{
 		bInternalToolOnly = false;
 		SeparateNode = false;
diff --git a/Engine/Source/Programs/CrashReporter/MinidumpDiagnostics/MinidumpDiagnostics.Target.cs b/Engine/Source/Programs/CrashReporter/MinidumpDiagnostics/MinidumpDiagnostics.Target.cs
--- a/Engine/Source/Programs/CrashReporter/MinidumpDiagnostics/MinidumpDiagnostics.Target.cs
+++ b/Engine/Source/Programs/CrashReporter/MinidumpDiagnostics/MinidumpDiagnostics.Target.cs
@@ -62,7 +62,7 @@ public class MinidumpDiagnosticsTarget : TargetRules
 
 		OutCPPEnvironmentConfiguration.Definitions.Add("MINIDUMPDIAGNOSTICS=1");
 	}
-    public override bool GUBP_AlwaysBuildWithTools(UnrealTargetPlatform InHostPlatform, bool bBuildingRocket, out bool bInternalToolOnly, out bool SeparateNode, out bool CrossCompile)
+    public override bool GUBP_AlwaysBuildWithTools(UnrealTargetPlatform InHostPlatform, out bool bInternalToolOnly, out bool SeparateNode, out bool CrossCompile)
     {
         bInternalToolOnly = true;
         SeparateNode = false;
diff --git a/Engine/Source/Programs/Mac/DsymExporter/DsymExporter.Target.cs b/Engine/Source/Programs/Mac/DsymExporter/DsymExporter.Target.cs
--- a/Engine/Source/Programs/Mac/DsymExporter/DsymExporter.Target.cs
+++ b/Engine/Source/Programs/Mac/DsymExporter/DsymExporter.Target.cs
@@ -54,7 +54,7 @@ public class DsymExporterTarget : TargetRules
 		// Do NOT produce additional console app exe
 		OutLinkEnvironmentConfiguration.bIsBuildingConsoleApplication = true;
 	}
-    public override bool GUBP_AlwaysBuildWithTools(UnrealTargetPlatform InHostPlatform, bool bBuildingRocket, out bool bInternalToolOnly, out bool SeparateNode, out bool CrossCompile)
+    public override bool GUBP_AlwaysBuildWithTools(UnrealTargetPlatform InHostPlatform, out bool bInternalToolOnly, out bool SeparateNode, out bool CrossCompile)
     {
         bInternalToolOnly = false;
         SeparateNode = false;
diff --git a/Engine/Source/Programs/Mac/UE4EditorServices/UE4EditorServices.Target.cs b/Engine/Source/Programs/Mac/UE4EditorServices/UE4EditorServices.Target.cs
--- a/Engine/Source/Programs/Mac/UE4EditorServices/UE4EditorServices.Target.cs
+++ b/Engine/Source/Programs/Mac/UE4EditorServices/UE4EditorServices.Target.cs
@@ -49,7 +49,7 @@ public class UE4EditorServicesTarget : TargetRules
 		UEBuildConfiguration.bCompileAgainstEngine = false;
 		UEBuildConfiguration.bCompileAgainstCoreUObject = false;
 	}
-    public override bool GUBP_AlwaysBuildWithTools(UnrealTargetPlatform InHostPlatform, bool bBuildingRocket, out bool bInternalToolOnly, out bool SeparateNode, out bool CrossCompile)
+    public override bool GUBP_AlwaysBuildWithTools(UnrealTargetPlatform InHostPlatform, out bool bInternalToolOnly, out bool SeparateNode, out bool CrossCompile)
 	{
 		bInternalToolOnly = false;
 		SeparateNode = false;
diff --git a/Engine/Source/Programs/Mac/UnrealAtoS/UnrealAtoS.Target.cs b/Engine/Source/Programs/Mac/UnrealAtoS/UnrealAtoS.Target.cs
--- a/Engine/Source/Programs/Mac/UnrealAtoS/UnrealAtoS.Target.cs
+++ b/Engine/Source/Programs/Mac/UnrealAtoS/UnrealAtoS.Target.cs
@@ -54,7 +54,7 @@ public class UnrealAtoSTarget : TargetRules
 		// Do NOT produce additional console app exe
 		OutLinkEnvironmentConfiguration.bIsBuildingConsoleApplication = true;
 	}
-    public override bool GUBP_AlwaysBuildWithTools(UnrealTargetPlatform InHostPlatform, bool bBuildingRocket, out bool bInternalToolOnly, out bool SeparateNode, out bool CrossCompile)
+    public override bool GUBP_AlwaysBuildWithTools(UnrealTargetPlatform InHostPlatform, out bool bInternalToolOnly, out bool SeparateNode, out bool CrossCompile)
     {
         bInternalToolOnly = false;
         SeparateNode = false;
diff --git a/Engine/Source/Programs/UnrealBuildTool/System/RulesCompiler.cs b/Engine/Source/Programs/UnrealBuildTool/System/RulesCompiler.cs
--- a/Engine/Source/Programs/UnrealBuildTool/System/RulesCompiler.cs
+++ b/Engine/Source/Programs/UnrealBuildTool/System/RulesCompiler.cs
@@ -641,7 +641,7 @@ namespace UnrealBuildTool
         /// <param name="SeparateNode">If this is set to true, the program will get its own node</param>
         /// </summary>
         /// <returns>true if this target should always be built with the base editor.</returns>        
-        public virtual bool GUBP_AlwaysBuildWithTools(UnrealTargetPlatform InHostPlatform, bool bBuildingRocket, out bool bInternalToolOnly, out bool SeparateNode, out bool CrossCompile)
+        public virtual bool GUBP_AlwaysBuildWithTools(UnrealTargetPlatform InHostPlatform, out bool bInternalToolOnly, out bool SeparateNode, out bool CrossCompile)
         {
             bInternalToolOnly = false;
             SeparateNode = false;
diff --git a/Engine/Source/Programs/UnrealFileServer/UnrealFileServer.Target.cs b/Engine/Source/Programs/UnrealFileServer/UnrealFileServer.Target.cs
--- a/Engine/Source/Programs/UnrealFileServer/UnrealFileServer.Target.cs
+++ b/Engine/Source/Programs/UnrealFileServer/UnrealFileServer.Target.cs
@@ -57,7 +57,7 @@ public class UnrealFileServerTarget : TargetRules
 		// UnrealHeaderTool is a console application, not a Windows app (sets entry point to main(), instead of WinMain())
 		OutLinkEnvironmentConfiguration.bIsBuildingConsoleApplication = true;
 	}
-    public override bool GUBP_AlwaysBuildWithTools(UnrealTargetPlatform InHostPlatform, bool bBuildingRocket, out bool bInternalToolOnly, out bool SeparateNode, out bool CrossCompile)
+    public override bool GUBP_AlwaysBuildWithTools(UnrealTargetPlatform InHostPlatform, out bool bInternalToolOnly, out bool SeparateNode, out bool CrossCompile)
     {
         bInternalToolOnly = false;
         SeparateNode = false;
diff --git a/Engine/Source/Programs/UnrealLightmass/UnrealLightmass.Target.cs b/Engine/Source/Programs/UnrealLightmass/UnrealLightmass.Target.cs
--- a/Engine/Source/Programs/UnrealLightmass/UnrealLightmass.Target.cs
+++ b/Engine/Source/Programs/UnrealLightmass/UnrealLightmass.Target.cs
@@ -75,7 +75,7 @@ public class UnrealLightmassTarget : TargetRules
 		// Disable logging, lightmass will create its own unique logging file
 		OutCPPEnvironmentConfiguration.Definitions.Add("ALLOW_LOG_FILE=0");
 	}
-    public override bool GUBP_AlwaysBuildWithTools(UnrealTargetPlatform InHostPlatform, bool bBuildingRocket, out bool bInternalToolOnly, out bool SeparateNode, out bool CrossCompile)
+    public override bool GUBP_AlwaysBuildWithTools(UnrealTargetPlatform InHostPlatform, out bool bInternalToolOnly, out bool SeparateNode, out bool CrossCompile)
     {
         bInternalToolOnly = false;
         SeparateNode = false;
diff --git a/Engine/Source/Programs/UnrealPak/UnrealPak.Target.cs b/Engine/Source/Programs/UnrealPak/UnrealPak.Target.cs
--- a/Engine/Source/Programs/UnrealPak/UnrealPak.Target.cs
+++ b/Engine/Source/Programs/UnrealPak/UnrealPak.Target.cs
@@ -51,7 +51,7 @@ public class UnrealPakTarget : TargetRules
 		// UnrealHeaderTool is a console application, not a Windows app (sets entry point to main(), instead of WinMain())
 		OutLinkEnvironmentConfiguration.bIsBuildingConsoleApplication = true;
 	}
-    public override bool GUBP_AlwaysBuildWithTools(UnrealTargetPlatform InHostPlatform, bool bBuildingRocket, out bool bInternalToolOnly, out bool SeparateNode, out bool CrossCompile)
+    public override bool GUBP_AlwaysBuildWithTools(UnrealTargetPlatform InHostPlatform, out bool bInternalToolOnly, out bool SeparateNode, out bool CrossCompile)
     {
         bInternalToolOnly = false;
         SeparateNode = false;
diff --git a/Engine/Source/Programs/UnrealSync/UnrealSync.Target.cs b/Engine/Source/Programs/UnrealSync/UnrealSync.Target.cs
--- a/Engine/Source/Programs/UnrealSync/UnrealSync.Target.cs
+++ b/Engine/Source/Programs/UnrealSync/UnrealSync.Target.cs
@@ -56,7 +56,7 @@ public class UnrealSyncTarget : TargetRules
 		OutLinkEnvironmentConfiguration.bIsBuildingConsoleApplication = false;
 	}
 
-    public override bool GUBP_AlwaysBuildWithTools(UnrealTargetPlatform InHostPlatform, bool bBuildingRocket, out bool bInternalToolOnly, out bool SeparateNode, out bool CrossCompile)
+    public override bool GUBP_AlwaysBuildWithTools(UnrealTargetPlatform InHostPlatform, out bool bInternalToolOnly, out bool SeparateNode, out bool CrossCompile)
 	{
 		CrossCompile = false;
 		if (InHostPlatform == UnrealTargetPlatform.Win32 || InHostPlatform == UnrealTargetPlatform.Win64)
diff --git a/Engine/Source/Programs/UnrealVersionSelector/UnrealVersionSelector.Target.cs b/Engine/Source/Programs/UnrealVersionSelector/UnrealVersionSelector.Target.cs
--- a/Engine/Source/Programs/UnrealVersionSelector/UnrealVersionSelector.Target.cs
+++ b/Engine/Source/Programs/UnrealVersionSelector/UnrealVersionSelector.Target.cs
@@ -52,7 +52,7 @@ public class UnrealVersionSelectorTarget : TargetRules
 		// UnrealHeaderTool is a console application, not a Windows app (sets entry point to main(), instead of WinMain())
 		//OutLinkEnvironmentConfiguration.bIsBuildingConsoleApplication = true;
 	}
-    public override bool GUBP_AlwaysBuildWithTools(UnrealTargetPlatform InHostPlatform, bool bBuildingRocket, out bool bInternalToolOnly, out bool SeparateNode, out bool CrossCompile)
+    public override bool GUBP_AlwaysBuildWithTools(UnrealTargetPlatform InHostPlatform, out bool bInternalToolOnly, out bool SeparateNode, out bool CrossCompile)
 	{
 		CrossCompile = false;
 		if (InHostPlatform == UnrealTargetPlatform.Win32 || InHostPlatform == UnrealTargetPlatform.Win64)
diff --git a/Engine/Source/Programs/Windows/BootstrapPackagedGame/BootstrapPackagedGame.Target.cs b/Engine/Source/Programs/Windows/BootstrapPackagedGame/BootstrapPackagedGame.Target.cs
--- a/Engine/Source/Programs/Windows/BootstrapPackagedGame/BootstrapPackagedGame.Target.cs
+++ b/Engine/Source/Programs/Windows/BootstrapPackagedGame/BootstrapPackagedGame.Target.cs
@@ -56,7 +56,7 @@ public class BootstrapPackagedGameTarget : TargetRules
 		}
 	}
 	
-    public override bool GUBP_AlwaysBuildWithTools(UnrealTargetPlatform InHostPlatform, bool bBuildingRocket, out bool bInternalToolOnly, out bool SeparateNode, out bool CrossCompile)
+    public override bool GUBP_AlwaysBuildWithTools(UnrealTargetPlatform InHostPlatform, out bool bInternalToolOnly, out bool SeparateNode, out bool CrossCompile)
 	{
 		bInternalToolOnly = false;
 		SeparateNode = false;
