Fix incorrect parsing of changelist descriptions if the summary ends in a newline.

diff --git a/Engine/Source/Programs/AutomationTool/P4Utils.cs b/Engine/Source/Programs/AutomationTool/P4Utils.cs
--- a/Engine/Source/Programs/AutomationTool/P4Utils.cs
+++ b/Engine/Source/Programs/AutomationTool/P4Utils.cs
@@ -928,6 +928,12 @@ namespace AutomationTool
 							throw new AutomationException("Was expecting 'Affected files' to appear after the summary output from P4, but there were no more lines to read");
 						}
 
+						// If the summary ends with an empty newline, it doesn't seem to be prefixed with a tab
+						while (LineIndex < Lines.Length && Lines[LineIndex].Length == 0)
+						{
+							LineIndex++;
+						}
+
 						Line = Lines[ LineIndex ];
 
 						string MatchAffectedFiles = "Affected files";
