build script, fix hastests

diff --git a/Engine/Source/Programs/AutomationTool/Scripts/GUBP.Automation.cs b/Engine/Source/Programs/AutomationTool/Scripts/GUBP.Automation.cs
--- a/Engine/Source/Programs/AutomationTool/Scripts/GUBP.Automation.cs
+++ b/Engine/Source/Programs/AutomationTool/Scripts/GUBP.Automation.cs
@@ -7342,7 +7342,7 @@ public class GUBP : BuildCommand
                 }
             }
             WriteECPerl(StepList);
-            ECProps.Add(String.Format("HasTests={0}", bHasTests));
+            RunECTool(String.Format("setProperty \"/myWorkflow/HasTests\" \"{0}\"", bHasTests));
             {
                 ECProps.Add("GUBP_LoadedProps=1");
                 string BranchDefFile = CommandUtils.CombinePaths(CommandUtils.CmdEnv.LogFolder, "BranchDef.properties");
