Add support for copying build products to Amazon S3 when staging a build.

Amazon S3 will be our new origin for CloudFront.

diff --git a/Engine/Source/Programs/AutomationTool/MCPPublic.cs b/Engine/Source/Programs/AutomationTool/MCPPublic.cs
--- a/Engine/Source/Programs/AutomationTool/MCPPublic.cs
+++ b/Engine/Source/Programs/AutomationTool/MCPPublic.cs
@@ -552,6 +552,118 @@ namespace EpicGames.MCP.Automation
         public abstract string SendWebRequest(WebRequest Upload, string Method, string ContentType, byte[] Data);
     }
 
+	/// <summary>
+	/// Helper class to manage files stored in some arbitrary cloud storage system
+	/// </summary>
+	public abstract class CloudStorageBase
+	{
+		private static readonly object LockObj = new object();
+		private static CloudStorageBase Handler = null;
+
+		public static CloudStorageBase Get()
+		{
+			if (Handler == null)
+			{
+				lock (LockObj)
+				{
+					if (Handler == null)
+					{
+						Assembly[] LoadedAssemblies = AppDomain.CurrentDomain.GetAssemblies();
+						foreach (var Dll in LoadedAssemblies)
+						{
+							Type[] AllTypes = Dll.GetTypes();
+							foreach (var PotentialConfigType in AllTypes)
+							{
+								if (PotentialConfigType != typeof(CloudStorageBase) && typeof(CloudStorageBase).IsAssignableFrom(PotentialConfigType))
+								{
+									Handler = Activator.CreateInstance(PotentialConfigType) as CloudStorageBase;
+									break;
+								}
+							}
+						}
+					}
+				}
+				if (Handler == null)
+				{
+					throw new AutomationException("Attempt to use CloudStorageBase.Get() and it doesn't appear that there are any modules that implement this class.");
+				}
+			}
+			return Handler;
+		}
+
+		/// <summary>
+		/// Initializes the provider.
+		/// <param name="Config">Configuration data to initialize the provider. The exact format of the data is provider specific. It might, for example, contain an API key.</param>
+		/// </summary>
+		abstract public void Init(Dictionary<string,object> Config);
+
+		/// <summary>
+		/// Retrieves a file from the cloud storage provider
+		/// </summary>
+		/// <param name="Container">The name of the folder or container from which contains the file being checked.</param>
+		/// <param name="Identifier">The identifier or filename of the file to check.</param>
+		/// <returns>True if the file exists in cloud storage, false otherwise.</returns>
+		abstract public bool FileExists(string Container, string Identifier);
+
+		/// <summary>
+		/// Retrieves a file from the cloud storage provider
+		/// </summary>
+		/// <param name="Container">The name of the folder or container from which to retrieve the file.</param>
+		/// <param name="Identifier">The identifier or filename of the file to retrieve.</param>
+		/// <param name="ContentType">An OUTPUT parameter containing the content's type (null if the cloud provider does not provide this information)</param>
+		/// <returns>A byte array containing the file's contents.</returns>
+		abstract public byte[] GetFile(string Container, string Identifier, out string ContentType);
+
+		/// <summary>
+		/// Posts a file to the cloud storage provider
+		/// NOTE: the method returns void, rather than bSuccess as you might imagine, because of an apparent bug in VS2013 debugger, where an AccessViolationExeption is
+		/// thrown when stepping in to an overridden method, with a return value, and at least one out parameter, when optimizations are enabled.
+		/// </summary>
+		/// <param name="Container">The name of the folder or container in which to store the file.</param>
+		/// <param name="Identifier">The identifier or filename of the file to write.</param>
+		/// <param name="Contents">A byte array containing the data to write.</param>
+		/// <param name="ObjectURL">An OUTPUT parameter which will be set to the URL of the uploaded file on success.</param>
+		/// <param name="bSuccess">An OUTPUT parameter which will be set to true if the write succeeds, false otherwise.</param>
+		/// <param name="bOverwrite">If true, will overwrite an existing file.  If false, will throw an exception if the file exists.</param>
+		/// <param name="bMakePublic">Specified whether the file should be made public readable.</param>
+		abstract public void PostFile(string Container, string Identifier, byte[] Contents, out string ObjectURL, out bool bSuccess, bool bOverwrite = true, bool bMakePublic = false);
+
+		/// <summary>
+		/// Posts a file to the cloud storage provider.
+		/// NOTE: the method returns void, rather than bSuccess as you might imagine, because of an apparent bug in VS2013 debugger, where an AccessViolationExeption is
+		/// thrown when stepping in to an overridden method, with a return value, and at least one out parameter, when optimizations are enabled.
+		/// </summary>
+		/// <param name="Container">The name of the folder or container in which to store the file.</param>
+		/// <param name="Identifier">The identifier or filename of the file to write.</param>
+		/// <param name="SourceFilePath">The full path of the file to upload.</param>
+		/// <param name="ObjectURL">An OUTPUT parameter which will be set to the URL of the uploaded file on success.</param>
+		/// <param name="bSuccess">An OUTPUT parameter which will be set to true if the write succeeds, false otherwise.</param>
+		/// <param name="bOverwrite">If true, will overwrite an existing file.  If false, will throw an exception if the file exists.</param>
+		/// <param name="bMakePublic">Specified whether the file should be made public readable.</param>
+		abstract public void PostFile(string Container, string Identifier, string SourceFilePath, out string ObjectURL, out bool bSuccess, bool bOverwrite = true, bool bMakePublic = false);
+
+		/// <summary>
+		/// Deletes a file from cloud storage
+		/// </summary>
+		/// <param name="Container">The name of the folder or container in which to store the file.</param>
+		/// <param name="Identifier">The identifier or filename of the file to write.</param>
+		abstract public void DeleteFile(string Container, string Identifier);
+
+		/// <summary>
+		/// Retrieves a list of files from the cloud storage provider
+		/// </summary>
+		/// <param name="Container">The name of the folder or container from which to list files.</param>
+		/// <param name="Prefix">A string with which the identifier or filename should start. Typically used to specify a relative directory within the container to list all of its files recursively. Specify null to return all files.</param>
+		/// <returns>An array of paths to the files in the specified location and matching the prefix constraint.</returns>
+		abstract public string[] ListFiles(string Container, string Prefix = null);
+
+		/// <summary>
+		/// Copies chunks from a staged location to cloud storage.
+		/// </summary>
+		/// <param name="Container">The name of the folder or container in which to store files.</param>
+		/// <param name="stagingInfo">Staging info used to determine where the chunks are to copy.</param>
+		abstract public void CopyChunksToCloudStorage(string Container, BuildPatchToolStagingInfo StagingInfo);
+	}
 }
 
 namespace EpicGames.MCP.Config
