build script, attempt to make mac temp storage slightly less fatal

diff --git a/Engine/Build/Target.cs.template b/Engine/Build/Target.cs.template
--- a/Engine/Build/Target.cs.template
+++ b/Engine/Build/Target.cs.template
@@ -143,7 +143,7 @@ public class {GAME_NAME}Target : TargetRules
             { 
                 new GUBPFormalBuild(UnrealTargetPlatform.IOS, UnrealTargetConfiguration.Development),
 				new GUBPFormalBuild(UnrealTargetPlatform.IOS, UnrealTargetConfiguration.Shipping),
-				new GUBPFormalBuild(UnrealTargetPlatform.IOS, UnrealTargetConfiguration.Test),
+				new GUBPFormalBuild(UnrealTargetPlatform.IOS, UnrealTargetConfiguration.Test, false, true),
             };
 		}
 	}
diff --git a/Engine/Source/Programs/AutomationTool/TempStorage.cs b/Engine/Source/Programs/AutomationTool/TempStorage.cs
--- a/Engine/Source/Programs/AutomationTool/TempStorage.cs
+++ b/Engine/Source/Programs/AutomationTool/TempStorage.cs
@@ -103,12 +103,13 @@ namespace AutomationTool
                 // mac is terrible on shares, this isn't a solution, but a stop gap
                 if (Filename.StartsWith("/Volumes/"))
                 {
-                    int Retry = 60;
-                    while (!bFound && --Retry > 0)
+                    int Retry = 0;
+                    while (!bFound && Retry < 60)
                     {
                         CommandUtils.Log(System.Diagnostics.TraceEventType.Warning, "*** Mac temp storage retry {0}", Filename);
                         System.Threading.Thread.Sleep(10000);
                         bFound = FileExists_NoExceptions(bQuiet, Filename);
+						Retry++;
                     }
                 }
                 if (!bFound)
@@ -118,6 +119,27 @@ namespace AutomationTool
             }
         }
 
+		public static bool Robust_DirectoryExists_NoExceptions(string Directoryname, string Message)
+		{
+			bool bFound = false;
+			if (!DirectoryExists_NoExceptions(Directoryname))
+			{				
+				// mac is terrible on shares, this isn't a solution, but a stop gap
+				if (Directoryname.StartsWith("/Volumes/"))
+				{
+					int Retry = 0;
+					while (!bFound && Retry < 60)
+					{
+						CommandUtils.Log(System.Diagnostics.TraceEventType.Warning, "*** Mac temp storage retry {0}", Directoryname);
+						System.Threading.Thread.Sleep(10000);
+						bFound = DirectoryExists_NoExceptions(Directoryname);
+						Retry++;
+					}
+				}				
+			}
+			return bFound;
+		}
+
         public class TempStorageManifest
         {
             private static readonly string RootElementName = "tempstorage";
@@ -561,8 +583,8 @@ namespace AutomationTool
             var Dir = UE4TempStorageDirectory();
             if (ForSaving)
             {
-                int Retries = 24;
-                while (--Retries > 0)
+                int Retries = 0;
+                while (Retries < 24)
                 {
                     if (DirectoryExistsAndIsWritable_NoExceptions(Dir))
                     {
@@ -570,9 +592,10 @@ namespace AutomationTool
                     }
                     FindDirectories_NoExceptions(false, "*", false, Dir); // there is some internet evidence to suggest this might perk up the mac share
                     System.Threading.Thread.Sleep(5000);
+					Retries++;
                 }
             }
-            else if (DirectoryExists_NoExceptions(Dir))
+            else if (Robust_DirectoryExists_NoExceptions(Dir, "Could not find {0}"))
             {
                 return true;
             }
@@ -583,11 +606,11 @@ namespace AutomationTool
         {
             string SharedSubdir = SharedTempStorageDirectory();
             string Result = CombinePaths(ResolveSharedBuildDirectory(GameFolder), SharedSubdir);
-            if (!DirectoryExists_NoExceptions(Result))
+            if (!Robust_DirectoryExists_NoExceptions(Result, "Could not find {0}"))
             {
                 CreateDirectory_NoExceptions(Result);
             }
-            if (!DirectoryExists_NoExceptions(Result))
+            if (!Robust_DirectoryExists_NoExceptions(Result, "Could not find {0}"))
             {
                 throw new AutomationException("Could not create an appropriate shared temp folder {0}", Result);
             }
