build script, fix cis

diff --git a/Engine/Source/Programs/AutomationTool/TempStorage.cs b/Engine/Source/Programs/AutomationTool/TempStorage.cs
--- a/Engine/Source/Programs/AutomationTool/TempStorage.cs
+++ b/Engine/Source/Programs/AutomationTool/TempStorage.cs
@@ -137,6 +137,10 @@ namespace AutomationTool
 					}
 				}				
 			}
+			else
+			{
+				bFound = true;
+			}
 			return bFound;
 		}
 
