UAT: Add support for deleting an entire folder from a cloud directory.

diff --git a/Engine/Source/Programs/AutomationTool/MCPPublic.cs b/Engine/Source/Programs/AutomationTool/MCPPublic.cs
--- a/Engine/Source/Programs/AutomationTool/MCPPublic.cs
+++ b/Engine/Source/Programs/AutomationTool/MCPPublic.cs
@@ -649,6 +649,13 @@ namespace EpicGames.MCP.Automation
 		/// <param name="Identifier">The identifier or filename of the file to write.</param>
 		abstract public void DeleteFile(string Container, string Identifier);
 
+		/// <summary>
+		/// Deletes a folder from cloud storage
+		/// </summary>
+		/// <param name="Container">The name of the folder or container from which to delete the file.</param>
+		/// <param name="FolderIdentifier">The identifier or name of the folder to delete.</param>
+		abstract public void DeleteFolder(string Container, string FolderIdentifier);
+
 		/// <summary>
 		/// Retrieves a list of files from the cloud storage provider
 		/// </summary>
