UAT: Add support for setting content type when uploading to cloud storage.

diff --git a/Engine/Source/Programs/AutomationTool/MCPPublic.cs b/Engine/Source/Programs/AutomationTool/MCPPublic.cs
--- a/Engine/Source/Programs/AutomationTool/MCPPublic.cs
+++ b/Engine/Source/Programs/AutomationTool/MCPPublic.cs
@@ -624,9 +624,10 @@ namespace EpicGames.MCP.Automation
 		/// <param name="Contents">A byte array containing the data to write.</param>
 		/// <param name="ObjectURL">An OUTPUT parameter which will be set to the URL of the uploaded file on success.</param>
 		/// <param name="bSuccess">An OUTPUT parameter which will be set to true if the write succeeds, false otherwise.</param>
+		/// <param name="ContentType">The MIME type of the file being uploaded.</param>
 		/// <param name="bOverwrite">If true, will overwrite an existing file.  If false, will throw an exception if the file exists.</param>
 		/// <param name="bMakePublic">Specified whether the file should be made public readable.</param>
-		abstract public void PostFile(string Container, string Identifier, byte[] Contents, out string ObjectURL, out bool bSuccess, bool bOverwrite = true, bool bMakePublic = false);
+		abstract public void PostFile(string Container, string Identifier, byte[] Contents, out string ObjectURL, out bool bSuccess, string ContentType = null, bool bOverwrite = true, bool bMakePublic = false);
 
 		/// <summary>
 		/// Posts a file to the cloud storage provider.
@@ -638,9 +639,10 @@ namespace EpicGames.MCP.Automation
 		/// <param name="SourceFilePath">The full path of the file to upload.</param>
 		/// <param name="ObjectURL">An OUTPUT parameter which will be set to the URL of the uploaded file on success.</param>
 		/// <param name="bSuccess">An OUTPUT parameter which will be set to true if the write succeeds, false otherwise.</param>
+		/// <param name="ContentType">The MIME type of the file being uploaded.</param>
 		/// <param name="bOverwrite">If true, will overwrite an existing file.  If false, will throw an exception if the file exists.</param>
 		/// <param name="bMakePublic">Specified whether the file should be made public readable.</param>
-		abstract public void PostFile(string Container, string Identifier, string SourceFilePath, out string ObjectURL, out bool bSuccess, bool bOverwrite = true, bool bMakePublic = false);
+		abstract public void PostFile(string Container, string Identifier, string SourceFilePath, out string ObjectURL, out bool bSuccess, string ContentType = null, bool bOverwrite = true, bool bMakePublic = false);
 
 		/// <summary>
 		/// Deletes a file from cloud storage
