Stage Mac Steam lib when packaging a game that uses it

diff --git a/Engine/Source/Programs/AutomationTool/Mac/MacPlatform.Automation.cs b/Engine/Source/Programs/AutomationTool/Mac/MacPlatform.Automation.cs
--- a/Engine/Source/Programs/AutomationTool/Mac/MacPlatform.Automation.cs
+++ b/Engine/Source/Programs/AutomationTool/Mac/MacPlatform.Automation.cs
@@ -74,6 +74,11 @@ public class MacPlatform : Platform
 			if (!string.IsNullOrEmpty(AppBundlePath))
 			{
 				SC.StageFiles(StagedFileType.NonUFS, CombinePaths(SC.ProjectRoot, "Build/Mac"), "Application.icns", false, null, CombinePaths(AppBundlePath, "Contents/Resources"), true);
+
+				if (Params.bUsesSteam)
+				{
+					SC.StageFiles(StagedFileType.NonUFS, CombinePaths(SC.LocalRoot, "Engine/Source/ThirdParty/Steamworks/Steamv132/sdk/redistributable_bin/osx32"), "libsteam_api.dylib", false, null, CombinePaths(AppBundlePath, "Contents/MacOS"), true);
+				}
 			}
 		}
 
