UAT: Ensure we have a staged build of launcher in S3 before carrying out live promotions, or prod custom labels.

diff --git a/Engine/Source/Programs/AutomationTool/MCPPublic.cs b/Engine/Source/Programs/AutomationTool/MCPPublic.cs
--- a/Engine/Source/Programs/AutomationTool/MCPPublic.cs
+++ b/Engine/Source/Programs/AutomationTool/MCPPublic.cs
@@ -672,6 +672,14 @@ namespace EpicGames.MCP.Automation
 		/// <param name="Container">The name of the folder or container in which to store files.</param>
 		/// <param name="stagingInfo">Staging info used to determine where the chunks are to copy.</param>
 		abstract public void CopyChunksToCloudStorage(string Container, BuildPatchToolStagingInfo StagingInfo);
+
+		/// <summary>
+		/// Verifies whether a manifest for a given build is in cloud storage.
+		/// </summary>
+		/// <param name="Container">The name of the folder or container in which to store files.</param>
+		/// <param name="stagingInfo">Staging info representing the build to check.</param>
+		/// <returns>True if the manifest exists in cloud storage, false otherwise.</returns>
+		abstract public bool IsManifestOnCloudStorage(string Container, BuildPatchToolStagingInfo StagingInfo);
 	}
 }
 
