UAT: Ensure build patch tool is built in CopyBuildToStaging following removal of BuildPatchTool.exe from depot.

diff --git a/Engine/Source/Programs/AutomationTool/Scripts/CopyBuildToStagingDirectory.Automation.cs b/Engine/Source/Programs/AutomationTool/Scripts/CopyBuildToStagingDirectory.Automation.cs
--- a/Engine/Source/Programs/AutomationTool/Scripts/CopyBuildToStagingDirectory.Automation.cs
+++ b/Engine/Source/Programs/AutomationTool/Scripts/CopyBuildToStagingDirectory.Automation.cs
@@ -22,6 +22,43 @@ public partial class Project : CommandUtils
 
 	#region Utilities
 
+	private static readonly object SyncLock = new object();
+
+	/// <returns>The path for the BuildPatchTool executable depending on host platform.</returns>
+	private static string GetBuildPatchToolExecutable()
+	{
+		switch (UnrealBuildTool.BuildHostPlatform.Current.Platform)
+		{
+			case UnrealTargetPlatform.Win32:
+				return CommandUtils.CombinePaths(CommandUtils.CmdEnv.LocalRoot, "Engine/Binaries/Win32/BuildPatchTool.exe");
+			case UnrealTargetPlatform.Win64:
+				return CommandUtils.CombinePaths(CommandUtils.CmdEnv.LocalRoot, "Engine/Binaries/Win64/BuildPatchTool.exe");
+			case UnrealTargetPlatform.Mac:
+				return CommandUtils.CombinePaths(CommandUtils.CmdEnv.LocalRoot, "Engine/Binaries/Mac/BuildPatchTool");
+			case UnrealTargetPlatform.Linux:
+				return CommandUtils.CombinePaths(CommandUtils.CmdEnv.LocalRoot, "Engine/Binaries/Linux/BuildPatchTool");
+		}
+		throw new AutomationException(string.Format("Unknown host platform for BuildPatchTool - {0}", UnrealBuildTool.BuildHostPlatform.Current.Platform));
+	}
+
+	/// <summary>
+	/// Checks the existence of the BuildPatchTool executable exists and builds it if it is missing
+	/// </summary>
+	private static void EnsureBuildPatchToolExists()
+	{
+		string BuildPatchToolExe = GetBuildPatchToolExecutable();
+		if (!CommandUtils.FileExists_NoExceptions(BuildPatchToolExe))
+		{
+			lock (SyncLock)
+			{
+				if (!CommandUtils.FileExists_NoExceptions(BuildPatchToolExe))
+				{
+					UE4BuildUtils.BuildBuildPatchTool(null, UnrealBuildTool.BuildHostPlatform.Current.Platform);
+				}
+			}
+		}
+	}
+
 	/// <summary>
 	/// Writes a pak response file to disk
 	/// </summary>
@@ -658,14 +695,8 @@ public partial class Project : CommandUtils
 			int ChunkID = Convert.ToInt32(Matches[0].Groups[1].ToString());
 			if (ChunkID != 0)
 			{
-				var BPTExe = CombinePaths(CmdEnv.LocalRoot, "Engine/Binaries/Win64/BuildPatchTool.exe");
-
-				// Ensure that BuildPatchTool.exe exists
-				//if (!InternalUtils.SafeFileExists(BPTExe))
-				//{
-				//	CommandUtils.Log("Building BuildPatchTool.exe");
-				//	UE4BuildUtils.BuildBuildPatchTool(null, UnrealTargetPlatform.Win64);
-				//}
+				var BPTExe = GetBuildPatchToolExecutable();
+				EnsureBuildPatchToolExists();
 
 				string P4Change = "UnknownCL";
 				string P4Branch = "UnknownBranch";
