build script, fix frequency for tappy

diff --git a/Engine/Build/Target.cs.template b/Engine/Build/Target.cs.template
--- a/Engine/Build/Target.cs.template
+++ b/Engine/Build/Target.cs.template
@@ -96,7 +96,7 @@ public class {GAME_NAME}Target : TargetRules
 	public override GUBPProjectOptions GUBP_IncludeProjectInPromotedBuild_EditorTypeOnly(UnrealTargetPlatform HostPlatform)
     {
         var Result = new GUBPProjectOptions();
-        Result.bTestWithShared = true;
+        Result.bTestWithShared = false;
 		Result.bIsNonCode = true;
         return Result;
     }
diff --git a/Engine/Source/Programs/AutomationTool/Scripts/GUBP.Automation.cs b/Engine/Source/Programs/AutomationTool/Scripts/GUBP.Automation.cs
--- a/Engine/Source/Programs/AutomationTool/Scripts/GUBP.Automation.cs
+++ b/Engine/Source/Programs/AutomationTool/Scripts/GUBP.Automation.cs
@@ -6032,9 +6032,9 @@ public class GUBP : BuildCommand
             {
                 var Options = CodeProj.Options(HostPlatform);
 
-                if (!Options.bIsPromotable && !Options.bTestWithShared)
+                if (!Options.bIsPromotable && !Options.bTestWithShared && !Options.bIsNonCode)
                 {
-                    continue; // we skip things that aren't promotable and aren't tested
+                    continue; // we skip things that aren't promotable and aren't tested - except noncode as code situations
                 }
                 var AgentShareName = CodeProj.GameName;
                 if (!Options.bSeparateGamePromotion)
