build script, remove wait for test shared dependency in branches that don't have a shared promotable

diff --git a/Engine/Source/Programs/AutomationTool/Scripts/GUBP.Automation.cs b/Engine/Source/Programs/AutomationTool/Scripts/GUBP.Automation.cs
--- a/Engine/Source/Programs/AutomationTool/Scripts/GUBP.Automation.cs
+++ b/Engine/Source/Programs/AutomationTool/Scripts/GUBP.Automation.cs
@@ -6061,7 +6061,7 @@ public class GUBP : BuildCommand
 						foreach (var Test in EditorTests)
 						{
 							EditorTestNodes.Add(AddNode(new UATTestNode(this, HostPlatform, CodeProj, Test.Key, Test.Value, AgentSharingGroup)));
-							if (!Options.bTestWithShared)
+							if (!Options.bTestWithShared || !HasNode(WaitForTestShared.StaticGetFullName()))
 							{
 								RemovePseudodependencyFromNode((UATTestNode.StaticGetFullName(HostPlatform, CodeProj, Test.Key)), WaitForTestShared.StaticGetFullName());
 							}
@@ -6210,7 +6210,7 @@ public class GUBP : BuildCommand
 										{
 											var TestNodeName = Test.Key + "_" + Plat.ToString();
 											ThisMonoGameTestNodes.Add(AddNode(new UATTestNode(this, HostPlatform, CodeProj, TestNodeName, Test.Value, CookedAgentSharingGroup, false, RequiredPlatforms)));
-											if (!Options.bTestWithShared)
+											if (!Options.bTestWithShared || !HasNode(WaitForTestShared.StaticGetFullName()))
 											{
 												RemovePseudodependencyFromNode((UATTestNode.StaticGetFullName(HostPlatform, CodeProj, TestNodeName)), WaitForTestShared.StaticGetFullName());
 											}
