Make CrashReportClient optional when packaging for Linux.

- Requirement to build CrashReportClient when packaging projects on/for Linux may be non-obvious, especially for users building from sources (UE-12063).

diff --git a/Engine/Source/Programs/AutomationTool/Linux/LinuxPlatform.Automation.cs b/Engine/Source/Programs/AutomationTool/Linux/LinuxPlatform.Automation.cs
--- a/Engine/Source/Programs/AutomationTool/Linux/LinuxPlatform.Automation.cs
+++ b/Engine/Source/Programs/AutomationTool/Linux/LinuxPlatform.Automation.cs
@@ -27,7 +27,7 @@ public abstract class BaseLinuxPlatform : Platform
 	{
 		if (SC.bStageCrashReporter)
 		{
-			SC.StageFiles(StagedFileType.NonUFS, CombinePaths(SC.LocalRoot, "Engine/Binaries", SC.PlatformDir), "CrashReportClient", false);
+			SC.StageFiles(StagedFileType.NonUFS, CombinePaths(SC.LocalRoot, "Engine/Binaries", SC.PlatformDir), "CrashReportClient", false, null, null, true);
 		}
 
 		{
