Check for Android .so instead of .apk to determine if UE4 project for target configuration was built
#ue4
#android

diff --git a/Engine/Source/Programs/AutomationTool/Android/AndroidPlatform.Automation.cs b/Engine/Source/Programs/AutomationTool/Android/AndroidPlatform.Automation.cs
--- a/Engine/Source/Programs/AutomationTool/Android/AndroidPlatform.Automation.cs
+++ b/Engine/Source/Programs/AutomationTool/Android/AndroidPlatform.Automation.cs
@@ -164,12 +164,13 @@ public class AndroidPlatform : Platform
 				string ApkName = GetFinalApkName(Params, SC.StageExecutables[0], true, bMakeSeparateApks ? Architecture : "", bMakeSeparateApks ? GPUArchitecture : "");
 				if (!SC.IsCodeBasedProject)
 				{
-					string UE4GameApkName = GetFinalApkName(Params, SC.StageExecutables[0], false, bMakeSeparateApks ? Architecture : "", bMakeSeparateApks ? GPUArchitecture : "");
-					if (FileExists_NoExceptions(UE4GameApkName) == false)
+					string UE4SOName = GetFinalApkName(Params, SC.StageExecutables[0], false, bMakeSeparateApks ? Architecture : "", bMakeSeparateApks ? GPUArchitecture : "");
+                    UE4SOName = UE4SOName.Replace(".apk", ".so");
+                    if (FileExists_NoExceptions(UE4SOName) == false)
 					{
-						Log("Failed to find game apk " + UE4GameApkName);
+                        Log("Failed to find game .so " + UE4SOName);
 						AutomationTool.ErrorReporter.Error("Stage Failed.", (int)AutomationTool.ErrorCodes.Error_MissingExecutable);
-						throw new AutomationException("Could not find apk {0}. You may need to build the UE4 project with your target configuration and platform.", UE4GameApkName);
+                        throw new AutomationException("Could not find .so {0}. You may need to build the UE4 project with your target configuration and platform.", UE4SOName);
 					}
 				}
 				
