Rebuild feature packs in a single GUBP node. It's slow and overly verbose to build them separately.

diff --git a/Engine/Source/Programs/AutomationTool/Scripts/GUBP.Automation.cs b/Engine/Source/Programs/AutomationTool/Scripts/GUBP.Automation.cs
--- a/Engine/Source/Programs/AutomationTool/Scripts/GUBP.Automation.cs
+++ b/Engine/Source/Programs/AutomationTool/Scripts/GUBP.Automation.cs
@@ -1396,18 +1396,23 @@ public class GUBP : BuildCommand
         }
     }
 
-	public class MakeFeaturePackNode : HostPlatformNode
+	public class MakeFeaturePacksNode : HostPlatformNode
 	{
-        BranchInfo.BranchUProject GameProj;
+		List<BranchInfo.BranchUProject> Projects;
 
-        public MakeFeaturePackNode(UnrealTargetPlatform InHostPlatform, BranchInfo.BranchUProject InGameProj)
+        public MakeFeaturePacksNode(UnrealTargetPlatform InHostPlatform, IEnumerable<BranchInfo.BranchUProject> InProjects)
             : base(InHostPlatform)
         {
-			GameProj = InGameProj;
+			Projects = new List<BranchInfo.BranchUProject>(InProjects);
 			AddDependency(ToolsNode.StaticGetFullName(InHostPlatform)); // for UnrealPak
             AgentSharingGroup = "FeaturePacks"  + StaticGetHostPlatformSuffix(InHostPlatform);
         }
 
+		public static string GetOutputFile(BranchInfo.BranchUProject Project)
+		{
+			return CommandUtils.CombinePaths(CommandUtils.CmdEnv.LocalRoot, "FeaturePacks", Path.GetFileNameWithoutExtension(Project.GameName) + ".upack");
+		}
+
 		public static bool IsFeaturePack(BranchInfo.BranchUProject InGameProj)
 		{
 			// No obvious way to store this in the project options; it's a property of non-code projects too.
@@ -1438,14 +1443,14 @@ public class GUBP : BuildCommand
 			}
 		}
 
-        public static string StaticGetFullName(UnrealTargetPlatform InHostPlatform, BranchInfo.BranchUProject InGameProj)
+        public static string StaticGetFullName(UnrealTargetPlatform InHostPlatform)
         {
-            return InGameProj.GameName + "_MakeFeaturePack" + StaticGetHostPlatformSuffix(InHostPlatform);
+            return "MakeFeaturePacks" + StaticGetHostPlatformSuffix(InHostPlatform);
         }
 
 		public override string GetFullName()
         {
-            return StaticGetFullName(HostPlatform, GameProj);
+            return StaticGetFullName(HostPlatform);
         }
 
 		public override int CISFrequencyQuantumShift(GUBP bp)
@@ -1454,38 +1459,42 @@ public class GUBP : BuildCommand
 		}
 		public override void DoBuild(GUBP bp)
         {
-			string ContentsFileName = CommandUtils.CombinePaths(CommandUtils.GetDirectoryName(GameProj.FilePath), "contents.txt");
+			BuildProducts = new List<string>();
+			foreach(BranchInfo.BranchUProject Project in Projects)
+			{
+				string ContentsFileName = CommandUtils.CombinePaths(CommandUtils.GetDirectoryName(Project.FilePath), "contents.txt");
 
-			// Make sure we delete the output file. It may be read-only.
-			string OutputFileName = CommandUtils.CombinePaths(CommandUtils.CmdEnv.LocalRoot, "FeaturePacks", Path.GetFileNameWithoutExtension(GameProj.GameName) + ".upack");
-			CommandUtils.DeleteFile(OutputFileName);
+				// Make sure we delete the output file. It may be read-only.
+				string OutputFileName = GetOutputFile(Project);
+				CommandUtils.DeleteFile(OutputFileName);
 
-			// Get the command line
-			string CmdLine = CommandUtils.MakePathSafeToUseWithCommandLine(OutputFileName) + " " + CommandUtils.MakePathSafeToUseWithCommandLine("-create=" + ContentsFileName);
-			if (GlobalCommandLine.Installed)
-			{
-				CmdLine += " -installed";
-			}
-			if (GlobalCommandLine.UTF8Output)
-			{
-				CmdLine += " -UTF8Output";
-			}
+				// Get the command line
+				string CmdLine = CommandUtils.MakePathSafeToUseWithCommandLine(OutputFileName) + " " + CommandUtils.MakePathSafeToUseWithCommandLine("-create=" + ContentsFileName);
+				if (GlobalCommandLine.Installed)
+				{
+					CmdLine += " -installed";
+				}
+				if (GlobalCommandLine.UTF8Output)
+				{
+					CmdLine += " -UTF8Output";
+				}
 
-			// Run UnrealPak
-			string UnrealPakExe;
-			if(HostPlatform == UnrealTargetPlatform.Win64)
-			{
-				UnrealPakExe = CombinePaths(CmdEnv.LocalRoot, "Engine/Binaries/Win64/UnrealPak.exe");
-			}
-			else
-			{
-				throw new AutomationException("Unknown path to UnrealPak for host platform ({0})", HostPlatform);
-			}
-			RunAndLog(CmdEnv, UnrealPakExe, CmdLine, Options: ERunOptions.Default | ERunOptions.AllowSpew | ERunOptions.UTF8Output);
+				// Run UnrealPak
+				string UnrealPakExe;
+				if(HostPlatform == UnrealTargetPlatform.Win64)
+				{
+					UnrealPakExe = CombinePaths(CmdEnv.LocalRoot, "Engine/Binaries/Win64/UnrealPak.exe");
+				}
+				else
+				{
+					throw new AutomationException("Unknown path to UnrealPak for host platform ({0})", HostPlatform);
+				}
+				RunAndLog(CmdEnv, UnrealPakExe, CmdLine, Options: ERunOptions.Default | ERunOptions.AllowSpew | ERunOptions.UTF8Output);
 
-			// Add the build products
-			BuildProducts = new List<string>();
-			BuildProducts.Add(OutputFileName);
+				// Add the build products
+				BuildProducts.Add(OutputFileName);
+			}
+			SaveRecordOfSuccessAndAddToBuildProducts();
 		}
 	}
 
@@ -2196,15 +2205,7 @@ public class GUBP : BuildCommand
                     }
                 }
             }
-
-			UnrealTargetPlatform FeaturePackPlatform = MakeFeaturePackNode.GetDefaultBuildPlatform(bp);
-			foreach(var Proj in bp.Branch.AllProjects)
-			{
-				if(MakeFeaturePackNode.IsFeaturePack(Proj))
-				{
-					AddDependency(MakeFeaturePackNode.StaticGetFullName(FeaturePackPlatform, Proj));
-				}
-			}
+			AddDependency(MakeFeaturePacksNode.StaticGetFullName(MakeFeaturePacksNode.GetDefaultBuildPlatform(bp)));
         }
 		public override bool IsSeparatePromotable()
 		{
@@ -6017,15 +6018,9 @@ public class GUBP : BuildCommand
 				AddNode(new SharedCookAggregateNode(this, HostPlatforms, NonCodeProjectNames, NonCodeFormalBuilds));
 			}
 			
-			if(HostPlatform == MakeFeaturePackNode.GetDefaultBuildPlatform(this))
+			if(HostPlatform == MakeFeaturePacksNode.GetDefaultBuildPlatform(this))
 			{
-				foreach(BranchInfo.BranchUProject Project in Branch.AllProjects)
-				{
-					if(MakeFeaturePackNode.IsFeaturePack(Project))
-					{
-						AddNode(new MakeFeaturePackNode(HostPlatform, Project));
-					}
-				}
+				AddNode(new MakeFeaturePacksNode(HostPlatform, Branch.AllProjects.Where(x => MakeFeaturePacksNode.IsFeaturePack(x))));
 			}
 
             foreach (var CodeProj in Branch.CodeProjects)
