Use -precompilemodules when building RootEditor - builds additional modules which aren't referenced directly, and includes import libraries in build products (as required for Rocket). Don't submit import libraries to P4.

diff --git a/Engine/Source/Programs/AutomationTool/Scripts/GUBP.Automation.cs b/Engine/Source/Programs/AutomationTool/Scripts/GUBP.Automation.cs
--- a/Engine/Source/Programs/AutomationTool/Scripts/GUBP.Automation.cs
+++ b/Engine/Source/Programs/AutomationTool/Scripts/GUBP.Automation.cs
@@ -825,7 +825,7 @@ public class GUBP : BuildCommand
         {
             var Agenda = new UE4Build.BuildAgenda();
 
-            string AddArgs = "-nobuilduht" + bp.RocketUBTArgs();
+            string AddArgs = "-nobuilduht -precompilemodules";
             if (bp.bOrthogonalizeEditorPlatforms)
             {
                 AddArgs += " -skipnonhostplatforms";
@@ -2513,7 +2513,8 @@ public class GUBP : BuildCommand
                             !Product.EndsWith("version.h", StringComparison.InvariantCultureIgnoreCase) && 
                             !Product.EndsWith("version.cpp", StringComparison.InvariantCultureIgnoreCase) &&
 							!Product.Replace('\\', '/').EndsWith("DotNetCommon/MetaData.cs", StringComparison.InvariantCultureIgnoreCase) &&
-                            !Product.EndsWith("_Success.log", StringComparison.InvariantCultureIgnoreCase)
+                            !Product.EndsWith("_Success.log", StringComparison.InvariantCultureIgnoreCase) &&
+							!Product.Replace('\\', '/').Contains("/Intermediate/")
                             )
                         {
                             ProductsToSubmit.Add(Product);
