Remove "headers" nodes from GUBP, and special casing for Rocket builds. Rocket now generates headers into the project directory for all modules, so we don't need to save them off.

diff --git a/Engine/Source/Programs/AutomationTool/Scripts/GUBP.Automation.cs b/Engine/Source/Programs/AutomationTool/Scripts/GUBP.Automation.cs
--- a/Engine/Source/Programs/AutomationTool/Scripts/GUBP.Automation.cs
+++ b/Engine/Source/Programs/AutomationTool/Scripts/GUBP.Automation.cs
@@ -1747,129 +1747,6 @@ public class GUBP : BuildCommand
 		}
 	}
 
-    public class RootEditorHeadersNode : HostPlatformNode
-    {
-        public RootEditorHeadersNode(UnrealTargetPlatform InHostPlatform)
-            : base(InHostPlatform)
-        {
-            AgentSharingGroup = "Editor" + StaticGetHostPlatformSuffix(HostPlatform);
-            AddDependency(RootEditorNode.StaticGetFullName(HostPlatform));
-        }
-        public static string StaticGetFullName(UnrealTargetPlatform InHostPlatform)
-        {
-            return "RootEditorHeaders" + StaticGetHostPlatformSuffix(InHostPlatform);
-        }
-        public override string GetFullName()
-        {
-            return StaticGetFullName(HostPlatform);
-        }
-        public override float Priority()
-        {
-            return 1000000.0f; // right after the root editor
-        }
-        public override void DoBuild(GUBP bp)
-        {
-			HashSet<string> HeaderFiles = new HashSet<string>(StringComparer.InvariantCultureIgnoreCase);
-            foreach (var FileToCopy in CommandUtils.FindFiles("*.h", true, CommandUtils.CombinePaths(CmdEnv.LocalRoot, @"Engine\Intermediate\Build\", HostPlatform.ToString(), "Inc")))
-            {
-				HeaderFiles.Add(FileToCopy);
-            }
-            var Targets = new List<string> { bp.Branch.BaseEngineProject.Properties.Targets[TargetRules.TargetType.Editor].TargetName };
-            foreach (var ProgramTarget in bp.Branch.BaseEngineProject.Properties.Programs)
-            {
-                if (ProgramTarget.Rules.GUBP_AlwaysBuildWithBaseEditor() && ProgramTarget.Rules.SupportsPlatform(HostPlatform))
-                {
-                    Targets.Add(ProgramTarget.TargetName);
-                }
-            }
-            foreach (var Target in Targets)
-            {
-                foreach (var FileToCopy in CommandUtils.FindFiles("*.h", true, CommandUtils.CombinePaths(CmdEnv.LocalRoot, @"Engine\Intermediate\Build\", HostPlatform.ToString(), Target, "Inc")))
-                {
-					HeaderFiles.Add(FileToCopy);
-                }
-            }
-
-            var EnginePluginsDirectory = Path.Combine(CommandUtils.CmdEnv.LocalRoot, "Engine/Plugins");
-            var EnginePlugins = new List<PluginInfo>();
-            Plugins.FindPluginsIn(EnginePluginsDirectory, PluginInfo.LoadedFromType.Engine, ref EnginePlugins);
-
-            foreach (var EnginePlugin in EnginePlugins)
-            {
-                foreach (var FileToCopy in CommandUtils.FindFiles("*", true, CommandUtils.CombinePaths(EnginePlugin.Directory, @"Intermediate\Build", HostPlatform.ToString(), "Inc")))
-                {
-					HeaderFiles.Add(FileToCopy);
-                }
-            }
-
-			// Create a zip file containing the headers. 
-			string ZipFileName = CommandUtils.CombinePaths(CmdEnv.LocalRoot, "Engine\\Intermediate\\Build", HostPlatform.ToString(), "RootEditorHeaders.zip");
-			HeadersNode.ZipHeaders(HeaderFiles, ZipFileName);
-            BuildProducts = new List<string>{ ZipFileName };
-        }
-    }
-
-    public class GameMonolithicHeadersNode : HostPlatformNode
-    {
-        UnrealTargetPlatform TargetPlatform;
-
-        public GameMonolithicHeadersNode(GUBP bp, UnrealTargetPlatform InHostPlatform, UnrealTargetPlatform InTargetPlatform)
-            : base(InHostPlatform)
-        {
-            TargetPlatform = InTargetPlatform;
-            AgentSharingGroup = "UE4_" + InTargetPlatform + "_Mono" + StaticGetHostPlatformSuffix(InHostPlatform);
-            AddDependency(GamePlatformMonolithicsNode.StaticGetFullName(HostPlatform, bp.Branch.BaseEngineProject, TargetPlatform));            
-        }
-        public static string StaticGetFullName(UnrealTargetPlatform InHostPlatform, UnrealTargetPlatform InTargetPlatform)
-        {
-            return "UE4_" + InTargetPlatform + "_MonolithicHeaders" + StaticGetHostPlatformSuffix(InHostPlatform);
-        }
-        public override string GetFullName()
-        {
-            return StaticGetFullName(HostPlatform, TargetPlatform);
-        }
-        public override float Priority()
-        {
-            return 1000000.0f; // right after the game monolithics
-        }
-        public override void DoBuild(GUBP bp)
-        {
-            if (!UnrealBuildTool.Utils.IsRunningOnMono && TargetPlatform == UnrealTargetPlatform.IOS)
-            {
-                throw new AutomationException("these rocket header node require real mac path.");
-            }
-
-			HashSet<string> HeaderFiles = new HashSet<string>(StringComparer.InvariantCultureIgnoreCase);
-			if (TargetPlatform != HostPlatform)
-			{
-				// host platform overlaps with the editor, so we don't want them here
-				var PlatformDir = TargetPlatform.ToString();
-				foreach (var FileToCopy in CommandUtils.FindFiles("*.h", true, CommandUtils.CombinePaths(CmdEnv.LocalRoot, @"Engine\Intermediate\Build\", PlatformDir, "Inc")))
-				{
-					HeaderFiles.Add(FileToCopy);
-				}
-
-				// these may not be any new build products here, but we will check anyway
-				var EnginePluginsDirectory = Path.Combine(CommandUtils.CmdEnv.LocalRoot, "Engine/Plugins");
-				var EnginePlugins = new List<PluginInfo>();
-				Plugins.FindPluginsIn(EnginePluginsDirectory, PluginInfo.LoadedFromType.Engine, ref EnginePlugins);
-
-				foreach (var EnginePlugin in EnginePlugins)
-				{
-					foreach (var FileToCopy in CommandUtils.FindFiles("*", true, CommandUtils.CombinePaths(EnginePlugin.Directory, @"Intermediate\Build", PlatformDir, "Inc")))
-					{
-						HeaderFiles.Add(FileToCopy);
-					}
-				}
-			}
-
-			// Create a zip containing all the headers
-			string ZipFileName = CommandUtils.CombinePaths(CmdEnv.LocalRoot, "Engine\\Intermediate\\Build", TargetPlatform.ToString(), "MonolithicHeaders.zip");
-			HeadersNode.ZipHeaders(HeaderFiles, ZipFileName);
-            BuildProducts = new List<string>{ ZipFileName };
-        }
-    }
-
     public class SuccessNode : GUBPNode
     {
         public SuccessNode()
@@ -5597,10 +5474,6 @@ public class GUBP : BuildCommand
 			if (!BranchOptions.ExcludePlatformsForEditor.Contains(HostPlatform))
 			{
 				AddNode(new RootEditorNode(HostPlatform));			
-				if (bBuildRocket)
-				{
-					AddNode(new RootEditorHeadersNode(HostPlatform));
-				}
 				AddNode(new ToolsNode(HostPlatform));            
 				AddNode(new InternalToolsNode(HostPlatform));
 			if (HostPlatform == UnrealTargetPlatform.Win64 && ActivePlatforms.Contains(UnrealTargetPlatform.Linux))
@@ -5828,10 +5701,6 @@ public class GUBP : BuildCommand
                                 if (!GUBPNodes.ContainsKey(GamePlatformMonolithicsNode.StaticGetFullName(HostPlatform, Branch.BaseEngineProject, Plat)))
                                 {
                                     AddNode(new GamePlatformMonolithicsNode(this, HostPlatform, Branch.BaseEngineProject, Plat));
-                                    if (bBuildRocket)
-                                    {
-                                        AddNode(new GameMonolithicHeadersNode(this, HostPlatform, Plat));
-                                    }
                                 }
 								if (Plat == UnrealTargetPlatform.Win32 && Target.Rules.GUBP_BuildWindowsXPMonolithics() && Kind == TargetRules.TargetType.Game)
 								{
