Fix formatting.

diff --git a/Engine/Source/Programs/AutomationTool/Scripts/GUBP.Automation.cs b/Engine/Source/Programs/AutomationTool/Scripts/GUBP.Automation.cs
--- a/Engine/Source/Programs/AutomationTool/Scripts/GUBP.Automation.cs
+++ b/Engine/Source/Programs/AutomationTool/Scripts/GUBP.Automation.cs
@@ -1503,12 +1503,14 @@ public class GUBP : BuildCommand
         BranchInfo.BranchUProject GameProj;
         UnrealTargetPlatform TargetPlatform;
 		bool WithXp;
+
         public GamePlatformMonolithicsNode(GUBP bp, UnrealTargetPlatform InHostPlatform, BranchInfo.BranchUProject InGameProj, UnrealTargetPlatform InTargetPlatform, bool InWithXp = false)
             : base(InHostPlatform)
         {
             GameProj = InGameProj;
             TargetPlatform = InTargetPlatform;
 			WithXp = InWithXp;
+
             if (TargetPlatform == UnrealTargetPlatform.PS4)
             {
                 var PS4MapFileUtil = bp.Branch.FindProgram("PS4MapFileUtil");
@@ -1602,7 +1604,7 @@ public class GUBP : BuildCommand
 			if (!WithXp)
 			{
 				Name = InGameProj.GameName + "_" + InTargetPlatform + "_Mono" + StaticGetHostPlatformSuffix(InHostPlatform);
-        }
+	        }
 			else
 			{
 				Name = InGameProj.GameName + "_WinXP_Mono" + StaticGetHostPlatformSuffix(InHostPlatform);
@@ -1663,35 +1665,40 @@ public class GUBP : BuildCommand
 					var AllowXp = Target.Rules.GUBP_BuildWindowsXPMonolithics();
 					if (!WithXp || (AllowXp && WithXp))
 					{
-                    var Platforms = Target.Rules.GUBP_GetPlatforms_MonolithicOnly(HostPlatform);
-					var AdditionalPlatforms = Target.Rules.GUBP_GetBuildOnlyPlatforms_MonolithicOnly(HostPlatform);
-					var AllPlatforms = Platforms.Union(AdditionalPlatforms);
-					if (AllPlatforms.Contains(TargetPlatform) && Target.Rules.SupportsPlatform(TargetPlatform))
-                    {
-                        var Configs = Target.Rules.GUBP_GetConfigs_MonolithicOnly(HostPlatform, TargetPlatform);
-                        foreach (var Config in Configs)
-                        {
+						var Platforms = Target.Rules.GUBP_GetPlatforms_MonolithicOnly(HostPlatform);
+						var AdditionalPlatforms = Target.Rules.GUBP_GetBuildOnlyPlatforms_MonolithicOnly(HostPlatform);
+						var AllPlatforms = Platforms.Union(AdditionalPlatforms);
+						if (AllPlatforms.Contains(TargetPlatform) && Target.Rules.SupportsPlatform(TargetPlatform))
+						{
+							var Configs = Target.Rules.GUBP_GetConfigs_MonolithicOnly(HostPlatform, TargetPlatform);
+							foreach (var Config in Configs)
+							{
+								if(Config != UnrealTargetConfiguration.Development || Target.TargetName != "UE4Game")
+								{
+									continue;
+								}
+
 								if (WithXp)
 								{
 									Args += " -winxp";
 								}
 
-                            if (GUBP.bBuildRocket)
-                            {
-                                if (HostPlatform == UnrealTargetPlatform.Win64)
-                                {
-                                    if (TargetPlatform == UnrealTargetPlatform.Win32 && Config != UnrealTargetConfiguration.Shipping)
-                                    {
-                                        continue;
-                                    }
-                                    if (TargetPlatform == UnrealTargetPlatform.Win64 && Config != UnrealTargetConfiguration.Development)
-                                    {
-                                        continue;
-                                    }
-                                    if (TargetPlatform == UnrealTargetPlatform.Android && Config != UnrealTargetConfiguration.Shipping && Config != UnrealTargetConfiguration.Development)
-                                    {
-                                        continue;
-                                    }
+								if (GUBP.bBuildRocket)
+								{
+									if (HostPlatform == UnrealTargetPlatform.Win64)
+									{
+										if (TargetPlatform == UnrealTargetPlatform.Win32 && Config != UnrealTargetConfiguration.Shipping)
+										{
+											continue;
+										}
+										if (TargetPlatform == UnrealTargetPlatform.Win64 && Config != UnrealTargetConfiguration.Development)
+										{
+											continue;
+										}
+										if (TargetPlatform == UnrealTargetPlatform.Android && Config != UnrealTargetConfiguration.Shipping && Config != UnrealTargetConfiguration.Development)
+										{
+											continue;
+										}
 										if (TargetPlatform == UnrealTargetPlatform.HTML5 && Config != UnrealTargetConfiguration.Shipping && Config != UnrealTargetConfiguration.Development)
 										{
 											continue;
@@ -5700,7 +5707,7 @@ public class GUBP : BuildCommand
                                 }
                                 if (!GUBPNodes.ContainsKey(GamePlatformMonolithicsNode.StaticGetFullName(HostPlatform, Branch.BaseEngineProject, Plat)))
                                 {
-                                    AddNode(new GamePlatformMonolithicsNode(this, HostPlatform, Branch.BaseEngineProject, Plat));
+                                    AddNode(new GamePlatformMonolithicsNode(this, HostPlatform, Branch.BaseEngineProject, Plat, InSaveManifests: true));
                                 }
 								if (Plat == UnrealTargetPlatform.Win32 && Target.Rules.GUBP_BuildWindowsXPMonolithics() && Kind == TargetRules.TargetType.Game)
 								{
