Add separate monolithic nodes to build precompiled binaries (AKA Rocket static libraries) for each target platform. GUBP_GetConfigsForPrecompiledBuilds() in the target properties allows a target to specify which platform/configuration combinations to build.

diff --git a/Engine/Source/Programs/AutomationTool/Scripts/GUBP.Automation.cs b/Engine/Source/Programs/AutomationTool/Scripts/GUBP.Automation.cs
--- a/Engine/Source/Programs/AutomationTool/Scripts/GUBP.Automation.cs
+++ b/Engine/Source/Programs/AutomationTool/Scripts/GUBP.Automation.cs
@@ -1503,13 +1503,15 @@ public class GUBP : BuildCommand
         BranchInfo.BranchUProject GameProj;
         UnrealTargetPlatform TargetPlatform;
 		bool WithXp;
+		bool Precompiled; // If true, just builds targets which generate static libraries for the -UsePrecompiledModules option to UBT. If false, just build those that don't.
 
-        public GamePlatformMonolithicsNode(GUBP bp, UnrealTargetPlatform InHostPlatform, BranchInfo.BranchUProject InGameProj, UnrealTargetPlatform InTargetPlatform, bool InWithXp = false)
+        public GamePlatformMonolithicsNode(GUBP bp, UnrealTargetPlatform InHostPlatform, BranchInfo.BranchUProject InGameProj, UnrealTargetPlatform InTargetPlatform, bool InWithXp = false, bool InPrecompiled = false)
             : base(InHostPlatform)
         {
             GameProj = InGameProj;
             TargetPlatform = InTargetPlatform;
 			WithXp = InWithXp;
+			Precompiled = InPrecompiled;
 
             if (TargetPlatform == UnrealTargetPlatform.PS4)
             {
@@ -1547,12 +1549,11 @@ public class GUBP : BuildCommand
                 {
                     AddPseudodependency(GamePlatformMonolithicsNode.StaticGetFullName(InHostPlatform, bp.Branch.BaseEngineProject, InHostPlatform));
                 }
-                if (GUBP.bBuildRocket && InGameProj.GameName == bp.Branch.BaseEngineProject.GameName)
-                {
-                    AgentSharingGroup = "UE4_" + InTargetPlatform + "_Mono" + StaticGetHostPlatformSuffix(InHostPlatform);
-                    // lets just start this right away AddPseudodependency(RootEditorHeadersNode.StaticGetFullName(HostPlatform)); // maybe we should start these sooner, but that rather tangles the agent groups
-                }
             }
+			if(!WithXp && !Precompiled && HasPrecompiledTargets(GameProj, HostPlatform, TargetPlatform))
+			{
+				AddDependency(GamePlatformMonolithicsNode.StaticGetFullName(InHostPlatform, GameProj, InTargetPlatform, WithXp, Precompiled: true));
+			}
             if (InGameProj.Options(InHostPlatform).bTestWithShared)  /// compiling templates is only for testing purposes, and we will group them to avoid saturating the farm
             {
                 AddPseudodependency(WaitForTestShared.StaticGetFullName());
@@ -1598,22 +1599,26 @@ public class GUBP : BuildCommand
 
         }
 
-        public static string StaticGetFullName(UnrealTargetPlatform InHostPlatform, BranchInfo.BranchUProject InGameProj, UnrealTargetPlatform InTargetPlatform, bool WithXp = false)
+        public static string StaticGetFullName(UnrealTargetPlatform InHostPlatform, BranchInfo.BranchUProject InGameProj, UnrealTargetPlatform InTargetPlatform, bool WithXp = false, bool Precompiled = false)
         {
-			string Name;
-			if (!WithXp)
+			string Name = InGameProj.GameName;
+			if(WithXp)
 			{
-				Name = InGameProj.GameName + "_" + InTargetPlatform + "_Mono" + StaticGetHostPlatformSuffix(InHostPlatform);
-	        }
+				Name += "_WinXP_Mono";
+			}
 			else
 			{
-				Name = InGameProj.GameName + "_WinXP_Mono" + StaticGetHostPlatformSuffix(InHostPlatform);
+				Name += "_" + InTargetPlatform + "_Mono";
 			}
-			return Name;
+			if(Precompiled)
+			{
+				Name += "_Precompiled";
+			}
+			return Name + StaticGetHostPlatformSuffix(InHostPlatform);
         }
         public override string GetFullName()
         {
-            return StaticGetFullName(HostPlatform, GameProj, TargetPlatform, WithXp);
+            return StaticGetFullName(HostPlatform, GameProj, TargetPlatform, WithXp, Precompiled);
         }
         public override string GameNameIfAnyForTempStorage()
         {
@@ -1638,6 +1643,22 @@ public class GUBP : BuildCommand
             return Result;
         }
 
+		public static bool HasPrecompiledTargets(BranchInfo.BranchUProject Project, UnrealTargetPlatform HostPlatform, UnrealTargetPlatform TargetPlatform)
+		{
+            foreach(TargetRules.TargetType Kind in BranchInfo.MonolithicKinds)
+            {
+                if (Project.Properties.Targets.ContainsKey(Kind))
+                {
+                    SingleTargetProperties Target = Project.Properties.Targets[Kind];
+					if(Target.Rules.GUBP_GetConfigsForPrecompiledBuilds_MonolithicOnly(HostPlatform, TargetPlatform).Any())
+					{
+						return true;
+					}
+				}
+			}
+			return false;
+		}
+
         public override UE4Build.BuildAgenda GetAgenda(GUBP bp)
         {
             if (!bp.ActivePlatforms.Contains(TargetPlatform))
@@ -1653,6 +1674,16 @@ public class GUBP : BuildCommand
                 Args += " -nodebuginfo";
             }
 
+			if(Precompiled)
+			{
+				Args += " -precompilemodules";
+			}
+
+			if (WithXp)
+			{
+				Args += " -winxp";
+			}
+
             foreach (var Kind in BranchInfo.MonolithicKinds)
             {
                 if (GUBP.bBuildRocket && Kind != TargetRules.TargetType.Game)
@@ -1670,19 +1701,18 @@ public class GUBP : BuildCommand
 						var AllPlatforms = Platforms.Union(AdditionalPlatforms);
 						if (AllPlatforms.Contains(TargetPlatform) && Target.Rules.SupportsPlatform(TargetPlatform))
 						{
-							var Configs = Target.Rules.GUBP_GetConfigs_MonolithicOnly(HostPlatform, TargetPlatform);
+							List<UnrealTargetConfiguration> Configs;
+							if(Precompiled)
+							{
+								Configs = Target.Rules.GUBP_GetConfigsForPrecompiledBuilds_MonolithicOnly(HostPlatform, TargetPlatform);
+							}
+							else
+							{
+								Configs = Target.Rules.GUBP_GetConfigs_MonolithicOnly(HostPlatform, TargetPlatform);
+							}
+							
 							foreach (var Config in Configs)
 							{
-								if(Config != UnrealTargetConfiguration.Development || Target.TargetName != "UE4Game")
-								{
-									continue;
-								}
-
-								if (WithXp)
-								{
-									Args += " -winxp";
-								}
-
 								if (GUBP.bBuildRocket)
 								{
 									if (HostPlatform == UnrealTargetPlatform.Win64)
@@ -1703,28 +1733,28 @@ public class GUBP : BuildCommand
 										{
 											continue;
 										}
-                                    if (TargetPlatform == UnrealTargetPlatform.Linux && Config != UnrealTargetConfiguration.Shipping && Config != UnrealTargetConfiguration.Development)
-                                    {
-                                        continue;
-                                    }
-                                }
-                                else if (Config != UnrealTargetConfiguration.Shipping && Config != UnrealTargetConfiguration.Development)
-                                {
-                                    continue;
-                                }
-                            }
-                            if (GameProj.GameName == bp.Branch.BaseEngineProject.GameName)
-                            {
-                                Agenda.AddTargets(new string[] { Target.TargetName }, TargetPlatform, Config, InAddArgs: Args);
-                            }
-                            else
-                            {
-                                Agenda.AddTargets(new string[] { Target.TargetName }, TargetPlatform, Config, GameProj.FilePath, InAddArgs: Args);
-                            }
-                        }
-                    }
-                }
-            }
+										if (TargetPlatform == UnrealTargetPlatform.Linux && Config != UnrealTargetConfiguration.Shipping && Config != UnrealTargetConfiguration.Development)
+										{
+											continue;
+										}
+									}
+									else if (Config != UnrealTargetConfiguration.Shipping && Config != UnrealTargetConfiguration.Development)
+									{
+										continue;
+									}
+								}
+								if (GameProj.GameName == bp.Branch.BaseEngineProject.GameName)
+								{
+									Agenda.AddTargets(new string[] { Target.TargetName }, TargetPlatform, Config, InAddArgs: Args);
+								}
+								else
+								{
+									Agenda.AddTargets(new string[] { Target.TargetName }, TargetPlatform, Config, GameProj.FilePath, InAddArgs: Args);
+								}
+							}
+						}
+					}
+				}
             }
 
             return Agenda;
@@ -2025,11 +2055,11 @@ public class GUBP : BuildCommand
 							if(GameProj.Properties.Targets[TargetRules.TargetType.Game].Rules.GUBP_BuildWindowsXPMonolithics())
 							{
 								AddDependency(GamePlatformMonolithicsNode.StaticGetFullName(HostPlatform, GameProj, Plat, true));
-                    }
-                }
-            }
-        }
-        }
+							}
+						}
+					}
+				}
+			}
         }
 
         public static string StaticGetFullName(BranchInfo.BranchUProject InGameProj)
@@ -2398,7 +2428,10 @@ public class GUBP : BuildCommand
                             !Product.EndsWith("version.cpp", StringComparison.InvariantCultureIgnoreCase) &&
 							!Product.Replace('\\', '/').EndsWith("DotNetCommon/MetaData.cs", StringComparison.InvariantCultureIgnoreCase) &&
                             !Product.EndsWith("_Success.log", StringComparison.InvariantCultureIgnoreCase) &&
-							!Product.Replace('\\', '/').Contains("/Intermediate/")
+							!Product.Replace('\\', '/').Contains("/Intermediate/") &&
+							!Product.EndsWith(".lib") &&
+							!Product.EndsWith(".a") && 
+							!Product.EndsWith(".bc")
                             )
                         {
                             ProductsToSubmit.Add(Product);
@@ -5707,6 +5740,10 @@ public class GUBP : BuildCommand
                                 }
                                 if (!GUBPNodes.ContainsKey(GamePlatformMonolithicsNode.StaticGetFullName(HostPlatform, Branch.BaseEngineProject, Plat)))
                                 {
+									if(GamePlatformMonolithicsNode.HasPrecompiledTargets(Branch.BaseEngineProject, HostPlatform, Plat))
+									{
+	                                    AddNode(new GamePlatformMonolithicsNode(this, HostPlatform, Branch.BaseEngineProject, Plat, InPrecompiled: true));
+									}
                                     AddNode(new GamePlatformMonolithicsNode(this, HostPlatform, Branch.BaseEngineProject, Plat));
                                 }
 								if (Plat == UnrealTargetPlatform.Win32 && Target.Rules.GUBP_BuildWindowsXPMonolithics() && Kind == TargetRules.TargetType.Game)
@@ -5715,7 +5752,7 @@ public class GUBP : BuildCommand
 									{
 										AddNode(new GamePlatformMonolithicsNode(this, HostPlatform, Branch.BaseEngineProject, Plat, true));
 									}
-                            }
+	                            }
                             }
                         }
                     }
@@ -5987,6 +6024,10 @@ public class GUBP : BuildCommand
                                 }
                                 if (!GUBPNodes.ContainsKey(GamePlatformMonolithicsNode.StaticGetFullName(HostPlatform, CodeProj, Plat)))
                                 {
+									if(GamePlatformMonolithicsNode.HasPrecompiledTargets(CodeProj, HostPlatform, Plat))
+									{
+	                                    AddNode(new GamePlatformMonolithicsNode(this, HostPlatform, CodeProj, Plat, InPrecompiled: true));
+									}
                                     AddNode(new GamePlatformMonolithicsNode(this, HostPlatform, CodeProj, Plat));
                                 }
 								if (!AdditionalPlatforms.Contains(Plat))
diff --git a/Engine/Source/Programs/UnrealBuildTool/System/RulesCompiler.cs b/Engine/Source/Programs/UnrealBuildTool/System/RulesCompiler.cs
--- a/Engine/Source/Programs/UnrealBuildTool/System/RulesCompiler.cs
+++ b/Engine/Source/Programs/UnrealBuildTool/System/RulesCompiler.cs
@@ -781,6 +781,14 @@ namespace UnrealBuildTool
             return new List<UnrealTargetConfiguration>{UnrealTargetConfiguration.Development};
         }
         /// <summary>
+        /// Return a list of configs which are precompiled for the given target platform
+        /// </summary>
+        /// <returns>a list of configs for a target platforms for the monolithic</returns>        
+        public virtual List<UnrealTargetConfiguration> GUBP_GetConfigsForPrecompiledBuilds_MonolithicOnly(UnrealTargetPlatform HostPlatform, UnrealTargetPlatform Platform)
+        {
+            return new List<UnrealTargetConfiguration>();
+        }
+        /// <summary>
         /// Return a list of configs for target platforms for formal builds
         /// </summary>
         /// <returns>a list of configs for a target platforms for the monolithic</returns>        
diff --git a/Engine/Source/UE4Game.Target.cs b/Engine/Source/UE4Game.Target.cs
--- a/Engine/Source/UE4Game.Target.cs
+++ b/Engine/Source/UE4Game.Target.cs
@@ -121,6 +121,27 @@ public class UE4GameTarget : TargetRules
     public override List<UnrealTargetConfiguration> GUBP_GetConfigs_MonolithicOnly(UnrealTargetPlatform HostPlatform, UnrealTargetPlatform Platform)
     {
         return new List<UnrealTargetConfiguration> { UnrealTargetConfiguration.Development, UnrealTargetConfiguration.Shipping, UnrealTargetConfiguration.Test };
+    }
+    public override List<UnrealTargetConfiguration> GUBP_GetConfigsForPrecompiledBuilds_MonolithicOnly(UnrealTargetPlatform HostPlatform, UnrealTargetPlatform Platform)
+    {
+		List<UnrealTargetConfiguration> Platforms = new List<UnrealTargetConfiguration>();
+		if(HostPlatform == UnrealTargetPlatform.Mac)
+		{
+			if(Platform == UnrealTargetPlatform.Mac || Platform == UnrealTargetPlatform.IOS)
+			{
+				Platforms.Add(UnrealTargetConfiguration.Development);
+				Platforms.Add(UnrealTargetConfiguration.Shipping);
+			}
+		}
+		else if(HostPlatform == UnrealTargetPlatform.Win64)
+		{
+			if(Platform == UnrealTargetPlatform.Win32 || Platform == UnrealTargetPlatform.Win64 || Platform == UnrealTargetPlatform.Android || Platform == UnrealTargetPlatform.HTML5)
+			{
+				Platforms.Add(UnrealTargetConfiguration.Development);
+				Platforms.Add(UnrealTargetConfiguration.Shipping);
+			}
+		}
+		return Platforms;
     }
 	public override bool GUBP_BuildWindowsXPMonolithics()
 	{
