UAT: Fix compactify so that it builds BuildPatchTool.exe when missing.
Also refactor compactify mode of BuildPatchTool to follow the pattern used by patch generation and data file enumeration, and remove legacy C# code now performed by BuildPatchTool.exe.

diff --git a/Engine/Source/Programs/AutomationTool/MCPPublic.cs b/Engine/Source/Programs/AutomationTool/MCPPublic.cs
--- a/Engine/Source/Programs/AutomationTool/MCPPublic.cs
+++ b/Engine/Source/Programs/AutomationTool/MCPPublic.cs
@@ -331,6 +331,48 @@ namespace EpicGames.MCP.Automation
 			public List<KeyValuePair<string, float>> CustomFloatArgs;
 		}
 
+		/// <summary>
+		/// Represents the options passed to the compactify process
+		/// </summary>
+		public class CompactifyOptions
+		{
+			private static readonly int DefaultDataAgeThreshold = PatchGenerationOptions.DEFAULT_DATA_AGE_THRESHOLD + 2;
+
+			public CompactifyOptions()
+			{
+				DataAgeThreshold = DefaultDataAgeThreshold;
+			}
+
+			/// <summary>
+			/// BuildPatchTool will run a compactify on this directory.
+			/// </summary>
+			public string CompactifyDirectory;
+			/// <summary>
+			/// Corresponds to the -preview parameter
+			/// </summary>
+			public bool bPreviewCompactify;
+			/// <summary>
+			/// Corresponds to the -nopatchdelete parameter
+			/// </summary>
+			public bool bNoPatchDeleteCompactify;
+			/// <summary>
+			/// The full list of manifest files in the compactify directory that we wish to keep; all others will be deleted.
+			/// </summary>
+			public string[] ManifestsToKeep;
+			/// <summary>
+			/// A filename (relative to the compactify directory) which contains a list of manifests to keep, one manifest per line.
+			/// N.b. If ManifestsToKeep is specified, then this option is ignored.
+			/// </summary>
+			public string ManifestsToKeepFile;
+			/// <summary>
+			/// Path data files modified within this number of days will *not* be deleted, allowing them to be reused by patch generation processes.
+			/// IMPORTANT: This should always be larger than the data age threshold for any build processes which will run on the directory, to ensure
+			/// that we do not delete any files which could be reused by a concurrently running build. It is recommended that this number be at least
+			/// two days greater than the build data age threshold.
+			/// </summary>
+			public int DataAgeThreshold;
+		}
+
 		public class DataEnumerationOptions
 		{
 			/// <summary>
@@ -380,6 +422,12 @@ namespace EpicGames.MCP.Automation
 		/// <param name="Opts">Parameters which will be passed to the patch tool generation process</param>
 		public abstract void Execute(PatchGenerationOptions Opts);
 
+		/// <summary>
+		/// Runs the Build Patch Tool executable to compactify a cloud directory using the supplied parameters.
+		/// </summary>
+		/// <param name="Opts">Parameters which will be passed to the patch tool generation process</param>
+		public abstract void Execute(CompactifyOptions Opts);
+
 		/// <summary>
 		/// Runs the Build Patch Tool executable to enumerate patch data files referenced by a manifest using the supplied parameters.
 		/// </summary>
