Simplify filtering for Rocket nodes; use list of rules in Engine/Build/InstalledEngineFilters.ini.

diff --git a/Engine/Build/InstalledEngineFilters.ini b/Engine/Build/InstalledEngineFilters.ini
new file mode 100644
--- /dev/null
+++ b/Engine/Build/InstalledEngineFilters.ini
@@ -0,0 +1,164 @@
+[CopyEditor]
+; Tool dependencies
+/Engine/Binaries/DotNET/Ionic.Zip.Reduced.dll
+/Engine/Binaries/DotNET/OneSky.dll
+
+; In-editor documentation
+/Engine/Documentation/Source/Shared/...
+
+; Content folders
+/Engine/Content/...
+-/Engine/Content/....psd
+-/Engine/Content/....pdn
+-/Engine/Content/....fbx
+-/Engine/Content/....archive
+-/Engine/Content/....manifest
+
+; Config files
+/Engine/Config/...
+-/Engine/Config/....vdf
+/Engine/Programs/UnrealBuildTool/BuildConfiguration.xml
+/Engine/Programs/UnrealFrontend/Config/DefaultEngine.ini
+/Engine/Programs/UnrealLightmass/Config/DefaultEngine.ini
+
+; Plugins
+/Engine/Plugins/....uplugin
+/Engine/Plugins/.../Content/...
+
+; Shaders
+/Engine/Shaders/...
+
+; Win64-specific dependencies
+{Win64} /Engine/Binaries/ThirdParty/ImgTec/PVRTexToolCLI.exe
+
+{Win64} /Engine/Binaries/Win64/libfbxsdk.dll
+{Win64} /Engine/Binaries/ThirdParty/CEF3/Win32/...
+{Win64} /Engine/Binaries/ThirdParty/CEF3/Win64/...
+{Win64} /Engine/Binaries/ThirdParty/ICU/icu4c-53_1/Win32/VS2013/*
+{Win64} /Engine/Binaries/ThirdParty/ICU/icu4c-53_1/Win64/VS2013/*
+{Win64} /Engine/Binaries/ThirdParty/nvTextureTools/Win64/*.dll
+{Win64} /Engine/Binaries/ThirdParty/Ogg/Win64/VS2013/*.dll
+{Win64} /Engine/Binaries/ThirdParty/Vorbis/Win64/VS2013/*.dll
+{Win64} /Engine/Binaries/ThirdParty/PhysX/APEX-1.3/Win64/VS2013/*PROFILE_x64.dll
+{Win64} /Engine/Binaries/ThirdParty/PhysX/PhysX-3.3/Win64/VS2013/*PROFILE_x64.dll
+{Win64} /Engine/Binaries/ThirdParty/PhysX/PhysX-3.3/Win64/VS2013/nvToolsExt64_1.dll
+{Win64} /Engine/Binaries/ThirdParty/Perforce/p4api.dll
+{Win64} /Engine/Binaries/ThirdParty/Perforce/p4api.xml
+{Win64} /Engine/Binaries/ThirdParty/Steamworks/Steamv132/Win64/*.dll
+{Win64} /Engine/Binaries/ThirdParty/svn/Win64/...
+{Win64} /Engine/Binaries/ThirdParty/OpenSSL/Win64/VS2012/*.dll
+{Win64} /Engine/Binaries/ThirdParty/OpenSSL/Win64/VS2013/*.dll
+
+{Win64} /Engine/Build/BatchFiles/Build.bat
+{Win64} /Engine/Build/BatchFiles/Clean.bat
+{Win64} /Engine/Build/BatchFiles/Rebuild.bat
+{Win64} /Engine/Build/BatchFiles/RocketGenerateProjectFiles.bat
+{Win64} /Engine/Build/BatchFiles/RunUAT.bat
+{Win64} /Engine/Build/BatchFiles/GetVSComnToolsPath.bat
+
+{Win64} /Engine/Extras/3dsMaxScripts/*.ms
+{Win64} /Engine/Extras/VisualStudioDebugging/UE4.natvis
+{Win64} /Engine/Extras/Maya_AnimationRiggingTools/...
+{Win64} /Engine/Extras/MayaVelocityGridExporter/...
+{Win64} /Engine/Extras/UnrealVS/...
+{Win64} /Engine/Extras/Redist/en-us/*
+
+; Mac-specific dependencies
+{Mac} /Engine/Binaries/ThirdParty/CEF3/Mac/...
+{Mac} /Engine/Binaries/ThirdParty/ICU/icu4c-53_1/Mac/*
+{Mac} /Engine/Binaries/ThirdParty/ImgTec/PVRTexToolCLI
+
+{Mac} /Engine/Binaries/ThirdParty/svn/Mac/...
+{Mac} /Engine/Binaries/ThirdParty/Mono/Mac/...
+{Mac} /Engine/Source/ThirdParty/Mac/RadioEffectUnit/...
+
+{Mac} /Engine/Build/BatchFiles/Mac/Build.sh
+{Mac} /Engine/Build/BatchFiles/Mac/GenerateLLDBInit.sh
+{Mac} /Engine/Build/BatchFiles/Mac/FixDependencyFiles.sh
+{Mac} /Engine/Build/BatchFiles/Mac/SetupMono.sh
+{Mac} /Engine/Build/BatchFiles/Mac/FixMonoFiles.sh
+{Mac} /Engine/Build/BatchFiles/Mac/RocketBuild.sh
+{Mac} /Engine/Build/BatchFiles/Mac/RocketGenerateProjectFiles.sh
+{Mac} /Engine/Build/BatchFiles/Mac/RunMono.sh
+{Mac} /Engine/Build/BatchFiles/Mac/RunXBuild.sh
+{Mac} /Engine/Build/BatchFiles/Mac/xcodeunlock.sh
+{Mac} /Engine/Build/BatchFiles/RunUAT.command
+{Mac} /Engine/Build/BatchFiles/RunUAT.sh
+
+{Mac} /Engine/Extras/LLDBDataFormatters/UE4DataFormatters.py
+
+{Mac} /Engine/Source/UE4Game*.cs
+
+[CopyTargetPlatforms]
+; API docs
+/Engine/Documentation/CHM/API.chm
+
+; Source code
+/Engine/Source/Runtime/...
+/Engine/Source/Developer/...
+/Engine/Source/Editor/...
+/Engine/Source/ThirdParty/Licenses/...
+/Engine/Plugins/.../Source/...
+
+; Starter content
+/Samples/StarterContent/Content/...
+/Samples/MobileStarterContent/Content/...
+
+[CopyTargetPlatform.Android]
+; Build files
+/Engine/Build/Android/...
+/Engine/Source/ThirdParty/Android/cxa_demangle/*
+{Win64} /Engine/Binaries/ThirdParty/ARM/Win32/...
+
+; TADP installer
+{Win64} /Engine/Extras/Android/*.exe
+{Mac} /Engine/Extras/Android/*.dmg
+
+[CopyTargetPlatform.IOS]
+; Build files
+/Engine/Build/IOS/...
+-/Engine/Build/IOS/....png
+-/Engine/Build/IOS/....psd
+-/Engine/Build/IOS/....mobileprovision
+-/Engine/Build/IOS/UnrealRemoteTool
+
+; Content-only on Win64: include pre-compiled executables manually
+{Win64} /Engine/Binaries/IOS/UnrealLaunchDaemon*
+{Win64} /Engine/Binaries/IOS/UE4Game.stub
+{Win64} /Engine/Binaries/IOS/UE4Game-IOS-Shipping.stub
+{Win64} /Engine/Binaries/DotNET/IOS/openssl.exe
+{Win64} /Engine/Binaries/ThirdParty/IOS/*
+{Win64} /Engine/Binaries/ThirdParty/ARM/Win32/...
+
+; All tools for Mac
+{Mac} /Engine/Binaries/DotNET/iOS/DeploymentServer.exe
+{Mac} /Engine/Binaries/DotNET/iOS/DeploymentServer.exe.config
+{Mac} /Engine/Binaries/DotNET/iOS/DeploymentInterface.dll
+{Mac} /Engine/Binaries/DotNET/iOS/DeploymentInterface.dll.config
+{Mac} /Engine/Binaries/DotNET/iOS/iPhonePackager.exe
+{Mac} /Engine/Binaries/DotNET/iOS/iPhonePackager.exe.config
+{Mac} /Engine/Binaries/DotNET/iOS/MobileDeviceInterface.dll
+{Mac} /Engine/Binaries/DotNET/iOS/MobileDeviceInterface.dll.config
+
+[CopyTargetPlatform.HTML5]
+; Build files
+/Engine/Build/HTML5/...
+
+[CopyTargetPlatform.Win32]
+; Just third party dependencies.
+/Engine/Binaries/ThirdParty/Ogg/Win32/VS2013/*.dll
+/Engine/Binaries/ThirdParty/Vorbis/Win32/VS2013/*.dll
+/Engine/Binaries/ThirdParty/PhysX/APEX-1.3/Win32/VS2013/*PROFILE_x86.dll
+/Engine/Binaries/ThirdParty/PhysX/PhysX-3.3/Win32/VS2013/*PROFILE_x86.dll
+/Engine/Binaries/ThirdParty/PhysX/PhysX-3.3/Win32/VS2013/nvToolsExt32_1.dll
+/Engine/Binaries/ThirdParty/Steamworks/Steamv132/Win32/*.dll
+
+[CopyTargetPlatform.Linux]
+; Content-only: include pre-compiled executables manually
+/Engine/Binaries/Linux/CrashReportClient
+/Engine/Binaries/Linux/UE4Game
+/Engine/Binaries/Linux/UE4Game-Linux-Shipping
+/Engine/Binaries/ThirdParty/ICU/icu4c-53_1/Linux/x86_64-unknown-linux-gnu/*.so
+/Engine/Config/Controller.vdf
+/Engine/Binaries/ThirdParty/OpenAL/Linux/*
+/Engine/Binaries/ThirdParty/Steamworks/Steamv132/Linux/*.so
diff --git a/Engine/Source/Programs/AutomationTool/FileFilter.cs b/Engine/Source/Programs/AutomationTool/FileFilter.cs
--- a/Engine/Source/Programs/AutomationTool/FileFilter.cs
+++ b/Engine/Source/Programs/AutomationTool/FileFilter.cs
@@ -8,6 +8,15 @@ using System.Linq;
 using System.Text;
 using System.Text.RegularExpressions;
 
+/// <summary>
+/// Indicates whether files which match a pattern should be included or excluded
+/// </summary>
+public enum FileFilterType
+{
+	Include,
+	Exclude,
+}
+
 /// <summary>
 /// Stores a set of rules, similar to p4 path specifications, which can be used to efficiently filter files in or out of a set.
 /// The rules are merged into a tree with nodes for each path fragments and rule numbers in each leaf, allowing tree traversal in an arbitrary order 
@@ -42,7 +51,7 @@ public class FileFilter
 	{
 		foreach (FileFilterNode OtherRule in Other.Rules)
 		{
-			AddRule(OtherRule.ToString(), OtherRule.bInclude);
+			AddRule(OtherRule.ToString(), OtherRule.Type);
 		}
 	}
 
@@ -55,13 +64,103 @@ public class FileFilter
 	{
 		foreach (string Line in Lines)
 		{
-			if (Line.StartsWith("-"))
+			AddRule(Line);
+		}
+	}
+
+	/// <summary>
+	/// Adds an include or exclude rule to the filter
+	/// </summary>
+	/// <param name="Pattern">Pattern to match. See CreateRegex() for details.</param>
+	public void AddRule(string Rule)
+	{
+		if (Rule.StartsWith("-"))
+		{
+			Exclude(Rule.Substring(1).TrimStart());
+		}
+		else
+		{
+			Include(Rule);
+		}
+	}
+
+	/// <summary>
+	/// Adds an include or exclude rule to the filter. The rule may be 
+	/// </summary>
+	/// <param name="Pattern">Pattern to match. See CreateRegex() for details.</param>
+	public void AddRule(string Rule, params string[] AllowTags)
+	{
+		string CleanRule = Rule.Trim();
+		if(CleanRule.StartsWith("{"))
+		{
+			// Find the end of the condition
+			int ConditionEnd = CleanRule.IndexOf('}');
+			if(ConditionEnd == -1)
 			{
-				Exclude(Line.Substring(1).TrimStart());
+				throw new Exception(String.Format("Missing closing parenthesis in rule: {0}", CleanRule));
 			}
-			else
+
+			// Check there's a matching tag
+			string[] RuleTags = CleanRule.Substring(1, ConditionEnd - 1).Split(',').Select(x => x.Trim()).ToArray();
+			if(!RuleTags.Any(x => AllowTags.Contains(x)))
 			{
-				Include(Line);
+				return;
+			}
+
+			// Strip the condition from the rule
+			CleanRule = CleanRule.Substring(ConditionEnd + 1).TrimStart(); 
+		}
+		AddRule(CleanRule);
+	}
+
+	/// <summary>
+	/// Adds several rules to the filter
+	/// </summary>
+	/// <param name="Patterns">List of patterns to match.</param>
+	public void AddRules(IEnumerable<string> Rules)
+	{
+		foreach(string Rule in Rules)
+		{
+			AddRule(Rule);
+		}
+	}
+
+	/// <summary>
+	/// Adds several rules in the given lines. Rules may be prefixed with conditions of the syntax {Key=Value, Key2=Value2}, which
+	/// will be evaluated using variables in the given dictionary before being added.
+	/// </summary>
+	/// <param name="Lines"></param>
+	/// <param name="Variables">Lookup for variables to test against</param>
+	public void AddRules(IEnumerable<string> Rules, params string[] Tags)
+	{
+		foreach(string Rule in Rules)
+		{
+			AddRule(Rule, Tags);
+		}
+	}
+
+	/// <summary>
+	/// Reads a configuration file split into sections
+	/// </summary>
+	/// <param name="Filter"></param>
+	/// <param name="RulesFileName"></param>
+	/// <param name="Conditions"></param>
+	public void AddRulesFromFile(string FileName, string SectionName, params string[] AllowTags)
+	{
+		bool bInSection = false;
+		foreach(string Line in File.ReadAllLines(FileName))
+		{
+			string TrimLine = Line.Trim();
+			if(!TrimLine.StartsWith(";") && TrimLine.Length > 0)
+			{
+				if(TrimLine.StartsWith("["))
+				{
+					bInSection = (TrimLine == "[" + SectionName + "]");
+				}
+				else if(bInSection)
+				{
+					AddRule(Line, AllowTags);
+				}
 			}
 		}
 	}
@@ -72,7 +171,7 @@ public class FileFilter
 	/// <param name="Pattern">Pattern to match. See CreateRegex() for details.</param>
 	public void Include(string Pattern)
 	{
-		AddRule(Pattern, true);
+		AddRule(Pattern, FileFilterType.Include);
 	}
 
 	/// <summary>
@@ -102,7 +201,7 @@ public class FileFilter
 	/// <param name="Pattern">Mask to match. See CreateRegex() for details.</param>
 	public void Exclude(string Pattern)
 	{
-		AddRule(Pattern, false);
+		AddRule(Pattern, FileFilterType.Exclude);
 	}
 
 	/// <summary>
@@ -131,7 +230,7 @@ public class FileFilter
 	/// </summary>
 	/// <param name="Pattern">The pattern which the rule should match</param>
 	/// <param name="bInclude">Whether to include or exclude files matching this rule</param>
-	void AddRule(string Pattern, bool bInclude)
+	public void AddRule(string Pattern, FileFilterType Type)
 	{
 		string NormalizedPattern = Pattern.Replace('\\', '/');
 
@@ -188,12 +287,12 @@ public class FileFilter
 		// We've reached the end of the pattern, so mark it as a leaf node
 		Rules.Add(LastNode);
 		LastNode.RuleNumber = Rules.Count - 1;
-		LastNode.bInclude = bInclude;
+		LastNode.Type = Type;
 
 		// Update the maximums along that path
 		for (FileFilterNode UpdateNode = LastNode; UpdateNode != null; UpdateNode = UpdateNode.Parent)
 		{
-			if (bInclude)
+			if (Type == FileFilterType.Include)
 			{
 				UpdateNode.MaxIncludeRuleNumber = LastNode.RuleNumber;
 			}
@@ -215,7 +314,7 @@ public class FileFilter
 
 		FileFilterNode MatchingNode = FindMatchingNode(RootNode, Tokens, 0, RootNode);
 
-		return MatchingNode.bInclude;
+		return MatchingNode.Type == FileFilterType.Include;
 	}
 
 	/// <summary>
@@ -229,7 +328,7 @@ public class FileFilter
 
 		FileFilterNode MatchingNode = FindMatchingNode(RootNode, Tokens, 0, RootNode);
 
-		return MatchingNode.bInclude || HighestPossibleIncludeMatch(RootNode, Tokens, 0, MatchingNode.RuleNumber) > MatchingNode.RuleNumber;
+		return MatchingNode.Type == FileFilterType.Include || HighestPossibleIncludeMatch(RootNode, Tokens, 0, MatchingNode.RuleNumber) > MatchingNode.RuleNumber;
 	}
 
 	/// <summary>
@@ -296,7 +395,7 @@ public class FileFilter
 		}
 
 		// If there is no rule under the current node which is better than the current best node, early out
-		if (CurrentBestNode.bInclude)
+		if (CurrentBestNode.Type == FileFilterType.Include)
 		{
 			if (CurrentNode.MaxExcludeRuleNumber < CurrentBestNode.RuleNumber)
 			{
@@ -411,9 +510,9 @@ class FileFilterNode
 	public int RuleNumber;
 
 	/// <summary>
-	/// True if the rule terminating in this node was an include node.
+	/// Whether a rule terminating in this node should include files or exclude files.
 	/// </summary>
-	public bool bInclude;
+	public FileFilterType Type;
 
 	/// <summary>
 	/// Default constructor.
@@ -423,7 +522,7 @@ class FileFilterNode
 		Parent = InParent;
 		Pattern = InPattern;
 		RuleNumber = -1;
-		bInclude = false;
+		Type = FileFilterType.Exclude;
 	}
 
 	/// <summary>
diff --git a/Engine/Source/Programs/AutomationTool/Scripts/GUBP.Automation.cs b/Engine/Source/Programs/AutomationTool/Scripts/GUBP.Automation.cs
--- a/Engine/Source/Programs/AutomationTool/Scripts/GUBP.Automation.cs
+++ b/Engine/Source/Programs/AutomationTool/Scripts/GUBP.Automation.cs
@@ -1405,7 +1405,7 @@ public class GUBP : BuildCommand
         {
 			Projects = new List<BranchInfo.BranchUProject>(InProjects);
 			AddDependency(ToolsNode.StaticGetFullName(InHostPlatform)); // for UnrealPak
-            AgentSharingGroup = "FeaturePacks"  + StaticGetHostPlatformSuffix(InHostPlatform);
+			AgentSharingGroup = "BuildInstall" + StaticGetHostPlatformSuffix(HostPlatform);
         }
 
 		public static string GetOutputFile(BranchInfo.BranchUProject Project)
