Add EnvVarsToXML to build agenda for GitHub binaries, otherwise they cannot run UBT. Rename Release configuration to Development so that UAT can build it.

diff --git a/Engine/Source/Programs/AutomationTool/Scripts/BuildCommonTools.Automation.cs b/Engine/Source/Programs/AutomationTool/Scripts/BuildCommonTools.Automation.cs
--- a/Engine/Source/Programs/AutomationTool/Scripts/BuildCommonTools.Automation.cs
+++ b/Engine/Source/Programs/AutomationTool/Scripts/BuildCommonTools.Automation.cs
@@ -89,6 +89,7 @@ public class BuildCommonTools : BuildCommand
 		Agenda.DotNetProjects.Add(@"Engine/Source/Programs/DotNETCommon/DotNETUtilities/DotNETUtilities.csproj");
 		Agenda.DotNetProjects.Add(@"Engine/Source/Programs/RPCUtility/RPCUtility.csproj");
 		Agenda.DotNetProjects.Add(@"Engine/Source/Programs/UnrealControls/UnrealControls.csproj");
+		Agenda.DotNetProjects.Add(@"Engine/Source/Programs/EnvVarsToXML/EnvVarsToXML/EnvVarsToXML.csproj");
 
 		// Windows binaries
 		if(Platforms.Contains(UnrealBuildTool.UnrealTargetPlatform.Win64))
diff --git a/Engine/Source/Programs/EnvVarsToXML/EnvVarsToXML/EnvVarsToXML.csproj b/Engine/Source/Programs/EnvVarsToXML/EnvVarsToXML/EnvVarsToXML.csproj
--- a/Engine/Source/Programs/EnvVarsToXML/EnvVarsToXML/EnvVarsToXML.csproj
+++ b/Engine/Source/Programs/EnvVarsToXML/EnvVarsToXML/EnvVarsToXML.csproj
@@ -22,7 +22,7 @@
     <ErrorReport>prompt</ErrorReport>
     <WarningLevel>4</WarningLevel>
   </PropertyGroup>
-  <PropertyGroup Condition=" '$(Configuration)|$(Platform)' == 'Release|AnyCPU' ">
+  <PropertyGroup Condition=" '$(Configuration)|$(Platform)' == 'Development|AnyCPU' ">
     <PlatformTarget>AnyCPU</PlatformTarget>
     <DebugType>pdbonly</DebugType>
     <Optimize>true</Optimize>
