P4 utility functions to allow updating changelists and deleting shelved files.

diff --git a/Engine/Source/Programs/AutomationTool/P4Utils.cs b/Engine/Source/Programs/AutomationTool/P4Utils.cs
--- a/Engine/Source/Programs/AutomationTool/P4Utils.cs
+++ b/Engine/Source/Programs/AutomationTool/P4Utils.cs
@@ -205,6 +205,118 @@ namespace AutomationTool
 		public string Path;
 	}
 
+	public class P4Spec
+	{
+		public List<KeyValuePair<string, string>> Sections;
+
+		/// <summary>
+		/// Default constructor.
+		/// </summary>
+		public P4Spec()
+		{
+			Sections = new List<KeyValuePair<string,string>>();
+		}
+
+		/// <summary>
+		/// Gets the current value of a field with the given name 
+		/// </summary>
+		/// <param name="Name">Name of the field to search for</param>
+		/// <returns>The value of the field, or null if it does not exist</returns>
+		public string GetField(string Name)
+		{
+			int Idx = Sections.FindIndex(x => x.Key == Name);
+			return (Idx == -1)? null : Sections[Idx].Value;
+		}
+
+		/// <summary>
+		/// Sets the value of an existing field, or adds a new one with the given name
+		/// </summary>
+		/// <param name="Name">Name of the field to set</param>
+		/// <param name="Value">New value of the field</param>
+		public void SetField(string Name, string Value)
+		{
+			int Idx = Sections.FindIndex(x => x.Key == Name);
+			if(Idx == -1)
+			{
+				Sections.Add(new KeyValuePair<string,string>(Name, Value));
+			}
+			else
+			{
+				Sections[Idx] = new KeyValuePair<string,string>(Name, Value);
+			}
+		}
+
+		/// <summary>
+		/// Parses a spec (clientspec, branchspec, changespec) from an array of lines
+		/// </summary>
+		/// <param name="Lines">Text split into separate lines</param>
+		/// <returns>Array of section names and values</returns>
+		public static P4Spec FromString(string Text)
+		{
+			P4Spec Spec = new P4Spec();
+
+			string[] Lines = Text.Split('\n');
+			for(int LineIdx = 0; LineIdx < Lines.Length; LineIdx++)
+			{
+				if(!String.IsNullOrWhiteSpace(Lines[LineIdx]) && !Lines[LineIdx].StartsWith("#"))
+				{
+					// Read the section name
+					int SeparatorIdx = Lines[LineIdx].IndexOf(':');
+					if(SeparatorIdx == -1 || !Char.IsLetter(Lines[LineIdx][0]))
+					{
+						throw new P4Exception("Invalid spec format at line {0}: \"{1}\"", LineIdx, Lines[LineIdx]);
+					}
+
+					// Get the section name
+					string SectionName = Lines[LineIdx].Substring(0, SeparatorIdx);
+
+					// Parse the section value
+					StringBuilder Value = new StringBuilder(Lines[LineIdx].Substring(SeparatorIdx + 1));
+					for(; LineIdx + 1 < Lines.Length; LineIdx++)
+					{
+						if(Lines[LineIdx + 1].Length == 0)
+						{
+							Value.AppendLine();
+						}
+						else if(Lines[LineIdx + 1][0] == '\t')
+						{
+							Value.AppendLine(Lines[LineIdx + 1].Substring(1));
+						}
+						else
+						{
+							break;
+						}
+					}
+					Spec.Sections.Add(new KeyValuePair<string,string>(SectionName, Value.ToString().TrimEnd()));
+				}
+			}
+
+			return Spec;
+		}
+
+		/// <summary>
+		/// Formats a P4 specification as a block of text
+		/// </summary>
+		/// <returns></returns>
+		public override string ToString()
+		{
+			StringBuilder Result = new StringBuilder();
+			foreach(KeyValuePair<string, string> Section in Sections)
+			{
+				if(Section.Value.Contains('\n'))
+				{
+					Result.AppendLine(Section.Key + ":\n\t" + Section.Value.Replace("\n", "\n\t"));
+				}
+				else
+				{
+					Result.AppendLine(Section.Key + ":\t" + Section.Value);
+				}
+				Result.AppendLine();
+			}
+			return Result.ToString();
+		}
+	}
+
 	public partial class CommandUtils
 	{
 		#region Environment Setup
@@ -1040,6 +1152,23 @@ namespace AutomationTool
             CheckP4Enabled();
             LogP4("shelve " + String.Format("-r -c {0} ", FromCL) + CommandLine, AllowSpew: AllowSpew);
         }
+
+		/// <summary>
+        /// Deletes shelved files from a changelist
+		/// </summary>
+        /// <param name="FromCL">Changelist to unshelve.</param>
+        /// <param name="CommandLine">Commandline for the command.</param>
+        public void DeleteShelvedFiles(int FromCL, bool AllowSpew = true)
+        {
+            CheckP4Enabled();
+
+			string Output;
+            if (!LogP4Output(out Output, String.Format("shelve -d -c {0}", FromCL), AllowSpew: AllowSpew) && !Output.StartsWith("No shelved files in changelist to delete."))
+			{
+				throw new P4Exception("Couldn't unshelve files: {0}", Output);
+			}
+        }
+
 		/// <summary>
 		/// Invokes p4 edit command.
 		/// </summary>
@@ -1385,6 +1514,44 @@ namespace AutomationTool
 			return CL;
 		}
 
+
+		/// <summary>
+		/// Updates a changelist with the given fields
+		/// </summary>
+		/// <param name="CL"></param>
+		/// <param name="NewOwner"></param>
+		/// <param name="NewDescription"></param>
+		/// <param name="SpewIsVerbose"></param>
+		public void UpdateChange(int CL, string NewOwner, string NewDescription, bool SpewIsVerbose = false)
+		{
+			CheckP4Enabled();
+
+			string CmdOutput;
+			if(!LogP4Output(out CmdOutput, String.Format("change -o {0}", CL), SpewIsVerbose: SpewIsVerbose))
+			{
+				throw new P4Exception("Couldn't describe changelist {0}", CL);
+			}
+
+			P4Spec Spec = P4Spec.FromString(CmdOutput);
+			if(NewOwner != null)
+			{
+				Spec.SetField("Client", NewOwner);
+			}
+			if(NewDescription != null)
+			{
+				Spec.SetField("Description", NewDescription);
+			}
+
+			if(!LogP4Output(out CmdOutput, "change -i", Input: Spec.ToString(), SpewIsVerbose: SpewIsVerbose))
+			{
+				throw new P4Exception("Failed to update spec for changelist {0}", CL);
+			}
+			if(!CmdOutput.TrimEnd().EndsWith(String.Format("Change {0} updated.", CL)))
+			{
+				throw new P4Exception("Unexpected output from p4 change -i: {0}", CmdOutput);
+			}
+		}
+
 		/// <summary>
 		/// Deletes the specified changelist.
 		/// </summary>
