Add Rocket to CIS. Publishing/chunking nodes only happen after labeling a promoted build, all other steps happen before.

diff --git a/Engine/Source/Programs/AutomationTool/Scripts/AutomationScripts.Automation.csproj b/Engine/Source/Programs/AutomationTool/Scripts/AutomationScripts.Automation.csproj
--- a/Engine/Source/Programs/AutomationTool/Scripts/AutomationScripts.Automation.csproj
+++ b/Engine/Source/Programs/AutomationTool/Scripts/AutomationScripts.Automation.csproj
@@ -73,6 +73,7 @@
     <Compile Include="FixupRedirects.Automation.cs" />
     <Compile Include="GenerateDSYM.Automation.cs" />
     <Compile Include="GUBP.Automation.cs" />
+    <Compile Include="RocketBuild.Automation.cs" />
     <Compile Include="IPhonePackager.Automation.cs" />
     <Compile Include="LauncherLocalization.Automation.cs" />
     <Compile Include="ListMobileDevices.Automation.cs">
diff --git a/Engine/Source/Programs/AutomationTool/Scripts/GUBP.Automation.cs b/Engine/Source/Programs/AutomationTool/Scripts/GUBP.Automation.cs
--- a/Engine/Source/Programs/AutomationTool/Scripts/GUBP.Automation.cs
+++ b/Engine/Source/Programs/AutomationTool/Scripts/GUBP.Automation.cs
@@ -198,6 +198,7 @@ public class GUBP : BuildCommand
 			public List<UnrealTargetPlatform> ExcludePlatformsForEditor = new List<UnrealTargetPlatform>();
 			public bool bNoAutomatedTesting = false;
 			public bool bNoDocumentation = false;
+			public bool bNoInstalledEngine = false;
 			public bool bMakeFormalBuildWithoutLabelPromotable = false;
 			public int QuantumOverride = 0;
         }
@@ -1405,7 +1406,7 @@ public class GUBP : BuildCommand
         {
 			Projects = new List<BranchInfo.BranchUProject>(InProjects);
 			AddDependency(ToolsNode.StaticGetFullName(InHostPlatform)); // for UnrealPak
-			AgentSharingGroup = "BuildInstall" + StaticGetHostPlatformSuffix(HostPlatform);
+            AgentSharingGroup = "FeaturePacks"  + StaticGetHostPlatformSuffix(InHostPlatform);
         }
 
 		public static string GetOutputFile(BranchInfo.BranchUProject Project)
@@ -1455,7 +1456,7 @@ public class GUBP : BuildCommand
 
 		public override int CISFrequencyQuantumShift(GUBP bp)
 		{
-			return base.CISFrequencyQuantumShift(bp) + 3;
+			return base.CISFrequencyQuantumShift(bp) + 2;
 		}
 		public override void DoBuild(GUBP bp)
         {
diff --git a/Engine/Source/Programs/AutomationTool/Scripts/RocketBuild.Automation.cs b/Engine/Source/Programs/AutomationTool/Scripts/RocketBuild.Automation.cs
new file mode 100644
--- /dev/null
+++ b/Engine/Source/Programs/AutomationTool/Scripts/RocketBuild.Automation.cs
@@ -0,0 +1,768 @@
+﻿// Copyright 1998-2015 Epic Games, Inc. All Rights Reserved.
+
+using System;
+using System.Collections.Generic;
+using System.Linq;
+using System.Text;
+using System.IO;
+using AutomationTool;
+using UnrealBuildTool;
+
+namespace Rocket
+{
+	public class RocketBuild : GUBP.GUBPNodeAdder
+	{
+		static readonly string[] CurrentTemplates = 
+		{
+			"TP_FirstPerson",
+			"TP_FirstPersonBP",
+			"TP_Flying",
+			"TP_FlyingBP",
+			"TP_Rolling",
+			"TP_RollingBP",
+			"TP_SideScroller",
+			"TP_SideScrollerBP",
+			"TP_ThirdPerson",
+			"TP_ThirdPersonBP",
+			"TP_TopDown",
+			"TP_TopDownBP",
+			"TP_TwinStick",
+			"TP_TwinStickBP",
+			"TP_Vehicle",
+			"TP_VehicleBP",
+			"TP_Puzzle",
+			"TP_PuzzleBP",
+			"TP_2DSideScroller",
+			"TP_2DSideScrollerBP",
+			"TP_VehicleAdv",
+			"TP_VehicleAdvBP",
+		};
+
+		static readonly string[] CurrentFeaturePacks = 
+		{
+			"FP_FirstPerson",
+			"FP_FirstPersonBP",
+			"TP_Flying",
+			"TP_FlyingBP",
+			"TP_Rolling",
+			"TP_RollingBP",
+			"TP_SideScroller",
+			"TP_SideScrollerBP",
+			"TP_ThirdPerson",
+			"TP_ThirdPersonBP",
+			"TP_TopDown",
+			"TP_TopDownBP",
+			"TP_TwinStick",
+			"TP_TwinStickBP",
+			"TP_Vehicle",
+			"TP_VehicleBP",
+			"TP_Puzzle",
+			"TP_PuzzleBP",
+			"TP_2DSideScroller",
+			"TP_2DSideScrollerBP",
+			"TP_VehicleAdv",
+			"TP_VehicleAdvBP",
+			"StarterContent",
+			"MobileStarterContent"
+		};
+
+		public RocketBuild()
+		{
+		}
+
+		public override void AddNodes(GUBP bp, UnrealTargetPlatform HostPlatform)
+		{
+			if(!bp.BranchOptions.bNoInstalledEngine)
+			{
+				// Get the output paths for the installed engine
+				string InstallDir = CommandUtils.CombinePaths(CommandUtils.CmdEnv.LocalRoot, "LocalBuilds", "Engine", CommandUtils.GetGenericPlatformName(HostPlatform));
+				string SymbolsDir = CommandUtils.CombinePaths(CommandUtils.CmdEnv.LocalRoot, "LocalBuilds", "Engine", CommandUtils.GetGenericPlatformName(HostPlatform) + "Symbols");
+
+				// Find all the target platforms for this host platform
+				List<UnrealTargetPlatform> TargetPlatforms = GetTargetPlatforms(bp, HostPlatform);
+
+				// Get a list of all the code target platforms
+				List<UnrealTargetPlatform> CodeTargetPlatforms = new List<UnrealTargetPlatform>();
+				foreach(UnrealTargetPlatform TargetPlatform in TargetPlatforms)
+				{
+					if(IsCodeTargetPlatform(HostPlatform, TargetPlatform))
+					{
+						CodeTargetPlatforms.Add(TargetPlatform);
+					}
+				}
+
+				// Add the aggregate node for the whole install, if there isn't one already
+				if(!bp.HasNode(RocketAggregateNode.StaticGetFullName()))
+				{
+					bp.AddNode(new RocketAggregateNode());
+				}
+
+				// Add a dependency on this platform finishing
+				RocketAggregateNode AggNode = (RocketAggregateNode)bp.FindNode(RocketAggregateNode.StaticGetFullName());
+
+				// Get file lists for each of the target platforms
+				foreach(UnrealTargetPlatform CodeTargetPlatform in CodeTargetPlatforms)
+				{
+					UnrealTargetPlatform SourceHostPlatform = GetSourceHostPlatform(bp, HostPlatform, CodeTargetPlatform);
+					if(!bp.HasNode(GetExternalFileListNode.StaticGetFullName(SourceHostPlatform, CodeTargetPlatform)))
+					{
+						bp.AddNode(new GetExternalFileListNode(SourceHostPlatform, CodeTargetPlatform));
+					}
+				}
+
+				// Copy to output folders
+				AggNode.AddDependency(bp.AddNode(new CopyInstalledEditorNode(bp, HostPlatform, InstallDir)));
+				AggNode.AddDependency(bp.AddNode(new CopyInstalledPlatformsNode(bp, HostPlatform, TargetPlatforms, CodeTargetPlatforms, InstallDir, CurrentFeaturePacks, CurrentTemplates)));
+				AggNode.AddDependency(bp.AddNode(new CopyInstalledDerivedDataNode(HostPlatform, GetCookPlatforms(HostPlatform, TargetPlatforms), CurrentFeaturePacks, InstallDir)));
+				AggNode.AddDependency(bp.AddNode(new CopyInstalledSymbolsNode(bp, HostPlatform, CodeTargetPlatforms, SymbolsDir)));
+
+				// Add the aggregate node for the entire install
+				AggNode.AddDependency(bp.AddNode(new BuildInstalledEngineNode(HostPlatform, InstallDir, SymbolsDir)));
+			}
+		}
+
+		public static List<UnrealTargetPlatform> GetTargetPlatforms(BuildCommand Cmd, UnrealTargetPlatform HostPlatform)
+		{
+			List<UnrealTargetPlatform> TargetPlatforms = new List<UnrealTargetPlatform>();
+			if(!Cmd.ParseParam("NoPlatforms"))
+			{
+				TargetPlatforms.Add(HostPlatform);
+
+				if(HostPlatform == UnrealTargetPlatform.Win64)
+				{
+					TargetPlatforms.Add(UnrealTargetPlatform.Win32);
+				}
+				if(!Cmd.ParseParam("NoAndroid"))
+				{
+					TargetPlatforms.Add(UnrealTargetPlatform.Android);
+				}
+				if(!Cmd.ParseParam("NoIOS"))
+				{
+					TargetPlatforms.Add(UnrealTargetPlatform.IOS);
+				}
+				if(!Cmd.ParseParam("NoLinux") && HostPlatform == UnrealTargetPlatform.Win64)
+				{
+					TargetPlatforms.Add(UnrealTargetPlatform.Linux);
+				}
+				if (!Cmd.ParseParam("NoHTML5") && HostPlatform == UnrealTargetPlatform.Win64)
+				{
+					TargetPlatforms.Add(UnrealTargetPlatform.HTML5);
+				}
+			}
+			return TargetPlatforms;
+		}
+
+		public static string GetCookPlatforms(UnrealTargetPlatform HostPlatform, IEnumerable<UnrealTargetPlatform> TargetPlatforms)
+		{
+			// Always include the editor platform for cooking
+			List<string> CookPlatforms = new List<string>();
+			CookPlatforms.Add(Platform.Platforms[HostPlatform].GetEditorCookPlatform());
+
+			// Add all the target platforms
+			foreach(UnrealTargetPlatform TargetPlatform in TargetPlatforms)
+			{
+				if(TargetPlatform == UnrealTargetPlatform.Android)
+				{
+					CookPlatforms.Add(Platform.Platforms[TargetPlatform].GetCookPlatform(false, false, "ATC"));
+				}
+				else
+				{
+					CookPlatforms.Add(Platform.Platforms[TargetPlatform].GetCookPlatform(false, false, ""));
+				}
+			}
+			return CommandUtils.CombineCommandletParams(CookPlatforms.Distinct().ToArray());
+		}
+
+		public static bool ShouldDoSeriousThingsLikeP4CheckinAndPostToMCP()
+		{
+			return CommandUtils.P4Enabled && CommandUtils.AllowSubmit && !GUBP.bPreflightBuild; // we don't do serious things in a preflight
+		}
+
+		public static UnrealTargetPlatform GetSourceHostPlatform(GUBP bp, UnrealTargetPlatform HostPlatform, UnrealTargetPlatform TargetPlatform)
+		{
+			if (TargetPlatform == UnrealTargetPlatform.Android && HostPlatform == UnrealTargetPlatform.Mac && bp.HostPlatforms.Contains(UnrealTargetPlatform.Win64))
+			{
+				return UnrealTargetPlatform.Win64;
+			}
+			if(TargetPlatform == UnrealTargetPlatform.IOS && HostPlatform == UnrealTargetPlatform.Win64 && bp.HostPlatforms.Contains(UnrealTargetPlatform.Mac))
+			{
+				return UnrealTargetPlatform.Mac;
+			}
+			return HostPlatform;
+		}
+
+		public static bool IsCodeTargetPlatform(UnrealTargetPlatform HostPlatform, UnrealTargetPlatform TargetPlatform)
+		{
+			if(TargetPlatform == UnrealTargetPlatform.Linux)
+			{
+				return false;
+			}
+			if(HostPlatform == UnrealTargetPlatform.Win64 && TargetPlatform == UnrealTargetPlatform.IOS)
+			{
+				return false;
+			}
+			return true;
+		}
+	}
+
+	public class CopySymbolsNode : GUBP.HostPlatformNode
+	{
+		public string OutputDirectory;
+		public bool bAddToBuildProducts;
+
+		public CopySymbolsNode(UnrealTargetPlatform HostPlatform, string InOutputDirectory, bool bInAddToBuildProducts) : base(HostPlatform)
+		{
+			OutputDirectory = InOutputDirectory;
+			bAddToBuildProducts = bInAddToBuildProducts;
+			AgentSharingGroup = "BuildMinimalInstall" + StaticGetHostPlatformSuffix(HostPlatform);
+		}
+
+		public override void DoBuild(GUBP bp)
+		{
+			BuildProducts = new List<string>();
+
+			// Make a lookup for all the known debug extensions, and filter all the dependency build products against that
+			HashSet<string> DebugExtensions = new HashSet<string>(Platform.Platforms.Values.SelectMany(x => x.GetDebugFileExtentions()).Distinct().ToArray(), StringComparer.InvariantCultureIgnoreCase);
+			foreach(string InputFileName in AllDependencyBuildProducts)
+			{
+				string Extension = Path.GetExtension(InputFileName);
+				if(DebugExtensions.Contains(Extension))
+				{
+					string OutputFileName = CommandUtils.MakeRerootedFilePath(InputFileName, CommandUtils.CmdEnv.LocalRoot, OutputDirectory);
+					CommandUtils.CopyFile(InputFileName, OutputFileName);
+
+					if(bAddToBuildProducts)
+					{
+						BuildProducts.Add(OutputFileName);
+					}
+				}
+			}
+
+			// Add a dummy build product if we're not adding all the files we copied
+			if(!bAddToBuildProducts)
+			{
+				SaveRecordOfSuccessAndAddToBuildProducts();
+			}
+		}
+	}
+
+	public class CopyInstalledSymbolsNode : CopySymbolsNode
+	{
+		public CopyInstalledSymbolsNode(GUBP bp, UnrealTargetPlatform HostPlatform, IEnumerable<UnrealTargetPlatform> TargetPlatforms, string OutputDirectory) 
+			: base(HostPlatform, OutputDirectory, false)
+		{
+			AddDependency(GUBP.ToolsForCompileNode.StaticGetFullName(HostPlatform));
+			AddDependency(GUBP.RootEditorNode.StaticGetFullName(HostPlatform));
+			AddDependency(GUBP.ToolsNode.StaticGetFullName(HostPlatform));
+
+			foreach(UnrealTargetPlatform TargetPlatform in TargetPlatforms)
+			{
+				UnrealTargetPlatform SourceHostPlatform = CopyInstalledPlatformsNode.GetSourceHostPlatform(bp, HostPlatform, TargetPlatform);
+				AddDependency(GUBP.GamePlatformMonolithicsNode.StaticGetFullName(SourceHostPlatform, bp.Branch.BaseEngineProject, TargetPlatform, Precompiled: true));
+			}
+		}
+
+		public static string StaticGetFullName(UnrealTargetPlatform HostPlatform)
+		{
+			return "CopyEngineSymbols" + StaticGetHostPlatformSuffix(HostPlatform);
+		}
+
+		public override string GetFullName()
+		{
+			return StaticGetFullName(HostPlatform);
+		}
+	}
+
+	public class CopyNode : GUBP.HostPlatformNode
+	{
+		public CopyNode(UnrealTargetPlatform InHostPlatform) : base(InHostPlatform)
+		{
+		}
+
+		static IEnumerable<string> FindBuildProducts(GUBP bp, string NodeName)
+		{
+			GUBP.GUBPNode Node = bp.FindNode(NodeName);
+			if (Node == null)
+			{
+				throw new AutomationException("Couldn't find node '{0}'", NodeName);
+			}
+			return Node.BuildProducts;
+		}
+
+		static IEnumerable<string> MakeRerootedFilePaths(IEnumerable<string> FileNames, string SourceDir, string TargetDir)
+		{
+			foreach(string FileName in FileNames)
+			{
+				yield return CommandUtils.MakeRerootedFilePath(Path.GetFullPath(FileName), SourceDir, TargetDir);
+			}
+		}
+
+		public static void AddFileToFilter(FileFilter Filter, string FileName, FileFilterType Type)
+		{
+			AddFileToFilter(Filter, FileName, CommandUtils.CmdEnv.LocalRoot, Type);
+		}
+
+		public static void AddFileToFilter(FileFilter Filter, string FileName, string BaseDirectoryName, FileFilterType Type)
+		{
+			Filter.AddRule("/" + CommandUtils.StripBaseDirectory(Path.GetFullPath(FileName), BaseDirectoryName), Type);
+		}
+
+		public static void AddFilesToFilter(FileFilter Filter, IEnumerable<string> FileNames, FileFilterType Type)
+		{
+			AddFilesToFilter(Filter, FileNames, CommandUtils.CmdEnv.LocalRoot, Type);
+		}
+
+		public static void AddFilesToFilter(FileFilter Filter, IEnumerable<string> FileNames, string BaseDirectoryName, FileFilterType Type)
+		{
+			foreach (string FileName in FileNames)
+			{
+				AddFileToFilter(Filter, FileName, BaseDirectoryName, Type);
+			}
+		}
+
+		public static void AddNodeToFilter(FileFilter Filter, GUBP bp, string NodeName, FileFilterType Type)
+		{
+			AddNodeToFilter(Filter, bp, NodeName, CommandUtils.CmdEnv.LocalRoot, Type);
+		}
+
+		public static void AddNodeToFilter(FileFilter Filter, GUBP bp, string NodeName, string BaseDirectoryName, FileFilterType Type)
+		{
+			GUBP.GUBPNode Node = bp.FindNode(NodeName);
+			if(Node == null)
+			{
+				throw new AutomationException("Couldn't find node '{0}'", NodeName);
+			}
+			AddFilesToFilter(Filter, Node.BuildProducts, BaseDirectoryName, Type);
+		}
+
+		public static void ExcludeConfidentialFolders(FileFilter Filter)
+		{
+			// Exclude platforms that can't be redistributed without an NDA
+			Filter.Exclude(".../PS4/...");
+			Filter.Exclude(".../XboxOne/...");
+
+			// Exclude standard folders
+			Filter.Exclude(".../EpicInternal/...");
+			Filter.Exclude(".../CarefullyRedist/...");
+			Filter.Exclude(".../NotForLicensees/...");
+			Filter.Exclude(".../NoRedist/...");
+		}
+
+		public static string[] CopyFolder(string SourceDir, string TargetDir, FileFilter Filter)
+		{
+			// Filter all the relative paths
+			CommandUtils.Log("Applying filter to {0}...", SourceDir);
+			string[] RelativePaths = Filter.ApplyToDirectory(SourceDir).ToArray();
+
+			// Find all the source and target filenames
+			string[] SourceFileNames = new string[RelativePaths.Length];
+			string[] TargetFileNames = new string[RelativePaths.Length];
+			for (int Idx = 0; Idx < RelativePaths.Length; Idx++)
+			{
+				SourceFileNames[Idx] = CommandUtils.CombinePaths(SourceDir, RelativePaths[Idx]);
+				TargetFileNames[Idx] = CommandUtils.CombinePaths(TargetDir, RelativePaths[Idx]);
+			}
+
+			// Copy them all over
+			CommandUtils.Log("Copying {0} files to {1}...", SourceFileNames.Length, TargetDir);
+			CommandUtils.ThreadedCopyFiles(SourceFileNames, TargetFileNames);
+			return TargetFileNames;
+		}
+	}
+
+	public class CopyInstalledEditorNode : CopyNode
+	{
+		string TargetDir;
+
+		public CopyInstalledEditorNode(GUBP bp, UnrealTargetPlatform InHostPlatform, string InTargetDir) 
+			: base(InHostPlatform)
+		{
+			TargetDir = InTargetDir;
+
+			AddDependency(GUBP.VersionFilesNode.StaticGetFullName());
+			AddDependency(GUBP.ToolsForCompileNode.StaticGetFullName(HostPlatform));
+			AddDependency(GUBP.RootEditorNode.StaticGetFullName(HostPlatform));
+			AddDependency(GUBP.ToolsNode.StaticGetFullName(HostPlatform));
+
+			// Add win64 tools on Mac, to get the win64 build of UBT and UAT
+			if (HostPlatform == UnrealTargetPlatform.Mac && bp.HostPlatforms.Contains(UnrealTargetPlatform.Win64))
+			{
+				AddDependency(GUBP.ToolsForCompileNode.StaticGetFullName(UnrealTargetPlatform.Win64));
+			}
+
+			AgentSharingGroup = "BuildMinimalInstall" + StaticGetHostPlatformSuffix(HostPlatform);
+		}
+
+		public override int CISFrequencyQuantumShift(GUBP bp)
+		{
+			return base.CISFrequencyQuantumShift(bp);
+		}
+
+		public static string StaticGetFullName(UnrealTargetPlatform InHostPlatform)
+		{
+			return "CopyEditor" + StaticGetHostPlatformSuffix(InHostPlatform);
+		}
+
+		public override string GetFullName()
+		{
+			return StaticGetFullName(HostPlatform);
+		}
+
+		public override void DoBuild(GUBP bp)
+		{
+			FileFilter Filter = new FileFilter();
+
+			// Include all the editor products
+			AddNodeToFilter(Filter, bp, GUBP.ToolsForCompileNode.StaticGetFullName(HostPlatform), FileFilterType.Include);
+			AddNodeToFilter(Filter, bp, GUBP.RootEditorNode.StaticGetFullName(HostPlatform), FileFilterType.Include);
+			AddNodeToFilter(Filter, bp, GUBP.ToolsNode.StaticGetFullName(HostPlatform), FileFilterType.Include);
+
+			// Include all the standard rules
+			string RulesFileName = CommandUtils.CombinePaths(CommandUtils.CmdEnv.LocalRoot, "Engine", "Build", "InstalledEngineFilters.ini");
+			Filter.AddRulesFromFile(RulesFileName, "CopyEditor", HostPlatform.ToString());
+
+			// Add the final exclusions for legal reasons.
+			ExcludeConfidentialFolders(Filter);
+
+			// Wipe the target directory and copy all the files over
+			CommandUtils.DeleteDirectoryContents(TargetDir);
+
+			// Copy everything over and add the result as build products
+			string[] TargetFileNames = CopyFolder(CommandUtils.CmdEnv.LocalRoot, TargetDir, Filter);
+			BuildProducts = new List<string>(TargetFileNames);
+
+			// Write the default command line to the root of the output directory
+			string CommandLinePath = Path.Combine(TargetDir, "UE4CommandLine.txt");
+			File.WriteAllText(CommandLinePath, "-engineinstalled -rocket");
+			AddBuildProduct(CommandLinePath);
+		}
+	}
+
+	public class CopyInstalledPlatformsNode : CopyNode
+	{
+		List<UnrealTargetPlatform> TargetPlatforms;
+		List<UnrealTargetPlatform> CodeTargetPlatforms;
+		string InstallDir;
+		string[] CurrentFeaturePacks;
+		string[] CurrentTemplates;
+
+		public CopyInstalledPlatformsNode(GUBP bp, UnrealTargetPlatform InHostPlatform, List<UnrealTargetPlatform> InTargetPlatforms, List<UnrealTargetPlatform> InCodeTargetPlatforms, string InInstallDir, string[] InCurrentFeaturePacks, string[] InCurrentTemplates)
+			: base(InHostPlatform)
+		{
+			TargetPlatforms = new List<UnrealTargetPlatform>(InTargetPlatforms);
+			CodeTargetPlatforms = new List<UnrealTargetPlatform>(InCodeTargetPlatforms);
+			InstallDir = InInstallDir;
+			CurrentFeaturePacks = InCurrentFeaturePacks;
+			CurrentTemplates = InCurrentTemplates;
+
+			// copy the editor
+			AddDependency(CopyInstalledEditorNode.StaticGetFullName(HostPlatform));
+
+			// add all the feature packs
+			AddDependency(GUBP.MakeFeaturePacksNode.StaticGetFullName(GUBP.MakeFeaturePacksNode.GetDefaultBuildPlatform(bp)));
+
+			// Get all the external file lists. Android on Mac uses the same files as Win64.
+			foreach(UnrealTargetPlatform CodeTargetPlatform in CodeTargetPlatforms)
+			{
+				UnrealTargetPlatform SourceHostPlatform = GetSourceHostPlatform(bp, HostPlatform, CodeTargetPlatform);
+				AddDependency(GUBP.GamePlatformMonolithicsNode.StaticGetFullName(SourceHostPlatform, bp.Branch.BaseEngineProject, CodeTargetPlatform, Precompiled: true));
+				AddDependency(GetExternalFileListNode.StaticGetFullName(SourceHostPlatform, CodeTargetPlatform));
+			}
+
+			// Add win64 tools on Mac, to get the win64 build of IPP
+			if (HostPlatform == UnrealTargetPlatform.Mac && bp.HostPlatforms.Contains(UnrealTargetPlatform.Win64))
+			{
+				AddDependency(GUBP.ToolsNode.StaticGetFullName(UnrealTargetPlatform.Win64));
+				AddDependency(GUBP.ToolsForCompileNode.StaticGetFullName(UnrealTargetPlatform.Win64));
+			}
+
+			AgentSharingGroup = "BuildInstall" + StaticGetHostPlatformSuffix(HostPlatform);
+		}
+
+		public static UnrealTargetPlatform GetSourceHostPlatform(GUBP bp, UnrealTargetPlatform HostPlatform, UnrealTargetPlatform TargetPlatform)
+		{
+			if(HostPlatform == UnrealTargetPlatform.Mac && TargetPlatform == UnrealTargetPlatform.Android && bp.HostPlatforms.Contains(UnrealTargetPlatform.Win64))
+			{
+				return UnrealTargetPlatform.Win64;
+			}
+			else
+			{
+				return HostPlatform;
+			}
+		}
+
+		public override int CISFrequencyQuantumShift(GUBP bp)
+		{
+			return base.CISFrequencyQuantumShift(bp) + 2;
+		}
+
+		public override float Priority()
+		{
+			return base.Priority() + 5.0f;
+		}
+
+		public override void DoBuild(GUBP bp)
+		{
+			FileFilter Filter = new FileFilter();
+
+			// Include the build dependencies for every code platform
+			foreach(UnrealTargetPlatform CodeTargetPlatform in CodeTargetPlatforms)
+			{
+				UnrealTargetPlatform SourceHostPlatform = GetSourceHostPlatform(bp, HostPlatform, CodeTargetPlatform);
+				string FileListPath = GetExternalFileListNode.StaticGetFileListPath(SourceHostPlatform, CodeTargetPlatform);
+				AddFilesToFilter(Filter, UnrealBuildTool.Utils.ReadClass<UnrealBuildTool.ExternalFileList>(FileListPath).FileNames, FileFilterType.Include);
+			}
+
+			// Add the monolithic binaries
+			foreach(UnrealTargetPlatform CodeTargetPlatform in CodeTargetPlatforms)
+			{
+				UnrealTargetPlatform SourceHostPlatform = GetSourceHostPlatform(bp, HostPlatform, CodeTargetPlatform);
+				AddNodeToFilter(Filter, bp, GUBP.GamePlatformMonolithicsNode.StaticGetFullName(SourceHostPlatform, bp.Branch.BaseEngineProject, CodeTargetPlatform, Precompiled: true), FileFilterType.Include);
+			}
+
+			// Include the feature packs
+			foreach(string CurrentFeaturePack in CurrentFeaturePacks)
+			{
+				BranchInfo.BranchUProject Project = bp.Branch.FindGameChecked(CurrentFeaturePack);
+				AddFileToFilter(Filter, GUBP.MakeFeaturePacksNode.GetOutputFile(Project), FileFilterType.Include);
+			}
+
+			// Include all the templates
+			foreach (string Template in CurrentTemplates)
+			{
+				BranchInfo.BranchUProject Project = bp.Branch.FindGameChecked(Template);
+				string TemplateRelativeDirectory = "/" + Utils.StripBaseDirectory(Path.GetDirectoryName(Project.FilePath), CommandUtils.CmdEnv.LocalRoot) + "/";
+				Filter.Include(TemplateRelativeDirectory + "/...");
+				Filter.Exclude(TemplateRelativeDirectory + "/.../Binaries/...");
+				Filter.Exclude(TemplateRelativeDirectory + "/.../Intermediate/...");
+			}
+
+			// Custom rules for all platforms
+			string RulesFileName = CommandUtils.CombinePaths(CommandUtils.CmdEnv.LocalRoot, "Engine", "Build", "InstalledEngineFilters.ini");
+			Filter.AddRulesFromFile(RulesFileName, "CopyTargetPlatforms", HostPlatform.ToString());
+
+			// Custom rules for each target platform
+			foreach(UnrealTargetPlatform TargetPlaform in TargetPlatforms)
+			{
+				Filter.AddRulesFromFile(RulesFileName, String.Format("CopyTargetPlatform.{0}", TargetPlaform.ToString()), HostPlatform.ToString());
+			}
+
+			// Exclude any the files which have already been added by the CopyInstalledEditor node
+			AddNodeToFilter(Filter, bp, CopyInstalledEditorNode.StaticGetFullName(HostPlatform), FileFilterType.Exclude);
+
+			// Add the final exclusions for legal reasons.
+			ExcludeConfidentialFolders(Filter);
+
+			// Copy all the files over
+			string[] TargetFileNames = CopyFolder(CommandUtils.CmdEnv.LocalRoot, InstallDir, Filter);
+			BuildProducts = new List<string>(TargetFileNames);
+		}
+
+		public override string RootIfAnyForTempStorage() // this saves a bit of path space for temp storage
+		{
+			return InstallDir;
+		}
+
+		public static string StaticGetFullName(UnrealTargetPlatform InHostPlatform)
+		{
+			return "CopyPlatforms" + StaticGetHostPlatformSuffix(InHostPlatform);
+		}
+
+		public override string GetFullName()
+		{
+			return StaticGetFullName(HostPlatform);
+		}
+	}
+
+	class GetExternalFileListNode : GUBP.HostPlatformNode
+	{
+		UnrealTargetPlatform TargetPlatform;
+
+		public GetExternalFileListNode(UnrealTargetPlatform InHostPlatform, UnrealTargetPlatform InTargetPlatform)
+			: base(InHostPlatform)
+		{
+			TargetPlatform = InTargetPlatform;
+			AgentSharingGroup = "UE4_" + TargetPlatform.ToString() + "_Mono" + StaticGetHostPlatformSuffix(HostPlatform);
+		}
+		public static string StaticGetFullName(UnrealTargetPlatform HostPlatform, UnrealTargetPlatform TargetPlatform)
+		{
+            return "UE4_" + TargetPlatform.ToString() + "_GetExternalFiles" + StaticGetHostPlatformSuffix(HostPlatform);
+		}
+		public override string GetFullName()
+		{
+			return StaticGetFullName(HostPlatform, TargetPlatform);
+		}
+		public override void DoBuild(GUBP bp)
+		{
+			// Prepare an agenda of targets that need to be buildable under Rocket.
+			UE4Build.BuildAgenda Agenda = new UE4Build.BuildAgenda();
+
+			// Add the game dependencies for UE4Game on all the target platforms. We don't link libraries with their dependencies, so we can use UE4Game and -buildrocket 
+			// to pick up all the default plugins.
+			Agenda.AddTargets(new string[] { "UE4Game" }, TargetPlatform, UnrealTargetConfiguration.Development, null, false, false, false, "-precompilemodules");
+			Agenda.AddTargets(new string[] { "UE4Game" }, TargetPlatform, UnrealTargetConfiguration.Shipping, null, false, false, false, "-precompilemodules");
+
+			// Read the file list
+			string FileListPath = new UE4Build(bp).GenerateExternalFileList(Agenda);
+			UnrealBuildTool.ExternalFileList FileList = UnrealBuildTool.Utils.ReadClass<UnrealBuildTool.ExternalFileList>(FileListPath);
+
+			// Make all the paths relative to the root
+			string FilterPrefix = CommandUtils.CombinePaths(PathSeparator.Slash, CommandUtils.CmdEnv.LocalRoot).TrimEnd('/') + "/";
+			for(int Idx = 0; Idx < FileList.FileNames.Count; Idx++)
+			{
+				if(FileList.FileNames[Idx].StartsWith(FilterPrefix, StringComparison.InvariantCultureIgnoreCase))
+				{
+					FileList.FileNames[Idx] = FileList.FileNames[Idx].Substring(FilterPrefix.Length);
+				}
+				else
+				{
+					CommandUtils.LogError("Referenced external file is not under local root: {0}", FileList.FileNames[Idx]);
+				}
+			}
+
+			// Write the resulting file list out to disk
+			string OutputFileListPath = StaticGetFileListPath(HostPlatform, TargetPlatform);
+			UnrealBuildTool.Utils.WriteClass<UnrealBuildTool.ExternalFileList>(FileList, OutputFileListPath, "");
+
+			// Add it to the build products
+			BuildProducts = new List<string>();
+			AddBuildProduct(OutputFileListPath);
+		}
+		public static string StaticGetFileListPath(UnrealTargetPlatform HostPlatform, UnrealTargetPlatform TargetPlatform)
+		{
+			return CommandUtils.CombinePaths(CommandUtils.CmdEnv.LocalRoot, "Engine/Saved/ExternalFileList/" + TargetPlatform.ToString() + StaticGetHostPlatformSuffix(HostPlatform) + ".xml");
+		}
+	}
+
+	public class CopyInstalledDerivedDataNode : GUBP.HostPlatformNode
+	{
+		string TargetPlatforms;
+		string InstallDir;
+		string[] ProjectNames;
+
+		public CopyInstalledDerivedDataNode(UnrealTargetPlatform InHostPlatform, string InTargetPlatforms, string[] InProjectNames, string InInstallDir)
+			: base(InHostPlatform)
+		{
+			TargetPlatforms = InTargetPlatforms;
+			ProjectNames = InProjectNames;
+			InstallDir = InInstallDir;
+
+			AddDependency(CopyInstalledEditorNode.StaticGetFullName(HostPlatform));
+			AgentSharingGroup = "BuildMinimalInstall" + StaticGetHostPlatformSuffix(InHostPlatform);
+		}
+
+		public static string StaticGetFullName(UnrealTargetPlatform InHostPlatform)
+		{
+            return "CopyDerivedData" + StaticGetHostPlatformSuffix(InHostPlatform);
+		}
+
+		public override int CISFrequencyQuantumShift(GUBP bp)
+		{
+			return base.CISFrequencyQuantumShift(bp) + 2;
+		}
+
+		public override string GetFullName()
+		{
+			return StaticGetFullName(HostPlatform);
+		}
+
+		public override void DoBuild(GUBP bp)
+		{
+			// Remove any existing DDP file
+			string OutputFileName = CommandUtils.CombinePaths(InstallDir, "Engine/DerivedDataCache/Compressed.ddp");
+			if(CommandUtils.FileExists(OutputFileName))
+			{
+				CommandUtils.DeleteFile(OutputFileName);
+			}
+
+			// Enumerate all the files in the install directory. We'll remove everything in this list that's not a build product once we're done.
+			string[] InitialFiles = Directory.EnumerateFiles(InstallDir, "*", SearchOption.AllDirectories).ToArray();
+
+			// Generate DDC for all the non-code projects. We don't necessarily have editor DLLs for the code projects, but they should be the same as their blueprint counterparts.
+			foreach(string ProjectName in ProjectNames)
+			{
+                BranchInfo.BranchUProject Project = bp.Branch.FindGameChecked(ProjectName);
+				if(!Project.Properties.bIsCodeBasedProject)
+				{
+					string InstalledProjectPath = CommandUtils.MakeRerootedFilePath(Project.FilePath, CommandUtils.CmdEnv.LocalRoot, InstallDir);
+					if(!CommandUtils.FileExists(InstalledProjectPath))
+					{
+						CommandUtils.CopyFile(Project.FilePath, InstalledProjectPath);
+					}
+					CommandUtils.Log("Running RocketEditor to generate engine DDC data for {0} on {1}", InstalledProjectPath, TargetPlatforms);
+					CommandUtils.DDCCommandlet(InstalledProjectPath, CommandUtils.GetEditorCommandletExe(InstallDir, HostPlatform), null, TargetPlatforms, "-fill -DDC=CreateInstalledEnginePak -MergePaks");
+				}
+			}
+
+			// Create the list of DDC build products 
+			BuildProducts = new List<string>();
+			BuildProducts.Add(OutputFileName);
+
+			// Remove all the files all the temporary files that can be removed
+			SortedSet<string> TemporaryFiles = new SortedSet<string>(Directory.EnumerateFiles(InstallDir, "*", SearchOption.AllDirectories));
+			TemporaryFiles.ExceptWith(InitialFiles);
+			TemporaryFiles.ExceptWith(BuildProducts.Select(x => Path.GetFullPath(x)));
+			foreach(string TemporaryFile in TemporaryFiles)
+			{
+				CommandUtils.DeleteFile(TemporaryFile);
+			}
+
+			// Remove any empty directories left over
+			foreach(string DirectoryName in Directory.EnumerateDirectories(InstallDir, "*", SearchOption.AllDirectories).OrderByDescending(x => x.Length))
+			{
+				if(!Directory.EnumerateFileSystemEntries(DirectoryName).GetEnumerator().MoveNext())
+				{
+					try { Directory.Delete(DirectoryName); } catch(Exception){ }
+				}
+			}
+		}
+
+		public override float Priority()
+		{
+			return base.Priority() + 55.0f;
+		}
+	}
+
+	public class BuildInstalledEngineNode : GUBP.HostPlatformAggregateNode
+	{
+		public readonly string InstallDir;
+		public readonly string SymbolsDir;
+
+		public BuildInstalledEngineNode(UnrealTargetPlatform InHostPlatform, string InInstallDir, string InSymbolsDir)
+			: base(InHostPlatform)
+		{
+			InstallDir = InInstallDir;
+			SymbolsDir = InSymbolsDir;
+
+			AddDependency(CopyInstalledEditorNode.StaticGetFullName(InHostPlatform));
+			AddDependency(CopyInstalledPlatformsNode.StaticGetFullName(InHostPlatform));
+			AddDependency(CopyInstalledDerivedDataNode.StaticGetFullName(InHostPlatform));
+			AddDependency(CopyInstalledSymbolsNode.StaticGetFullName(InHostPlatform));
+		}
+
+		public static string StaticGetFullName(UnrealTargetPlatform InHostPlatform)
+		{
+			return "BuildInstalledEngine" + StaticGetHostPlatformSuffix(InHostPlatform);
+		}
+
+		public override string GetFullName()
+		{
+			return StaticGetFullName(HostPlatform);
+		}
+	}
+
+	public class RocketAggregateNode : GUBP.AggregateNode
+	{
+		public RocketAggregateNode()
+		{
+		}
+		public static string StaticGetFullName()
+		{
+			return "Rocket_Aggregate";
+		}
+		public override string GetFullName()
+		{
+			return StaticGetFullName();
+		}
+	}
+}
