Add EnvVarsToXML as a standard UBT build product in UE4Build, since it's referenced by the project directly, and remove it from building as a separate target in BuildCommonTools. Also add EnvVarsToXML.exe.config to source control, since it doesn't trigger a rebuild if missing (but will still be added as a build product).

diff --git a/Engine/Source/Programs/AutomationTool/Scripts/BuildCommonTools.Automation.cs b/Engine/Source/Programs/AutomationTool/Scripts/BuildCommonTools.Automation.cs
--- a/Engine/Source/Programs/AutomationTool/Scripts/BuildCommonTools.Automation.cs
+++ b/Engine/Source/Programs/AutomationTool/Scripts/BuildCommonTools.Automation.cs
@@ -89,7 +89,6 @@ public class BuildCommonTools : BuildCommand
 		Agenda.DotNetProjects.Add(@"Engine/Source/Programs/DotNETCommon/DotNETUtilities/DotNETUtilities.csproj");
 		Agenda.DotNetProjects.Add(@"Engine/Source/Programs/RPCUtility/RPCUtility.csproj");
 		Agenda.DotNetProjects.Add(@"Engine/Source/Programs/UnrealControls/UnrealControls.csproj");
-		Agenda.DotNetProjects.Add(@"Engine/Source/Programs/EnvVarsToXML/EnvVarsToXML/EnvVarsToXML.csproj");
 
 		// Windows binaries
 		if(Platforms.Contains(UnrealBuildTool.UnrealTargetPlatform.Win64))
diff --git a/Engine/Source/Programs/AutomationTool/UE4Build.cs b/Engine/Source/Programs/AutomationTool/UE4Build.cs
--- a/Engine/Source/Programs/AutomationTool/UE4Build.cs
+++ b/Engine/Source/Programs/AutomationTool/UE4Build.cs
@@ -1405,6 +1405,8 @@ namespace AutomationTool
 					{
 						"UnrealBuildTool.exe",
 						"UnrealBuildTool.exe.config",
+						"EnvVarsToXML.exe",
+						"EnvVarsToXML.exe.config"
 					});
 
 			foreach (var UBTFile in UBTFiles)
