Allow for case-insensitive platform names on command line to BuildCommonTools.

diff --git a/Engine/Source/Programs/AutomationTool/Scripts/BuildCommonTools.Automation.cs b/Engine/Source/Programs/AutomationTool/Scripts/BuildCommonTools.Automation.cs
--- a/Engine/Source/Programs/AutomationTool/Scripts/BuildCommonTools.Automation.cs
+++ b/Engine/Source/Programs/AutomationTool/Scripts/BuildCommonTools.Automation.cs
@@ -25,7 +25,7 @@ public class BuildCommonTools : BuildCommand
 		foreach(string PlatformName in PlatformNames)
 		{
 			UnrealBuildTool.UnrealTargetPlatform Platform;
-			if(!UnrealBuildTool.UnrealTargetPlatform.TryParse(PlatformName, out Platform))
+			if(!UnrealBuildTool.UnrealTargetPlatform.TryParse(PlatformName, true, out Platform))
 			{
 				throw new AutomationException("Unknown platform specified on command line - '{0}' - valid platforms are {1}", PlatformName, String.Join("/", Enum.GetNames(typeof(UnrealBuildTool.UnrealTargetPlatform))));
 			}
