build script, up ectool timeout to cover for server perf issues

diff --git a/Engine/Source/Programs/AutomationTool/Scripts/GUBP.Automation.cs b/Engine/Source/Programs/AutomationTool/Scripts/GUBP.Automation.cs
--- a/Engine/Source/Programs/AutomationTool/Scripts/GUBP.Automation.cs
+++ b/Engine/Source/Programs/AutomationTool/Scripts/GUBP.Automation.cs
@@ -4068,7 +4068,7 @@ public class GUBP : BuildCommand
             {
                 Opts = (Opts & ~ERunOptions.AllowSpew) | ERunOptions.NoLoggingOfRunCommand;
             }
-            return RunAndLog("ectool", Args, Options: Opts);
+            return RunAndLog("ectool --timeout 900", Args, Options: Opts);
         }
     }
     void WriteECPerl(List<string> Args)
