Strip private symbols from PDB files as part of Rocket build again.

diff --git a/Engine/Build/InstalledEngineFilters.ini b/Engine/Build/InstalledEngineFilters.ini
--- a/Engine/Build/InstalledEngineFilters.ini
+++ b/Engine/Build/InstalledEngineFilters.ini
@@ -97,6 +97,12 @@
 
 {Mac} /Engine/Source/UE4Game*.cs
 
+[StripEditor]
+{Win64} /Engine/Binaries/.../*.pdb
+{Win64} /Engine/Plugins/.../*.pdb
+{Win64} -/Engine/Binaries/DotNET/...
+{Win64} -/Engine/.../UE4Editor*
+
 [CopyTargetPlatforms]
 ; API docs
 /Engine/Documentation/CHM/API.chm
@@ -154,6 +160,12 @@
 /Engine/Binaries/ThirdParty/PhysX/PhysX-3.3/Win32/VS2013/nvToolsExt32_1.dll
 /Engine/Binaries/ThirdParty/Steamworks/Steamv132/Win32/*.dll
 
+[StripTargetPlatform.Win32]
+/Engine/Binaries/.../*.pdb
+
+[StripTargetPlatform.Win64]
+/Engine/Binaries/.../*.pdb
+
 [CopyTargetPlatform.Linux]
 ; Content-only: include pre-compiled executables manually
 /Engine/Binaries/Linux/CrashReportClient
diff --git a/Engine/Source/Programs/AutomationTool/FileFilter.cs b/Engine/Source/Programs/AutomationTool/FileFilter.cs
--- a/Engine/Source/Programs/AutomationTool/FileFilter.cs
+++ b/Engine/Source/Programs/AutomationTool/FileFilter.cs
@@ -341,6 +341,17 @@ public class FileFilter
 		return FileNames.Where(x => Matches(x));
 	}
 
+	/// <summary>
+	/// Applies the filter to each element in a sequence, and returns the list of files that match
+	/// </summary>
+	/// <param name="FileNames">List of filenames</param>
+	/// <returns>List of filenames which match the filter</returns>
+	public IEnumerable<string> ApplyTo(string BaseDirectory, IEnumerable<string> FileNames)
+	{
+		string FullBaseDirectory = Path.GetFullPath(BaseDirectory);
+		return FileNames.Where(x => Matches(AutomationTool.CommandUtils.StripBaseDirectory(Path.GetFullPath(x), FullBaseDirectory)));
+	}
+
 	/// <summary>
 	/// Finds a list of files within a given directory which match the filter.
 	/// </summary>
diff --git a/Engine/Source/Programs/AutomationTool/Scripts/RocketBuild.Automation.cs b/Engine/Source/Programs/AutomationTool/Scripts/RocketBuild.Automation.cs
--- a/Engine/Source/Programs/AutomationTool/Scripts/RocketBuild.Automation.cs
+++ b/Engine/Source/Programs/AutomationTool/Scripts/RocketBuild.Automation.cs
@@ -121,10 +121,10 @@ namespace Rocket
 			}
 		}
 
-		public static List<UnrealTargetPlatform> GetTargetPlatforms(BuildCommand Cmd, UnrealTargetPlatform HostPlatform)
+		public static List<UnrealTargetPlatform> GetTargetPlatforms(GUBP bp, UnrealTargetPlatform HostPlatform)
 		{
 			List<UnrealTargetPlatform> TargetPlatforms = new List<UnrealTargetPlatform>();
-			if(!Cmd.ParseParam("NoPlatforms"))
+			if(!bp.ParseParam("NoPlatforms"))
 			{
 				TargetPlatforms.Add(HostPlatform);
 
@@ -132,22 +132,24 @@ namespace Rocket
 				{
 					TargetPlatforms.Add(UnrealTargetPlatform.Win32);
 				}
-				if(!Cmd.ParseParam("NoAndroid"))
+				if(!bp.ParseParam("NoAndroid"))
 				{
 					TargetPlatforms.Add(UnrealTargetPlatform.Android);
 				}
-				if(!Cmd.ParseParam("NoIOS"))
+				if(!bp.ParseParam("NoIOS"))
 				{
 					TargetPlatforms.Add(UnrealTargetPlatform.IOS);
 				}
-				if(!Cmd.ParseParam("NoLinux") && HostPlatform == UnrealTargetPlatform.Win64)
+				if(!bp.ParseParam("NoLinux") && HostPlatform == UnrealTargetPlatform.Win64)
 				{
 					TargetPlatforms.Add(UnrealTargetPlatform.Linux);
 				}
-				if (!Cmd.ParseParam("NoHTML5") && HostPlatform == UnrealTargetPlatform.Win64)
+				if(!bp.ParseParam("NoHTML5") && HostPlatform == UnrealTargetPlatform.Win64)
 				{
 					TargetPlatforms.Add(UnrealTargetPlatform.HTML5);
 				}
+
+				TargetPlatforms.RemoveAll(x => !bp.ActivePlatforms.Contains(x));
 			}
 			return TargetPlatforms;
 		}
@@ -368,6 +370,28 @@ namespace Rocket
 			CommandUtils.ThreadedCopyFiles(SourceFileNames, TargetFileNames);
 			return TargetFileNames;
 		}
+
+		public static void StripSymbols(UnrealTargetPlatform TargetPlatform, IEnumerable<string> FileNames)
+		{
+			if(TargetPlatform == UnrealTargetPlatform.Win32 || TargetPlatform == UnrealTargetPlatform.Win64)
+			{
+				string PDBCopyPath = Path.Combine(Environment.GetFolderPath(Environment.SpecialFolder.ProgramFilesX86), "MSBuild", "Microsoft", "VisualStudio", "v12.0", "AppxPackage", "PDBCopy.exe");
+				foreach(string FileName in FileNames)
+				{
+					string TempFileName = FileName + ".stripped";
+					CommandUtils.Run(PDBCopyPath, String.Format("{0} {1} -p", CommandUtils.MakePathSafeToUseWithCommandLine(FileName), CommandUtils.MakePathSafeToUseWithCommandLine(TempFileName)));
+					CommandUtils.DeleteFile(FileName);
+					CommandUtils.RenameFile(TempFileName, FileName);
+				}
+			}
+			else
+			{
+				if(FileNames.Any())
+				{
+					throw new AutomationException("Symbol stripping has not been implemented for {0}", TargetPlatform.ToString());
+				}
+			}
+		}
 	}
 
 	public class CopyInstalledEditorNode : CopyNode
@@ -431,6 +455,11 @@ namespace Rocket
 			string[] TargetFileNames = CopyFolder(CommandUtils.CmdEnv.LocalRoot, TargetDir, Filter);
 			BuildProducts = new List<string>(TargetFileNames);
 
+			// Strip symbols from all the debug files
+			FileFilter StripFilter = new FileFilter();
+			StripFilter.AddRulesFromFile(RulesFileName, "StripEditor", HostPlatform.ToString());
+			StripSymbols(HostPlatform, StripFilter.ApplyTo(TargetDir, TargetFileNames));
+
 			// Write the default command line to the root of the output directory
 			string CommandLinePath = Path.Combine(TargetDir, "UE4CommandLine.txt");
 			File.WriteAllText(CommandLinePath, "-engineinstalled -rocket");
@@ -544,7 +573,8 @@ namespace Rocket
 			// Custom rules for each target platform
 			foreach(UnrealTargetPlatform TargetPlaform in TargetPlatforms)
 			{
-				Filter.AddRulesFromFile(RulesFileName, String.Format("CopyTargetPlatform.{0}", TargetPlaform.ToString()), HostPlatform.ToString());
+				string SectionName = String.Format("CopyTargetPlatform.{0}", TargetPlaform.ToString());
+				Filter.AddRulesFromFile(RulesFileName, SectionName, HostPlatform.ToString());
 			}
 
 			// Exclude any the files which have already been added by the CopyInstalledEditor node
@@ -556,6 +586,14 @@ namespace Rocket
 			// Copy all the files over
 			string[] TargetFileNames = CopyFolder(CommandUtils.CmdEnv.LocalRoot, InstallDir, Filter);
 			BuildProducts = new List<string>(TargetFileNames);
+
+			// Strip symbols from any of the build products
+			foreach(UnrealTargetPlatform TargetPlatform in TargetPlatforms)
+			{
+				FileFilter StripFilter = new FileFilter();
+				StripFilter.AddRulesFromFile(RulesFileName, String.Format("StripTargetPlatform.{0}", TargetPlatform.ToString()), HostPlatform.ToString());
+				StripSymbols(TargetPlatform, StripFilter.ApplyTo(InstallDir, TargetFileNames));
+			}
 		}
 
 		public override string RootIfAnyForTempStorage() // this saves a bit of path space for temp storage
