UEPLAT-4 - support iterativedeploy on Android
exclude copying movie uassets or umaps (they will be cooked and copied to staging)
#ue4
#android

diff --git a/Engine/Source/Programs/AutomationTool/Android/AndroidPlatform.Automation.cs b/Engine/Source/Programs/AutomationTool/Android/AndroidPlatform.Automation.cs
--- a/Engine/Source/Programs/AutomationTool/Android/AndroidPlatform.Automation.cs
+++ b/Engine/Source/Programs/AutomationTool/Android/AndroidPlatform.Automation.cs
@@ -434,6 +434,47 @@ public class AndroidPlatform : Platform
 	}
 	*/
 
+	public override bool RetrieveDeployedManifests(ProjectParams Params, DeploymentContext SC, out List<string> UFSManifests, out List<string> NonUFSManifests)
+	{
+		UFSManifests = null;
+		NonUFSManifests = null;
+
+		// Query the storage path from the device
+		string DeviceStorageQueryCommand = GetStorageQueryCommand();
+		ProcessResult StorageResult = RunAdbCommand(Params, DeviceStorageQueryCommand, null, ERunOptions.AppMustExist);
+		String StorageLocation = StorageResult.Output.Trim();
+		string RemoteDir = StorageLocation + "/UE4Game/" + Params.ShortProjectName;
+
+		// Note: appends the device name to make the filename unique; these files will be deleted later during delta manifest generation
+
+		// Try retrieving the UFS files manifest files from the device
+		string UFSManifestFileName = CombinePaths(SC.StageDirectory, DeploymentContext.UFSDeployedManifestFileName + "_" + Params.Device);
+		ProcessResult UFSResult = RunAdbCommand(Params, " pull " + RemoteDir + "/" + DeploymentContext.UFSDeployedManifestFileName + " \"" + UFSManifestFileName + "\"", null, ERunOptions.AppMustExist);
+		if (!UFSResult.Output.Contains("bytes"))
+		{
+			return false;
+		}
+
+		// Try retrieving the non UFS files manifest files from the device
+		string NonUFSManifestFileName = CombinePaths(SC.StageDirectory, DeploymentContext.NonUFSDeployedManifestFileName + "_" + Params.Device);
+		ProcessResult NonUFSResult = RunAdbCommand(Params, " pull " + RemoteDir + "/" + DeploymentContext.NonUFSDeployedManifestFileName + " \"" + NonUFSManifestFileName + "\"", null, ERunOptions.AppMustExist);
+		if (!NonUFSResult.Output.Contains("bytes"))
+		{
+			// Did not retrieve both so delete one we did retrieve
+			File.Delete(UFSManifestFileName);
+
+			return false;
+		}
+
+		// Return the manifest files
+		UFSManifests = new List<string>();
+		UFSManifests.Add(UFSManifestFileName);
+		NonUFSManifests = new List<string>();
+		NonUFSManifests.Add(NonUFSManifestFileName);
+
+		return true;
+	}
+
 	internal class LongestFirst : IComparer<string>
 	{
 		public int Compare(string a, string b)
@@ -462,120 +503,187 @@ public class AndroidPlatform : Platform
 		// now we can use the apk to get more info
 		string PackageName = GetPackageInfo(ApkName, false);
 
-		// try uninstalling an old app with the same identifier.
-		int SuccessCode = 0;
-		string UninstallCommandline = "uninstall " + PackageName;
-		RunAndLogAdbCommand(Params, UninstallCommandline, out SuccessCode);
-
-		// install the apk
-		string InstallCommandline = "install \"" + ApkName + "\"";
-		string InstallOutput = RunAndLogAdbCommand(Params, InstallCommandline, out SuccessCode);
-		int FailureIndex = InstallOutput.IndexOf("Failure"); 
-
-		// adb install doesn't always return an error code on failure, and instead prints "Failure", followed by an error code.
-		if (SuccessCode != 0 || FailureIndex != -1)
+		// Setup the OBB name and add the storage path (queried from the device) to it
+		string DeviceStorageQueryCommand = GetStorageQueryCommand();
+		ProcessResult Result = RunAdbCommand(Params, DeviceStorageQueryCommand, null, ERunOptions.AppMustExist);
+		String StorageLocation = Result.Output.Trim();
+		string DeviceObbName = StorageLocation + "/" + GetDeviceObbName(ApkName);
+		string RemoteDir = StorageLocation + "/UE4Game/" + Params.ShortProjectName;
+
+		// determine if APK out of date
+		string APKLastUpdateTime = new FileInfo(ApkName).LastWriteTime.ToString();
+		bool bNeedAPKInstall = true;
+		if (Params.IterativeDeploy)
 		{
-			string ErrorMessage = String.Format("Installation of apk '{0}' failed", ApkName);
-			if (FailureIndex != -1)
+			// Check for apk installed with this package name on the device
+			ProcessResult InstalledResult = RunAdbCommand(Params, "shell pm list packages " + PackageName, null, ERunOptions.AppMustExist);
+			if (InstalledResult.Output.Contains(PackageName))
 			{
-				string FailureString = InstallOutput.Substring(FailureIndex + 7).Trim();
-				if (FailureString != "")
+				// See if apk is up to date on device
+				InstalledResult = RunAdbCommand(Params, "shell cat " + RemoteDir + "/APKFileStamp.txt", null, ERunOptions.AppMustExist);
+				if (InstalledResult.Output.StartsWith("APK: "))
 				{
-					ErrorMessage += ": " + FailureString;
+					if (InstalledResult.Output.Substring(5).Trim() == APKLastUpdateTime)
+						bNeedAPKInstall = false;
 				}
 			}
+		}
+
+		// install new APK if needed
+		if (bNeedAPKInstall)
+		{
+			// try uninstalling an old app with the same identifier.
+			int SuccessCode = 0;
+			string UninstallCommandline = "uninstall " + PackageName;
+			RunAndLogAdbCommand(Params, UninstallCommandline, out SuccessCode);
+
+			// install the apk
+			string InstallCommandline = "install \"" + ApkName + "\"";
+			string InstallOutput = RunAndLogAdbCommand(Params, InstallCommandline, out SuccessCode);
+			int FailureIndex = InstallOutput.IndexOf("Failure"); 
+
+			// adb install doesn't always return an error code on failure, and instead prints "Failure", followed by an error code.
+			if (SuccessCode != 0 || FailureIndex != -1)
+			{
+				string ErrorMessage = String.Format("Installation of apk '{0}' failed", ApkName);
+				if (FailureIndex != -1)
+				{
+					string FailureString = InstallOutput.Substring(FailureIndex + 7).Trim();
+					if (FailureString != "")
+					{
+						ErrorMessage += ": " + FailureString;
+					}
+				}
 
-			ErrorReporter.Error(ErrorMessage, (int)ErrorCodes.Error_AppInstallFailed);
-			throw new AutomationException(ErrorMessage);
+				ErrorReporter.Error(ErrorMessage, (int)ErrorCodes.Error_AppInstallFailed);
+				throw new AutomationException(ErrorMessage);
+			}
 		}
-		
+ 
 		// update the ue4commandline.txt
 		// update and deploy ue4commandline.txt
 		// always delete the existing commandline text file, so it doesn't reuse an old one
 		string IntermediateCmdLineFile = CombinePaths(SC.StageDirectory, "UE4CommandLine.txt");
 		Project.WriteStageCommandline(IntermediateCmdLineFile, Params, SC);
 
-        // Setup the OBB name and add the storage path (queried from the device) to it
-        string DeviceStorageQueryCommand = GetStorageQueryCommand();
-		ProcessResult Result = RunAdbCommand(Params, DeviceStorageQueryCommand, null, ERunOptions.AppMustExist);
-        String StorageLocation = Result.Output.Trim();
-        string DeviceObbName = StorageLocation + "/" + GetDeviceObbName(ApkName);
-
 		// copy files to device if we were staging
 		if (SC.Stage)
 		{
 			// cache some strings
 			string BaseCommandline = "push";
-			string RemoteDir = StorageLocation + "/UE4Game/" + Params.ShortProjectName;
-            string UE4GameRemoteDir = StorageLocation + "/UE4Game/" + Params.ShortProjectName;
-
-			// make sure device is at a clean state
-			RunAdbCommand(Params, "shell rm -r " + RemoteDir);
-			RunAdbCommand(Params, "shell rm -r " + UE4GameRemoteDir);
-
-			// Copy UFS files..
-			string[] Files = Directory.GetFiles(SC.StageDirectory, "*", SearchOption.AllDirectories);
-			System.Array.Sort(Files);
-
-			// Find all the files we exclude from copying. And include
-			// the directories we need to individually copy.
-			HashSet<string> ExcludedFiles = new HashSet<string>();
-			SortedSet<string> IndividualCopyDirectories
-				= new SortedSet<string>((IComparer<string>)new LongestFirst());
-			foreach (string Filename in Files)
+
+            HashSet<string> EntriesToDeploy = new HashSet<string>();
+
+			if (Params.IterativeDeploy)
 			{
-				bool Exclude = false;
-				// Don't push the apk, we install it
-				Exclude |= Path.GetExtension(Filename).Equals(".apk", StringComparison.InvariantCultureIgnoreCase);
-				// For excluded files we add the parent dirs to our
-				// tracking of stuff to individually copy.
-				if (Exclude)
+				// Add non UFS files if any to deploy
+				String NonUFSManifestPath = SC.GetNonUFSDeploymentDeltaPath();
+				if (File.Exists(NonUFSManifestPath))
 				{
-					ExcludedFiles.Add(Filename);
-					// We include all directories up to the stage root in having
-					// to individually copy the files.
-					for (string FileDirectory = Path.GetDirectoryName(Filename);
-						!FileDirectory.Equals(SC.StageDirectory);
-						FileDirectory = Path.GetDirectoryName(FileDirectory))
+					string NonUFSFiles = File.ReadAllText(NonUFSManifestPath);
+					foreach (string Filename in NonUFSFiles.Split('\n'))
 					{
-						if (!IndividualCopyDirectories.Contains(FileDirectory))
+						if (!string.IsNullOrEmpty(Filename) && !string.IsNullOrWhiteSpace(Filename))
 						{
-							IndividualCopyDirectories.Add(FileDirectory);
+							EntriesToDeploy.Add(CombinePaths(SC.StageDirectory, Filename.Trim()));
 						}
 					}
-					if (!IndividualCopyDirectories.Contains(SC.StageDirectory))
+				}
+
+				// Add UFS files if any to deploy
+				String UFSManifestPath = SC.GetUFSDeploymentDeltaPath();
+				if (File.Exists(UFSManifestPath))
+				{
+					string UFSFiles = File.ReadAllText(UFSManifestPath);
+					foreach (string Filename in UFSFiles.Split('\n'))
 					{
-						IndividualCopyDirectories.Add(SC.StageDirectory);
+						if (!string.IsNullOrEmpty(Filename) && !string.IsNullOrWhiteSpace(Filename))
+						{
+							EntriesToDeploy.Add(CombinePaths(SC.StageDirectory, Filename.Trim()));
+						}
 					}
 				}
-			}
 
-			// The directories are sorted above in "deepest" first. We can
-			// therefore start copying those individual dirs which will
-			// recreate the tree. As the subtrees will get copied at each
-			// possible individual level.
-			HashSet<string> EntriesToDeploy = new HashSet<string>();
-			foreach (string DirectoryName in IndividualCopyDirectories)
+				// For now, if too many files may be better to just push them all
+				if (EntriesToDeploy.Count > 500)
+				{
+					// make sure device is at a clean state
+					RunAdbCommand(Params, "shell rm -r " + RemoteDir);
+
+					EntriesToDeploy.Clear();
+					EntriesToDeploy.TrimExcess();
+					EntriesToDeploy.Add(SC.StageDirectory);
+				}
+			}
+			else
 			{
-				string[] Entries
-					= Directory.GetFileSystemEntries(DirectoryName, "*", SearchOption.TopDirectoryOnly);
-				foreach (string Entry in Entries)
+				// make sure device is at a clean state
+				RunAdbCommand(Params, "shell rm -r " + RemoteDir);
+
+				// Copy UFS files..
+				string[] Files = Directory.GetFiles(SC.StageDirectory, "*", SearchOption.AllDirectories);
+				System.Array.Sort(Files);
+
+				// Find all the files we exclude from copying. And include
+				// the directories we need to individually copy.
+				HashSet<string> ExcludedFiles = new HashSet<string>();
+				SortedSet<string> IndividualCopyDirectories
+					= new SortedSet<string>((IComparer<string>)new LongestFirst());
+				foreach (string Filename in Files)
 				{
-					// We avoid excluded files and the individual copy dirs
-					// (the individual copy dirs will get handled as we iterate).
-					if (ExcludedFiles.Contains(Entry) || IndividualCopyDirectories.Contains(Entry))
+					bool Exclude = false;
+					// Don't push the apk, we install it
+					Exclude |= Path.GetExtension(Filename).Equals(".apk", StringComparison.InvariantCultureIgnoreCase);
+					// For excluded files we add the parent dirs to our
+					// tracking of stuff to individually copy.
+					if (Exclude)
 					{
-						continue;
+						ExcludedFiles.Add(Filename);
+						// We include all directories up to the stage root in having
+						// to individually copy the files.
+						for (string FileDirectory = Path.GetDirectoryName(Filename);
+							!FileDirectory.Equals(SC.StageDirectory);
+							FileDirectory = Path.GetDirectoryName(FileDirectory))
+						{
+							if (!IndividualCopyDirectories.Contains(FileDirectory))
+							{
+								IndividualCopyDirectories.Add(FileDirectory);
+							}
+						}
+						if (!IndividualCopyDirectories.Contains(SC.StageDirectory))
+						{
+							IndividualCopyDirectories.Add(SC.StageDirectory);
+						}
 					}
-					else
+				}
+
+				// The directories are sorted above in "deepest" first. We can
+				// therefore start copying those individual dirs which will
+				// recreate the tree. As the subtrees will get copied at each
+				// possible individual level.
+				foreach (string DirectoryName in IndividualCopyDirectories)
+				{
+					string[] Entries
+						= Directory.GetFileSystemEntries(DirectoryName, "*", SearchOption.TopDirectoryOnly);
+					foreach (string Entry in Entries)
 					{
-						EntriesToDeploy.Add(Entry);
+						// We avoid excluded files and the individual copy dirs
+						// (the individual copy dirs will get handled as we iterate).
+						if (ExcludedFiles.Contains(Entry) || IndividualCopyDirectories.Contains(Entry))
+						{
+							continue;
+						}
+						else
+						{
+							EntriesToDeploy.Add(Entry);
+						}
 					}
 				}
-			}
-			if (EntriesToDeploy.Count == 0)
-			{
-				EntriesToDeploy.Add(SC.StageDirectory);
+
+				if (EntriesToDeploy.Count == 0)
+				{
+					EntriesToDeploy.Add(SC.StageDirectory);
+				}
 			}
 
 			// We now have a minimal set of file & dir entries we need
@@ -633,7 +741,6 @@ public class AndroidPlatform : Platform
 		{
 			// cache some strings
 			string BaseCommandline = "push";
-            string RemoteDir = StorageLocation + "/UE4Game/" + Params.ShortProjectName;
 
 			string FinalRemoteDir = RemoteDir;
 			/*
@@ -648,7 +755,14 @@ public class AndroidPlatform : Platform
 			string Commandline = string.Format("{0} \"{1}\" \"{2}\"", BaseCommandline, IntermediateCmdLineFile, RemoteFilename);
 			RunAdbCommand(Params, Commandline);
 		}
-	}
+    
+        // write new timestamp for APK (do it here since RemoteDir will now exist)
+        if (bNeedAPKInstall)
+        {
+            int SuccessCode = 0;
+            RunAndLogAdbCommand(Params, "shell \"echo 'APK: " + APKLastUpdateTime + "' > " + RemoteDir + "/APKFileStamp.txt\"", out SuccessCode);
+        }
+    }
 
 	/** Internal usage for GetPackageName */
 	private static string PackageLine = null;
diff --git a/Engine/Source/Programs/AutomationTool/Scripts/CopyBuildToStagingDirectory.Automation.cs b/Engine/Source/Programs/AutomationTool/Scripts/CopyBuildToStagingDirectory.Automation.cs
--- a/Engine/Source/Programs/AutomationTool/Scripts/CopyBuildToStagingDirectory.Automation.cs
+++ b/Engine/Source/Programs/AutomationTool/Scripts/CopyBuildToStagingDirectory.Automation.cs
@@ -423,8 +423,8 @@ public partial class Project : CommandUtils
 
             if (SC.StageTargetPlatform.StageMovies)
             {
-                SC.StageFiles(StagedFileTypeForMovies, CombinePaths(SC.LocalRoot, "Engine/Content/Movies"), "*", true, null, CombinePaths(SC.RelativeProjectRootForStage, "Engine/Content/Movies"), true, !Params.UsePak(SC.StageTargetPlatform));
-                SC.StageFiles(StagedFileTypeForMovies, CombinePaths(SC.ProjectRoot, "Content/Movies"), "*", true, null, CombinePaths(SC.RelativeProjectRootForStage, "Content/Movies"), true, !Params.UsePak(SC.StageTargetPlatform));
+                SC.StageFiles(StagedFileTypeForMovies, CombinePaths(SC.LocalRoot, "Engine/Content/Movies"), "*", true, new string[] { "*.uasset", "*.umap" }, CombinePaths(SC.RelativeProjectRootForStage, "Engine/Content/Movies"), true, !Params.UsePak(SC.StageTargetPlatform));
+                SC.StageFiles(StagedFileTypeForMovies, CombinePaths(SC.ProjectRoot, "Content/Movies"), "*", true, new string[] { "*.uasset", "*.umap" }, CombinePaths(SC.RelativeProjectRootForStage, "Content/Movies"), true, !Params.UsePak(SC.StageTargetPlatform));
             }
 
             // eliminate the sand box
