Remove symlinks from list of files to copy in Rocket builds; C# tools don't support copying Mono symlinks on Mac, but they're recreated by SetupMono anyway.

diff --git a/Engine/Source/Programs/AutomationTool/Scripts/RocketBuild.Automation.cs b/Engine/Source/Programs/AutomationTool/Scripts/RocketBuild.Automation.cs
--- a/Engine/Source/Programs/AutomationTool/Scripts/RocketBuild.Automation.cs
+++ b/Engine/Source/Programs/AutomationTool/Scripts/RocketBuild.Automation.cs
@@ -356,19 +356,24 @@ namespace Rocket
 			CommandUtils.Log("Applying filter to {0}...", SourceDir);
 			string[] RelativePaths = Filter.ApplyToDirectory(SourceDir).ToArray();
 
-			// Find all the source and target filenames
-			string[] SourceFileNames = new string[RelativePaths.Length];
-			string[] TargetFileNames = new string[RelativePaths.Length];
+			// Find all the source and target filenames, removing any symlinks.
+			List<string> SourceFileNames = new List<string>();
+			List<string> TargetFileNames = new List<string>();
 			for (int Idx = 0; Idx < RelativePaths.Length; Idx++)
 			{
-				SourceFileNames[Idx] = CommandUtils.CombinePaths(SourceDir, RelativePaths[Idx]);
-				TargetFileNames[Idx] = CommandUtils.CombinePaths(TargetDir, RelativePaths[Idx]);
+				string SourceFileName = CommandUtils.CombinePaths(SourceDir, RelativePaths[Idx]);
+				if(!File.GetAttributes(SourceFileName).HasFlag(FileAttributes.ReparsePoint))
+				{
+					SourceFileNames.Add(SourceFileName);
+					TargetFileNames.Add(CommandUtils.CombinePaths(TargetDir, RelativePaths[Idx]));
+				}
 			}
 
 			// Copy them all over
-			CommandUtils.Log("Copying {0} files to {1}...", SourceFileNames.Length, TargetDir);
-			CommandUtils.ThreadedCopyFiles(SourceFileNames, TargetFileNames);
-			return TargetFileNames;
+			string[] TargetFileNamesArray = TargetFileNames.ToArray();
+			CommandUtils.Log("Copying {0} files to {1}...", SourceFileNames.Count, TargetDir);
+			CommandUtils.ThreadedCopyFiles(SourceFileNames.ToArray(), TargetFileNamesArray);
+			return TargetFileNamesArray;
 		}
 
 		public static void StripSymbols(UnrealTargetPlatform TargetPlatform, IEnumerable<string> FileNames)
