Remove symbols from static libraries on Win32/Win64. No way to do this without using PDBs on Windows (and that fails with XGE due to MSPDBSRV being 32-bit), so just don't create debug info at all (the linker will still be able to make a PDB to get callstacks).

diff --git a/Engine/Source/Programs/AutomationTool/Scripts/GUBP.Automation.cs b/Engine/Source/Programs/AutomationTool/Scripts/GUBP.Automation.cs
--- a/Engine/Source/Programs/AutomationTool/Scripts/GUBP.Automation.cs
+++ b/Engine/Source/Programs/AutomationTool/Scripts/GUBP.Automation.cs
@@ -1678,6 +1678,13 @@ public class GUBP : BuildCommand
 			if(Precompiled)
 			{
 				Args += " -precompile";
+
+				// MSVC doesn't provide a way to strip symbols from static libraries - you have to use PDBs, but that causes random OOM 
+				// exceptions with the /FS arg because mspdbsrv is 32-bit. Just disable compiler debug info manually for now.
+				if(TargetPlatform == UnrealTargetPlatform.Win32 || TargetPlatform == UnrealTargetPlatform.Win64)
+				{
+	                Args += " -nodebuginfo";
+				}
 			}
 
 			if (WithXp)
