Android IterativeDeploy will now always push UE4CommandLine.txt (delta checks occurred before it was written to staging directory)
#ue4
#android

diff --git a/Engine/Source/Programs/AutomationTool/Android/AndroidPlatform.Automation.cs b/Engine/Source/Programs/AutomationTool/Android/AndroidPlatform.Automation.cs
--- a/Engine/Source/Programs/AutomationTool/Android/AndroidPlatform.Automation.cs
+++ b/Engine/Source/Programs/AutomationTool/Android/AndroidPlatform.Automation.cs
@@ -576,6 +576,9 @@ public class AndroidPlatform : Platform
 
 			if (Params.IterativeDeploy)
 			{
+				// always send UE4CommandLine.txt (it was written above after delta checks applied)
+				EntriesToDeploy.Add(IntermediateCmdLineFile);
+
 				// Add non UFS files if any to deploy
 				String NonUFSManifestPath = SC.GetNonUFSDeploymentDeltaPath();
 				if (File.Exists(NonUFSManifestPath))
