Add a GUBP node which archives all the required build products for a Git promotion.

diff --git a/Engine/Source/Programs/AutomationTool/Scripts/RocketBuild.Automation.cs b/Engine/Source/Programs/AutomationTool/Scripts/RocketBuild.Automation.cs
--- a/Engine/Source/Programs/AutomationTool/Scripts/RocketBuild.Automation.cs
+++ b/Engine/Source/Programs/AutomationTool/Scripts/RocketBuild.Automation.cs
@@ -74,6 +74,9 @@ namespace Rocket
 		{
 			if(!bp.BranchOptions.bNoInstalledEngine)
 			{
+				// Get the promotable node, so we can add additional dependencies to it
+				GUBP.GUBPNode PromotableNode = bp.FindNode(GUBP.SharedAggregatePromotableNode.StaticGetFullName());
+
 				// Get the output paths for the installed engine
 				string InstallDir = CommandUtils.CombinePaths(CommandUtils.CmdEnv.LocalRoot, "LocalBuilds", "Engine", CommandUtils.GetGenericPlatformName(HostPlatform));
 				string SymbolsDir = CommandUtils.CombinePaths(CommandUtils.CmdEnv.LocalRoot, "LocalBuilds", "Engine", CommandUtils.GetGenericPlatformName(HostPlatform) + "Symbols");
@@ -118,6 +121,17 @@ namespace Rocket
 
 				// Add the aggregate node for the entire install
 				AggNode.AddDependency(bp.AddNode(new BuildInstalledEngineNode(HostPlatform, InstallDir, SymbolsDir)));
+
+				// Add a node for GitHub promotions
+				if(HostPlatform == UnrealTargetPlatform.Win64)
+				{
+					string GitConfigRelativePath = "Engine/Build/Git/UnrealBot.ini";
+					if(CommandUtils.FileExists(CommandUtils.CombinePaths(CommandUtils.CmdEnv.LocalRoot, GitConfigRelativePath)))
+					{
+						bp.AddNode(new BuildGitPromotable(bp, HostPlatform, GitConfigRelativePath));
+						PromotableNode.AddDependency(BuildGitPromotable.StaticGetFullName(HostPlatform));
+					}
+				}
 			}
 		}
 
@@ -207,6 +221,54 @@ namespace Rocket
 		}
 	}
 
+	public class BuildGitPromotable : GUBP.HostPlatformNode
+	{
+		string ConfigRelativePath;
+
+		public BuildGitPromotable(GUBP bp, UnrealTargetPlatform HostPlatform, string InConfigRelativePath) : base(HostPlatform)
+		{
+			ConfigRelativePath = InConfigRelativePath;
+
+			foreach(UnrealTargetPlatform OtherHostPlatform in bp.HostPlatforms)
+			{
+				AddDependency(GUBP.RootEditorNode.StaticGetFullName(OtherHostPlatform));
+				AddDependency(GUBP.ToolsNode.StaticGetFullName(OtherHostPlatform));
+				AddDependency(GUBP.InternalToolsNode.StaticGetFullName(OtherHostPlatform));
+			}
+
+			SingleTargetProperties VersionSelector = bp.Branch.FindProgram("UnrealVersionSelector"); 
+			if (VersionSelector.Rules != null) 
+			{ 
+				AddDependency(GUBP.SingleInternalToolsNode.StaticGetFullName(UnrealTargetPlatform.Win64, VersionSelector));
+			}
+		}
+
+		public static string StaticGetFullName(UnrealTargetPlatform HostPlatform)
+		{
+			return "BuildGitPromotable" + StaticGetHostPlatformSuffix(HostPlatform);
+		}
+
+		public override string GetFullName()
+		{
+			return StaticGetFullName(HostPlatform);
+		}
+
+		public override void DoBuild(GUBP bp)
+		{
+			// Create a filter for all the promoted binaries
+			FileFilter PromotableFilter = new FileFilter();
+			CopyNode.AddFilesToFilter(PromotableFilter, AllDependencyBuildProducts, CommandUtils.CmdEnv.LocalRoot, FileFilterType.Include);
+			PromotableFilter.AddRulesFromFile(CommandUtils.CombinePaths(CommandUtils.CmdEnv.LocalRoot, ConfigRelativePath), "promotable");
+			CopyNode.ExcludeConfidentialFolders(PromotableFilter);
+
+			// Copy everything that matches the filter to the promotion folder
+			string PromotableFolder = CommandUtils.CombinePaths(CommandUtils.CmdEnv.LocalRoot, "Engine", "Saved", "GitPromotable");
+			CommandUtils.DeleteDirectoryContents(PromotableFolder);
+			string[] PromotableFiles = CopyNode.CopyFolder(CommandUtils.CmdEnv.LocalRoot, PromotableFolder, PromotableFilter);
+			BuildProducts = new List<string>(PromotableFiles);
+		}
+	}
+
 	public class CopySymbolsNode : GUBP.HostPlatformNode
 	{
 		public string OutputDirectory;
@@ -337,12 +399,15 @@ namespace Rocket
 			AddFilesToFilter(Filter, Node.BuildProducts, BaseDirectoryName, Type);
 		}
 
-		public static void ExcludeConfidentialFolders(FileFilter Filter)
+		public static void ExcludeConfidentialPlatforms(FileFilter Filter)
 		{
 			// Exclude platforms that can't be redistributed without an NDA
 			Filter.Exclude(".../PS4/...");
 			Filter.Exclude(".../XboxOne/...");
+		}
 
+		public static void ExcludeConfidentialFolders(FileFilter Filter)
+		{
 			// Exclude standard folders
 			Filter.Exclude(".../EpicInternal/...");
 			Filter.Exclude(".../CarefullyRedist/...");
@@ -355,17 +420,21 @@ namespace Rocket
 			// Filter all the relative paths
 			CommandUtils.Log("Applying filter to {0}...", SourceDir);
 			string[] RelativePaths = Filter.ApplyToDirectory(SourceDir).ToArray();
+			return CopyFolder(SourceDir, TargetDir, RelativePaths);
+		}
 
+		public static string[] CopyFolder(string SourceDir, string TargetDir, IEnumerable<string> RelativePaths)
+		{
 			// Find all the source and target filenames, removing any symlinks.
 			List<string> SourceFileNames = new List<string>();
 			List<string> TargetFileNames = new List<string>();
-			for (int Idx = 0; Idx < RelativePaths.Length; Idx++)
+			foreach(string RelativePath in RelativePaths)
 			{
-				string SourceFileName = CommandUtils.CombinePaths(SourceDir, RelativePaths[Idx]);
+				string SourceFileName = CommandUtils.CombinePaths(SourceDir, RelativePath);
 				if(!File.GetAttributes(SourceFileName).HasFlag(FileAttributes.ReparsePoint))
 				{
 					SourceFileNames.Add(SourceFileName);
-					TargetFileNames.Add(CommandUtils.CombinePaths(TargetDir, RelativePaths[Idx]));
+					TargetFileNames.Add(CommandUtils.CombinePaths(TargetDir, RelativePath));
 				}
 			}
 
@@ -586,6 +655,7 @@ namespace Rocket
 			AddNodeToFilter(Filter, bp, CopyInstalledEditorNode.StaticGetFullName(HostPlatform), FileFilterType.Exclude);
 
 			// Add the final exclusions for legal reasons.
+			ExcludeConfidentialPlatforms(Filter);
 			ExcludeConfidentialFolders(Filter);
 
 			// Copy all the files over
