Merging CL 2495959

//depot/UE4-UT/...

to //depot/UE4/...

diff --git a/Engine/Source/Programs/AutomationTool/Mac/MacPlatform.Automation.cs b/Engine/Source/Programs/AutomationTool/Mac/MacPlatform.Automation.cs
--- a/Engine/Source/Programs/AutomationTool/Mac/MacPlatform.Automation.cs
+++ b/Engine/Source/Programs/AutomationTool/Mac/MacPlatform.Automation.cs
@@ -46,48 +46,96 @@ public class MacPlatform : Platform
 
 	public override void GetFilesToDeployOrStage(ProjectParams Params, DeploymentContext SC)
 	{
-		// the first app is the "main" one, the rest are marked as debug files for exclusion from chunking/distribution
-		StagedFileType WorkingFileType = StagedFileType.NonUFS;
-
-		List<string> Exes = GetExecutableNames(SC);
-		foreach (var Exe in Exes)
+	
+		if (Params.StageNonMonolithic)
 		{
-			string AppBundlePath = "";
-			if (Exe.StartsWith(CombinePaths(SC.RuntimeProjectRootDir, "Binaries", SC.PlatformDir)))
+			if (SC.DedicatedServer)
 			{
-				AppBundlePath = CombinePaths(SC.ShortProjectName, "Binaries", SC.PlatformDir, Path.GetFileNameWithoutExtension(Exe) + ".app");
-				StageAppBundle(SC, WorkingFileType, CombinePaths(SC.ProjectRoot, "Binaries", SC.PlatformDir, Path.GetFileNameWithoutExtension(Exe) + ".app"), AppBundlePath);
+				if (SC.StageTargetConfigurations.Contains(UnrealTargetConfiguration.Development))
+				{
+					SC.StageFiles(StagedFileType.NonUFS, CombinePaths(SC.LocalRoot, "Engine/Binaries", SC.PlatformDir, "UE4Server.app"));
+					SC.StageFiles(StagedFileType.NonUFS, CombinePaths(SC.ProjectRoot, "Binaries", SC.PlatformDir), "UE4Server-" + SC.ShortProjectName + ".dylib");
+				}
+				if (SC.StageTargetConfigurations.Contains(UnrealTargetConfiguration.Test))
+				{
+					SC.StageFiles(StagedFileType.NonUFS, CombinePaths(SC.LocalRoot, "Engine/Binaries", SC.PlatformDir, "UE4Server-Mac-Test.app"));
+					SC.StageFiles(StagedFileType.NonUFS, CombinePaths(SC.ProjectRoot, "Binaries", SC.PlatformDir), "UE4Server-" + SC.ShortProjectName + "-Test.dylib");
+				}
+				if (SC.StageTargetConfigurations.Contains(UnrealTargetConfiguration.Shipping))
+				{
+					SC.StageFiles(StagedFileType.NonUFS, CombinePaths(SC.LocalRoot, "Engine/Binaries", SC.PlatformDir, "UE4Server-Mac-Shipping.app"));
+					SC.StageFiles(StagedFileType.NonUFS, CombinePaths(SC.ProjectRoot, "Binaries", SC.PlatformDir), "UE4Server-" + SC.ShortProjectName + "-Shipping.dylib");
+				}
+
+				SC.StageFiles(StagedFileType.NonUFS, CombinePaths(SC.LocalRoot, "Engine/Plugins"), "UE4Server-*.dylib", true);
+				SC.StageFiles(StagedFileType.NonUFS, CombinePaths(SC.ProjectRoot, "Plugins"), "UE4Server-*.dylib", true);
 			}
-			else if (Exe.StartsWith(CombinePaths(SC.RuntimeRootDir, "Engine/Binaries", SC.PlatformDir)))
+			else
 			{
-				AppBundlePath = CombinePaths("Engine/Binaries", SC.PlatformDir, Path.GetFileNameWithoutExtension(Exe) + ".app");
-
-				string AbsoluteBundlePath = CombinePaths (SC.LocalRoot, AppBundlePath);
-				// ensure the ue4game binary exists, if applicable
-				if (!SC.IsCodeBasedProject && !Directory.Exists(AbsoluteBundlePath))
+				if (SC.StageTargetConfigurations.Contains(UnrealTargetConfiguration.Development))
+				{
+					SC.StageFiles(StagedFileType.NonUFS, CombinePaths(SC.LocalRoot, "Engine/Binaries", SC.PlatformDir, "UE4.app"));
+					SC.StageFiles(StagedFileType.NonUFS, CombinePaths(SC.ProjectRoot, "Binaries", SC.PlatformDir), "UE4-" + SC.ShortProjectName + ".dylib");
+				}
+				if (SC.StageTargetConfigurations.Contains(UnrealTargetConfiguration.Test))
 				{
-					Log("Failed to find app bundle " + AbsoluteBundlePath);
-					AutomationTool.ErrorReporter.Error("Stage Failed.", (int)AutomationTool.ErrorCodes.Error_MissingExecutable);
-					throw new AutomationException("Could not find app bundle {0}. You may need to build the UE4 project with your target configuration and platform.", AbsoluteBundlePath);
+					SC.StageFiles(StagedFileType.NonUFS, CombinePaths(SC.LocalRoot, "Engine/Binaries", SC.PlatformDir, "UE4-Mac-Test.app"));
+					SC.StageFiles(StagedFileType.NonUFS, CombinePaths(SC.ProjectRoot, "Binaries", SC.PlatformDir), "UE4-" + SC.ShortProjectName + "Test.dylib");
+				}
+				if (SC.StageTargetConfigurations.Contains(UnrealTargetConfiguration.Shipping))
+				{
+					SC.StageFiles(StagedFileType.NonUFS, CombinePaths(SC.LocalRoot, "Engine/Binaries", SC.PlatformDir, "UE4-Mac-Shipping.app"));
+					SC.StageFiles(StagedFileType.NonUFS, CombinePaths(SC.ProjectRoot, "Binaries", SC.PlatformDir), "UE4-" + SC.ShortProjectName + "Shipping.dylib");
 				}
 
-				StageAppBundle(SC, WorkingFileType, CombinePaths(SC.LocalRoot, "Engine/Binaries", SC.PlatformDir, Path.GetFileNameWithoutExtension(Exe) + ".app"), AppBundlePath);
+				SC.StageFiles(StagedFileType.NonUFS, CombinePaths(SC.LocalRoot, "Engine/Plugins"), "UE4-*.dylib", true);
+				SC.StageFiles(StagedFileType.NonUFS, CombinePaths(SC.ProjectRoot, "Plugins"), "UE4-*.dylib", true, null, null, true);
 			}
+		}
+		else
+		{
+			// the first app is the "main" one, the rest are marked as debug files for exclusion from chunking/distribution
+			StagedFileType WorkingFileType = StagedFileType.NonUFS;
 
-			if (!string.IsNullOrEmpty(AppBundlePath))
+			List<string> Exes = GetExecutableNames(SC);
+			foreach (var Exe in Exes)
 			{
-				SC.StageFiles(WorkingFileType, CombinePaths(SC.ProjectRoot, "Build/Mac"), "Application.icns", false, null, CombinePaths(AppBundlePath, "Contents/Resources"), true);
+				string AppBundlePath = "";
+				if (Exe.StartsWith(CombinePaths(SC.RuntimeProjectRootDir, "Binaries", SC.PlatformDir)))
+				{
+					AppBundlePath = CombinePaths(SC.ShortProjectName, "Binaries", SC.PlatformDir, Path.GetFileNameWithoutExtension(Exe) + ".app");
+					StageAppBundle(SC, WorkingFileType, CombinePaths(SC.ProjectRoot, "Binaries", SC.PlatformDir, Path.GetFileNameWithoutExtension(Exe) + ".app"), AppBundlePath);
+				}
+				else if (Exe.StartsWith(CombinePaths(SC.RuntimeRootDir, "Engine/Binaries", SC.PlatformDir)))
+				{
+					AppBundlePath = CombinePaths("Engine/Binaries", SC.PlatformDir, Path.GetFileNameWithoutExtension(Exe) + ".app");
+
+					string AbsoluteBundlePath = CombinePaths (SC.LocalRoot, AppBundlePath);
+					// ensure the ue4game binary exists, if applicable
+					if (!SC.IsCodeBasedProject && !Directory.Exists(AbsoluteBundlePath))
+					{
+						Log("Failed to find app bundle " + AbsoluteBundlePath);
+						AutomationTool.ErrorReporter.Error("Stage Failed.", (int)AutomationTool.ErrorCodes.Error_MissingExecutable);
+						throw new AutomationException("Could not find app bundle {0}. You may need to build the UE4 project with your target configuration and platform.", AbsoluteBundlePath);
+					}
 
-				if (Params.bUsesSteam)
+					StageAppBundle(SC, WorkingFileType, CombinePaths(SC.LocalRoot, "Engine/Binaries", SC.PlatformDir, Path.GetFileNameWithoutExtension(Exe) + ".app"), AppBundlePath);
+				}
+
+				if (!string.IsNullOrEmpty(AppBundlePath))
 				{
-					SC.StageFiles(StagedFileType.NonUFS, CombinePaths(SC.LocalRoot, "Engine/Source/ThirdParty/Steamworks/Steamv132/sdk/redistributable_bin/osx32"), "libsteam_api.dylib", false, null, CombinePaths(AppBundlePath, "Contents/MacOS"), true);
+					SC.StageFiles(WorkingFileType, CombinePaths(SC.ProjectRoot, "Build/Mac"), "Application.icns", false, null, CombinePaths(AppBundlePath, "Contents/Resources"), true);
+
+					if (Params.bUsesSteam)
+					{
+						SC.StageFiles(StagedFileType.NonUFS, CombinePaths(SC.LocalRoot, "Engine/Source/ThirdParty/Steamworks/Steamv132/sdk/redistributable_bin/osx32"), "libsteam_api.dylib", false, null, CombinePaths(AppBundlePath, "Contents/MacOS"), true);
+					}
 				}
-			}
 
-			// the first app is the "main" one, the rest are marked as debug files for exclusion from chunking/distribution
-			WorkingFileType = StagedFileType.DebugNonUFS;
+				// the first app is the "main" one, the rest are marked as debug files for exclusion from chunking/distribution
+				WorkingFileType = StagedFileType.DebugNonUFS;
+			}
 		}
-
 		// Copy the splash screen, Mac specific
 		SC.StageFiles(StagedFileType.NonUFS, CombinePaths(SC.ProjectRoot, "Content/Splash"), "Splash.bmp", false, null, null, true);
 
