Remove bBuildRocket from GUBP. Everything is now part of the larger graph.

diff --git a/Engine/Source/Programs/AutomationTool/Scripts/GUBP.Automation.cs b/Engine/Source/Programs/AutomationTool/Scripts/GUBP.Automation.cs
--- a/Engine/Source/Programs/AutomationTool/Scripts/GUBP.Automation.cs
+++ b/Engine/Source/Programs/AutomationTool/Scripts/GUBP.Automation.cs
@@ -113,8 +113,6 @@ public class GUBP : BuildCommand
     public List<UnrealTargetPlatform> HostPlatforms;
     public bool bFake = false;
     public static bool bNoIOSOnPC = false;
-    public static bool bBuildRocket = false;
-	public static bool bBuildOnlySamples = false;
     public static bool bForceIncrementalCompile = false;
     public static string ECProject = null;
     public string EmailHint;
@@ -143,16 +141,6 @@ public class GUBP : BuildCommand
 
     Dictionary<string, NodeHistory> GUBPNodesHistory;
 
-    public string RocketUBTArgs(bool bUseRocketInsteadOfBuildRocket = false)
-    {
-		string Args = "";
-        if (bBuildRocket)
-        {
-            Args = (bUseRocketInsteadOfBuildRocket ? "-Rocket" : "-BuildRocket -Precompile");
-        }
-        return Args;
-    }
-
     public abstract class GUBPNodeAdder
     {
         public virtual void AddNodes(GUBP bp, UnrealTargetPlatform InHostPlatform)
@@ -788,7 +776,7 @@ public class GUBP : BuildCommand
 			    }
                     );
             }
-            string AddArgs = "-CopyAppBundleBackToDevice" + bp.RocketUBTArgs();
+            string AddArgs = "-CopyAppBundleBackToDevice";
 
             Agenda.AddTargets(new string[] { "UnrealHeaderTool" }, HostPlatform, UnrealTargetConfiguration.Development, InAddArgs: AddArgs);
             return Agenda;
@@ -938,7 +926,7 @@ public class GUBP : BuildCommand
 		{
 			var Agenda = new UE4Build.BuildAgenda();
 
-			string AddArgs = "-nobuilduht" + bp.RocketUBTArgs();
+			string AddArgs = "-nobuilduht";
 			Agenda.AddTargets(
 				new string[] { bp.Branch.BaseEngineProject.Properties.Targets[TargetRules.TargetType.Editor].TargetName },
 				UnrealTargetPlatform.Linux, UnrealTargetConfiguration.Development, InAddArgs: AddArgs);
@@ -950,10 +938,7 @@ public class GUBP : BuildCommand
         public ToolsNode(UnrealTargetPlatform InHostPlatform)
             : base(InHostPlatform)
         {
-            if (!GUBP.bBuildRocket) // more errors and more performance by just starting before the root editor is done
-            {
-                AddPseudodependency(RootEditorNode.StaticGetFullName(HostPlatform));
-            }
+            AddPseudodependency(RootEditorNode.StaticGetFullName(HostPlatform));
             AgentSharingGroup = "ToolsGroup" + StaticGetHostPlatformSuffix(HostPlatform);
         }
         public static string StaticGetFullName(UnrealTargetPlatform InHostPlatform)
@@ -1013,7 +998,7 @@ public class GUBP : BuildCommand
 				}
             }
 
-            string AddArgs = "-nobuilduht -skipactionhistory -CopyAppBundleBackToDevice" + bp.RocketUBTArgs(); ;
+            string AddArgs = "-nobuilduht -skipactionhistory -CopyAppBundleBackToDevice";
 
             foreach (var ProgramTarget in bp.Branch.BaseEngineProject.Properties.Programs)
             {
@@ -1040,10 +1025,7 @@ public class GUBP : BuildCommand
 		public ToolsCrossCompileNode(UnrealTargetPlatform InHostPlatform)
 			: base(InHostPlatform)
 		{
-			if (!GUBP.bBuildRocket) // more errors and more performance by just starting before the root editor is done
-			{
-				AddPseudodependency(RootEditorCrossCompileLinuxNode.StaticGetFullName(HostPlatform));
-			}
+			AddPseudodependency(RootEditorCrossCompileLinuxNode.StaticGetFullName(HostPlatform));
 			AgentSharingGroup = "ToolsCrossCompileGroup" + StaticGetHostPlatformSuffix(HostPlatform);
 		}
 		public static string StaticGetFullName(UnrealTargetPlatform InHostPlatform)
@@ -1066,7 +1048,7 @@ public class GUBP : BuildCommand
 		{
 			var Agenda = new UE4Build.BuildAgenda();
 
-			string AddArgs = "-nobuilduht -skipactionhistory -CopyAppBundleBackToDevice" + bp.RocketUBTArgs(); ;
+			string AddArgs = "-nobuilduht -skipactionhistory -CopyAppBundleBackToDevice";
 
 			foreach (var ProgramTarget in bp.Branch.BaseEngineProject.Properties.Programs)
 			{
@@ -1092,10 +1074,7 @@ public class GUBP : BuildCommand
             : base(InHostPlatform)
         {
             ProgramTarget = InProgramTarget;
-            if (!GUBP.bBuildRocket) // more errors and more performance by just starting before the root editor is done
-            {
-                AddPseudodependency(RootEditorNode.StaticGetFullName(HostPlatform));
-            }
+            AddPseudodependency(RootEditorNode.StaticGetFullName(HostPlatform));
             AgentSharingGroup = "ToolsGroup" + StaticGetHostPlatformSuffix(HostPlatform);
         }
         public static string StaticGetFullName(UnrealTargetPlatform InHostPlatform, SingleTargetProperties InProgramTarget)
@@ -1118,7 +1097,7 @@ public class GUBP : BuildCommand
         {
             var Agenda = new UE4Build.BuildAgenda();
 
-            string AddArgs = "-nobuilduht -skipactionhistory -CopyAppBundleBackToDevice" + bp.RocketUBTArgs(); ;
+            string AddArgs = "-nobuilduht -skipactionhistory -CopyAppBundleBackToDevice";
 
             foreach (var Plat in ProgramTarget.Rules.GUBP_ToolPlatforms(HostPlatform))
             {
@@ -1136,10 +1115,7 @@ public class GUBP : BuildCommand
         public InternalToolsNode(UnrealTargetPlatform InHostPlatform)
             : base(InHostPlatform)
         {
-            if (!GUBP.bBuildRocket) // more errors and more performance by just starting before the root editor is done
-            {
-                AddPseudodependency(RootEditorNode.StaticGetFullName(HostPlatform));
-            }
+            AddPseudodependency(RootEditorNode.StaticGetFullName(HostPlatform));
             AgentSharingGroup = "ToolsGroup" + StaticGetHostPlatformSuffix(HostPlatform);
         }
         public static string StaticGetFullName(UnrealTargetPlatform InHostPlatform)
@@ -1195,7 +1171,7 @@ public class GUBP : BuildCommand
 			        }
                     );
             }
-            string AddArgs = "-nobuilduht -skipactionhistory -CopyAppBundleBackToDevice" + bp.RocketUBTArgs(); ;
+            string AddArgs = "-nobuilduht -skipactionhistory -CopyAppBundleBackToDevice";
 
             foreach (var ProgramTarget in bp.Branch.BaseEngineProject.Properties.Programs)
             {
@@ -1242,9 +1218,7 @@ public class GUBP : BuildCommand
 			ProgramTarget = InProgramTarget;
 			AgentSharingGroup = "ToolsGroup" + StaticGetHostPlatformSuffix(HostPlatform);
 			
-			// more errors and more performance for Rocket by just starting before the root editor is done
-			bool bSkipRootEditorPsuedoDependencyForRocket = GUBP.bBuildRocket;
-			if (!bSkipRootEditorPsuedoDependencyForRocket && !bSkipRootEditorPsuedoDependency) 
+			if (!bSkipRootEditorPsuedoDependency) 
 			{
 				AddPseudodependency(RootEditorNode.StaticGetFullName(HostPlatform));
 			}
@@ -1274,7 +1248,7 @@ public class GUBP : BuildCommand
         {
             var Agenda = new UE4Build.BuildAgenda();
 
-            string AddArgs = "-nobuilduht -skipactionhistory -CopyAppBundleBackToDevice" + bp.RocketUBTArgs(); ;
+            string AddArgs = "-nobuilduht -skipactionhistory -CopyAppBundleBackToDevice";
 
             foreach (var Plat in ProgramTarget.Rules.GUBP_ToolPlatforms(HostPlatform))
             {
@@ -1387,7 +1361,7 @@ public class GUBP : BuildCommand
         {
             var Agenda = new UE4Build.BuildAgenda();
 
-            string Args = "-nobuilduht -skipactionhistory -skipnonhostplatforms -CopyAppBundleBackToDevice -forceheadergeneration" + bp.RocketUBTArgs(true);
+            string Args = "-nobuilduht -skipactionhistory -skipnonhostplatforms -CopyAppBundleBackToDevice -forceheadergeneration";
 
             Agenda.AddTargets(
                 new string[] { GameProj.Properties.Targets[TargetRules.TargetType.Editor].TargetName },
@@ -1545,8 +1519,7 @@ public class GUBP : BuildCommand
             }
             else
             {
-                if (TargetPlatform != InHostPlatform && bp.HasNode(GamePlatformMonolithicsNode.StaticGetFullName(InHostPlatform, bp.Branch.BaseEngineProject, InHostPlatform))
-                    && !GUBP.bBuildRocket) // more errors and more performance by just starting, for example, IOS without waiting for mac
+                if (TargetPlatform != InHostPlatform && bp.HasNode(GamePlatformMonolithicsNode.StaticGetFullName(InHostPlatform, bp.Branch.BaseEngineProject, InHostPlatform)))
                 {
                     AddPseudodependency(GamePlatformMonolithicsNode.StaticGetFullName(InHostPlatform, bp.Branch.BaseEngineProject, InHostPlatform));
                 }
@@ -1561,43 +1534,6 @@ public class GUBP : BuildCommand
                 AgentSharingGroup = "TemplateMonolithics" + StaticGetHostPlatformSuffix(InHostPlatform);
             }
 
-            if (GUBP.bBuildRocket)
-            {
-                var Target = bp.Branch.BaseEngineProject.Properties.Targets[TargetRules.TargetType.Game];
-                var Configs = Target.Rules.GUBP_GetConfigs_MonolithicOnly(HostPlatform, TargetPlatform);
-                foreach (var Config in Configs)
-                {
-                    if (HostPlatform == UnrealTargetPlatform.Win64)
-                    {
-                        if (TargetPlatform == UnrealTargetPlatform.Win32 && Config != UnrealTargetConfiguration.Shipping)
-                        {
-                            continue;
-                        }
-                        if (TargetPlatform == UnrealTargetPlatform.Win64 && Config != UnrealTargetConfiguration.Development)
-                        {
-                            continue;
-                        }
-                        if (TargetPlatform == UnrealTargetPlatform.Android && Config != UnrealTargetConfiguration.Shipping && Config != UnrealTargetConfiguration.Development)
-                        {
-                            continue;
-                        }
-                        if (TargetPlatform == UnrealTargetPlatform.HTML5 && Config != UnrealTargetConfiguration.Shipping && Config != UnrealTargetConfiguration.Development)
-                        {
-                            continue;
-                        }
-                        if (TargetPlatform == UnrealTargetPlatform.Linux && Config != UnrealTargetConfiguration.Shipping && Config != UnrealTargetConfiguration.Development)
-                        {
-                            continue;
-                        }
-                    }
-                    else if (Config != UnrealTargetConfiguration.Shipping && Config != UnrealTargetConfiguration.Development)
-                    {
-                        continue;
-                    }
-                    Log("Building {0} for Host={1} Target={2} Config={3} for rocket.   Node={4}", Target.TargetName, HostPlatform, TargetPlatform, Config, GetFullName());
-                }
-            }
-
         }
 
         public static string StaticGetFullName(UnrealTargetPlatform InHostPlatform, BranchInfo.BranchUProject InGameProj, UnrealTargetPlatform InTargetPlatform, bool WithXp = false, bool Precompiled = false)
@@ -1668,12 +1604,7 @@ public class GUBP : BuildCommand
             }
             var Agenda = new UE4Build.BuildAgenda();
 
-            string Args = "-nobuilduht -skipactionhistory -CopyAppBundleBackToDevice" + bp.RocketUBTArgs();
-
-            if (GUBP.bBuildRocket && (TargetPlatform == UnrealTargetPlatform.Win32 || TargetPlatform == UnrealTargetPlatform.Win64))
-            {
-                Args += " -nodebuginfo";
-            }
+            string Args = "-nobuilduht -skipactionhistory -CopyAppBundleBackToDevice";
 
 			if(Precompiled)
 			{
@@ -1694,10 +1625,6 @@ public class GUBP : BuildCommand
 
             foreach (var Kind in BranchInfo.MonolithicKinds)
             {
-                if (GUBP.bBuildRocket && Kind != TargetRules.TargetType.Game)
-                {
-                    continue;
-                }
                 if (GameProj.Properties.Targets.ContainsKey(Kind))
                 {
                     var Target = GameProj.Properties.Targets[Kind];
@@ -1721,36 +1648,6 @@ public class GUBP : BuildCommand
 							
 							foreach (var Config in Configs)
 							{
-								if (GUBP.bBuildRocket)
-								{
-									if (HostPlatform == UnrealTargetPlatform.Win64)
-									{
-										if (TargetPlatform == UnrealTargetPlatform.Win32 && Config != UnrealTargetConfiguration.Shipping)
-										{
-											continue;
-										}
-										if (TargetPlatform == UnrealTargetPlatform.Win64 && Config != UnrealTargetConfiguration.Development)
-										{
-											continue;
-										}
-										if (TargetPlatform == UnrealTargetPlatform.Android && Config != UnrealTargetConfiguration.Shipping && Config != UnrealTargetConfiguration.Development)
-										{
-											continue;
-										}
-										if (TargetPlatform == UnrealTargetPlatform.HTML5 && Config != UnrealTargetConfiguration.Shipping && Config != UnrealTargetConfiguration.Development)
-										{
-											continue;
-										}
-										if (TargetPlatform == UnrealTargetPlatform.Linux && Config != UnrealTargetConfiguration.Shipping && Config != UnrealTargetConfiguration.Development)
-										{
-											continue;
-										}
-									}
-									else if (Config != UnrealTargetConfiguration.Shipping && Config != UnrealTargetConfiguration.Development)
-									{
-										continue;
-									}
-								}
 								if (GameProj.GameName == bp.Branch.BaseEngineProject.GameName)
 								{
 									Agenda.AddTargets(new string[] { Target.TargetName }, TargetPlatform, Config, InAddArgs: Args);
@@ -2384,10 +2281,6 @@ public class GUBP : BuildCommand
                 //shared promotable has a shorter name
                 CompleteLabelPrefix = "Promotable";
             }                   
-            if (GUBP.bBuildRocket)
-            {
-                CompleteLabelPrefix = "Rocket-" + CompleteLabelPrefix;
-            }
             if (GUBP.bPreflightBuild)
             {
                 CompleteLabelPrefix = CompleteLabelPrefix + PreflightMangleSuffix;
@@ -3523,7 +3416,7 @@ public class GUBP : BuildCommand
 			var Build = new UE4Build(bp);
 			var Agenda = new UE4Build.BuildAgenda();
 
-			string AddArgs = "-nobuilduht" + bp.RocketUBTArgs();
+			string AddArgs = "-nobuilduht";
 			if (bp.bOrthogonalizeEditorPlatforms)
 			{
 				AddArgs += " -skipnonhostplatforms";
@@ -5187,18 +5080,11 @@ public class GUBP : BuildCommand
 			HostPlatforms.Add(UnrealTargetPlatform.Linux);
 		}
 
-		bBuildRocket = ParseParam("BuildRocket");		
-		bBuildOnlySamples = ParseParam("OnlySamples");
         bForceIncrementalCompile = ParseParam("ForceIncrementalCompile");
         bool bNoAutomatedTesting = ParseParam("NoAutomatedTesting") || BranchOptions.bNoAutomatedTesting;		
         StoreName = ParseParamValue("Store");
         string StoreSuffix = ParseParamValue("StoreSuffix", "");
 
-        if (bBuildRocket)
-        {
-            StoreSuffix = StoreSuffix + "-Rkt";
-        }
-		
         if (bPreflightBuild)
         {
             int PreflightUID = ParseParamInt("PreflightUID", 0);
@@ -5418,7 +5304,6 @@ public class GUBP : BuildCommand
         Log("************************* bFakeEC:           		        {0}", bFakeEC);
         Log("************************* bHistory:           		        {0}", bHistory);
         Log("************************* TimeIndex:           		    {0}", TimeIndex);        
-        Log("************************* bBuildRocket:           		    {0}", bBuildRocket);
 
 
         GUBPNodes = new Dictionary<string, GUBPNode>();
@@ -5456,23 +5341,6 @@ public class GUBP : BuildCommand
         {
             ActivePlatforms     = new List<UnrealTargetPlatform>(CommandUtils.KnownTargetPlatforms);
         }
-        if (bBuildRocket)
-        {
-            var FilteredActivePlatforms = new List<UnrealTargetPlatform>();
-            foreach (var Plat in ActivePlatforms)
-            {
-                if (
-                    (Plat != UnrealTargetPlatform.PS4 || ParseParam("WithPS4")) &&
-                    Plat != UnrealTargetPlatform.WinRT &&
-                    Plat != UnrealTargetPlatform.WinRT_ARM &&
-                    Plat != UnrealTargetPlatform.XboxOne
-                    )
-                {
-                    FilteredActivePlatforms.Add(Plat);
-                }
-            }
-            ActivePlatforms = FilteredActivePlatforms;
-        }
         var SupportedPlatforms = new List<UnrealTargetPlatform>();
         foreach(var Plat in ActivePlatforms)
         {
@@ -5662,11 +5530,6 @@ public class GUBP : BuildCommand
 
             DoASharedPromotable = NumSharedCode > 0 || NonCodeProjectNames.Count > 0 || NonCodeFormalBuilds.Count > 0;
 
-            if (!DoASharedPromotable && bBuildRocket)
-            {
-                throw new AutomationException("we were asked to make a rocket build, but this branch does not have a shared promotable.");
-            }
-
 			AddNode(new NonUnityTestNode(HostPlatform));
 
             if (DoASharedPromotable)
@@ -6379,28 +6242,6 @@ public class GUBP : BuildCommand
                     bOnlyNode = true;
                 }
             }
-            if (String.IsNullOrEmpty(NodeSpec) && bBuildRocket)
-            {
-                // rocket is the shared promotable plus some other stuff and nothing else
-                bRelatedToNode = true;
-                NodeSpec = "";
-				if(!bBuildOnlySamples)
-				{
-					NodeSpec = "Rocket_Aggregate+";
-				}
-				if (!ParseParam("nosamples"))
-				{
-					NodeSpec = NodeSpec + "Sample_Aggregate";
-				}
-                if (!ParseParam("RocketValidate"))
-                {
-                    //build a full promotable NodeSpec = SharedAggregatePromotableNode.StaticGetFullName() + "+" + NodeSpec;
-                }
-                else if (!bFake)
-                {
-                    NodeSpec = "Rocket_MakeBuild+" + NodeSpec;
-                }
-            }
             if (!String.IsNullOrEmpty(NodeSpec))
             {
 				bGraphSubset = true;
