- Fixed packaging error [UE-12571]

diff --git a/Engine/Source/Programs/AutomationTool/Win/WinPlatformResources.Automation.cs b/Engine/Source/Programs/AutomationTool/Win/WinPlatformResources.Automation.cs
--- a/Engine/Source/Programs/AutomationTool/Win/WinPlatformResources.Automation.cs
+++ b/Engine/Source/Programs/AutomationTool/Win/WinPlatformResources.Automation.cs
@@ -90,22 +90,29 @@ class GroupIconResource
 
 	public static GroupIconResource FromByteArray(byte[] GroupData, ModuleResourceLibary Module)
 	{
-		using(MemoryStream Input = new MemoryStream(GroupData))
+		if (GroupData != null)
 		{
-			BinaryReader Reader = new BinaryReader(Input);
-
-			Reader.ReadUInt16(); // Reserved
-			Reader.ReadUInt16(); // Type
-			ushort Count = Reader.ReadUInt16();
-
-			IconResource[] Icons = new IconResource[Count];
-			for(int Idx = 0; Idx < Count; Idx++)
+			using (MemoryStream Input = new MemoryStream(GroupData))
 			{
-				Icons[Idx] = ReadIconHeader(Reader);
-				int IconId = Reader.ReadUInt16();
-				Icons[Idx].Data = Module.ReadResource(IconId, ResourceType.Icon);
+				BinaryReader Reader = new BinaryReader(Input);
+
+				Reader.ReadUInt16(); // Reserved
+				Reader.ReadUInt16(); // Type
+				ushort Count = Reader.ReadUInt16();
+
+				IconResource[] Icons = new IconResource[Count];
+				for (int Idx = 0; Idx < Count; Idx++)
+				{
+					Icons[Idx] = ReadIconHeader(Reader);
+					int IconId = Reader.ReadUInt16();
+					Icons[Idx].Data = Module.ReadResource(IconId, ResourceType.Icon);
+				}
+				return new GroupIconResource(Icons);
 			}
-			return new GroupIconResource(Icons);
+		}
+		else
+		{
+			return null;
 		}
 	}
 
