Allow disabling spew when deleting changelists.

diff --git a/Engine/Source/Programs/AutomationTool/P4Utils.cs b/Engine/Source/Programs/AutomationTool/P4Utils.cs
--- a/Engine/Source/Programs/AutomationTool/P4Utils.cs
+++ b/Engine/Source/Programs/AutomationTool/P4Utils.cs
@@ -1557,7 +1557,7 @@ namespace AutomationTool
 		/// </summary>
 		/// <param name="CL">Changelist to delete.</param>
 		/// <param name="RevertFiles">Indicates whether files in that changelist should be reverted.</param>
-		public void DeleteChange(int CL, bool RevertFiles = true, bool SpewIsVerbose = false)
+		public void DeleteChange(int CL, bool RevertFiles = true, bool SpewIsVerbose = false, bool AllowSpew = true)
 		{
 			CheckP4Enabled();
 			if (RevertFiles)
@@ -1566,7 +1566,7 @@ namespace AutomationTool
 			}
 
 			string CmdOutput;
-			if (LogP4Output(out CmdOutput, String.Format("change -d {0}", CL), SpewIsVerbose: SpewIsVerbose))
+			if (LogP4Output(out CmdOutput, String.Format("change -d {0}", CL), SpewIsVerbose: SpewIsVerbose, AllowSpew: AllowSpew))
 			{
 				string EndStr = " deleted.";
 				string ChangeStr = "Change ";
