Fix various edge cases not handled correctly by FileFilter:

* Filters were not earlying out of recursive directory searches in cases where a directory partially matched a file specification.
* Trailing file wildcards would match against directory names, so -/Something* would exclude /Something/...
* Path specifications starting with ... would have an additional ... token appended.

diff --git a/Engine/Source/Programs/AutomationTool/FileFilter.cs b/Engine/Source/Programs/AutomationTool/FileFilter.cs
--- a/Engine/Source/Programs/AutomationTool/FileFilter.cs
+++ b/Engine/Source/Programs/AutomationTool/FileFilter.cs
@@ -29,6 +29,11 @@ public class FileFilter
 	/// </summary>
 	FileFilterNode RootNode;
 
+	/// <summary>
+	/// The default node, which will match any path
+	/// </summary>
+	FileFilterNode DefaultNode;
+
 	/// <summary>
 	/// Terminating nodes for each rule added to the filter
 	/// </summary>
@@ -37,9 +42,12 @@ public class FileFilter
 	/// <summary>
 	/// Default constructor
 	/// </summary>
-	public FileFilter()
+	public FileFilter(FileFilterType DefaultType = FileFilterType.Exclude)
 	{
 		RootNode = new FileFilterNode(null, "");
+
+		DefaultNode = new FileFilterNode(RootNode, "...");
+		DefaultNode.Type = DefaultType;
 	}
 
 	/// <summary>
@@ -47,7 +55,7 @@ public class FileFilter
 	/// </summary>
 	/// <param name="Other">Filter to copy from</param>
 	public FileFilter(FileFilter Other)
-		: this()
+		: this(Other.DefaultNode.Type)
 	{
 		foreach (FileFilterNode OtherRule in Other.Rules)
 		{
@@ -239,7 +247,7 @@ public class FileFilter
 		{
 			NormalizedPattern = NormalizedPattern.Substring(1);
 		}
-		else
+		else if(!NormalizedPattern.StartsWith("..."))
 		{
 			NormalizedPattern = ".../" + NormalizedPattern;
 		}
@@ -312,7 +320,7 @@ public class FileFilter
 	{
 		string[] Tokens = FileName.TrimStart('/', '\\').Split('/', '\\');
 
-		FileFilterNode MatchingNode = FindMatchingNode(RootNode, Tokens, 0, RootNode);
+		FileFilterNode MatchingNode = FindMatchingNode(RootNode, Tokens, 0, DefaultNode);
 
 		return MatchingNode.Type == FileFilterType.Include;
 	}
@@ -326,7 +334,7 @@ public class FileFilter
 	{
 		string[] Tokens = FolderName.Trim('/', '\\').Split('/', '\\');
 
-		FileFilterNode MatchingNode = FindMatchingNode(RootNode, Tokens, 0, RootNode);
+		FileFilterNode MatchingNode = FindMatchingNode(RootNode, Tokens.Union(new string[]{ "" }).ToArray(), 0, DefaultNode);
 
 		return MatchingNode.Type == FileFilterType.Include || HighestPossibleIncludeMatch(RootNode, Tokens, 0, MatchingNode.RuleNumber) > MatchingNode.RuleNumber;
 	}
@@ -408,14 +416,14 @@ public class FileFilter
 		// If there is no rule under the current node which is better than the current best node, early out
 		if (CurrentBestNode.Type == FileFilterType.Include)
 		{
-			if (CurrentNode.MaxExcludeRuleNumber < CurrentBestNode.RuleNumber)
+			if (CurrentNode.MaxExcludeRuleNumber <= CurrentBestNode.RuleNumber)
 			{
 				return CurrentBestNode;
 			}
 		}
 		else
 		{
-			if (CurrentNode.MaxIncludeRuleNumber < CurrentBestNode.RuleNumber)
+			if (CurrentNode.MaxIncludeRuleNumber <= CurrentBestNode.RuleNumber)
 			{
 				return CurrentBestNode;
 			}
@@ -467,7 +475,7 @@ public class FileFilter
 			{
 				if (Branch.Pattern == "...")
 				{
-					if (Branch.MaxIncludeRuleNumber > CurrentBestRuleNumber)
+					if (Branch.MaxIncludeRuleNumber > BestRuleNumber)
 					{
 						BestRuleNumber = Branch.MaxIncludeRuleNumber;
 					}
@@ -503,12 +511,12 @@ class FileFilterNode
 	/// <summary>
 	/// Highest include rule number matched by this node or any child nodes.
 	/// </summary>
-	public int MaxIncludeRuleNumber;
+	public int MaxIncludeRuleNumber = -1;
 
 	/// <summary>
 	/// Highest exclude rule number matched by this node or any child nodes.
 	/// </summary>
-	public int MaxExcludeRuleNumber;
+	public int MaxExcludeRuleNumber = -1;
 
 	/// <summary>
 	/// Child branches of this node, distinct by the the pattern for each.
