Fix to make sure PS4/XboxOne editor files are excluded from Rocket builds.

diff --git a/Engine/Source/Programs/AutomationTool/Scripts/RocketBuild.Automation.cs b/Engine/Source/Programs/AutomationTool/Scripts/RocketBuild.Automation.cs
--- a/Engine/Source/Programs/AutomationTool/Scripts/RocketBuild.Automation.cs
+++ b/Engine/Source/Programs/AutomationTool/Scripts/RocketBuild.Automation.cs
@@ -260,7 +260,7 @@ namespace Rocket
 			CopyNode.AddFilesToFilter(PromotableFilter, AllDependencyBuildProducts, CommandUtils.CmdEnv.LocalRoot, FileFilterType.Include);
 			PromotableFilter.AddRulesFromFile(CommandUtils.CombinePaths(CommandUtils.CmdEnv.LocalRoot, ConfigRelativePath), "promotable");
 			CopyNode.ExcludeConfidentialFolders(PromotableFilter);
-
+			
 			// Copy everything that matches the filter to the promotion folder
 			string PromotableFolder = CommandUtils.CombinePaths(CommandUtils.CmdEnv.LocalRoot, "Engine", "Saved", "GitPromotable");
 			CommandUtils.DeleteDirectoryContents(PromotableFolder);
@@ -520,6 +520,7 @@ namespace Rocket
 			Filter.AddRulesFromFile(RulesFileName, "CopyEditor", HostPlatform.ToString());
 
 			// Add the final exclusions for legal reasons.
+			ExcludeConfidentialPlatforms(Filter);
 			ExcludeConfidentialFolders(Filter);
 
 			// Wipe the target directory and copy all the files over
