Merge GetExternalFileListNode with GamePlatformMonolithicsNode when running precompiled. There's a logical dependency and they share a build agenda, so it's innefficient to keep them separate.

diff --git a/Engine/Source/Programs/AutomationTool/Scripts/GUBP.Automation.cs b/Engine/Source/Programs/AutomationTool/Scripts/GUBP.Automation.cs
--- a/Engine/Source/Programs/AutomationTool/Scripts/GUBP.Automation.cs
+++ b/Engine/Source/Programs/AutomationTool/Scripts/GUBP.Automation.cs
@@ -1596,6 +1596,38 @@ public class GUBP : BuildCommand
 			return false;
 		}
 
+		public override void DoBuild(GUBP bp)
+		{
+			base.DoBuild(bp);
+
+			if(Precompiled)
+			{
+				// Get a list of all the build dependencies
+				UE4Build.BuildAgenda Agenda = GetAgenda(bp);
+				string FileListPath = new UE4Build(bp).GenerateExternalFileList(Agenda);
+				UnrealBuildTool.ExternalFileList FileList = UnrealBuildTool.Utils.ReadClass<UnrealBuildTool.ExternalFileList>(FileListPath);
+
+				// Make all the paths relative to the root
+				string FilterPrefix = CommandUtils.CombinePaths(PathSeparator.Slash, CommandUtils.CmdEnv.LocalRoot).TrimEnd('/') + "/";
+				for(int Idx = 0; Idx < FileList.FileNames.Count; Idx++)
+				{
+					if(FileList.FileNames[Idx].StartsWith(FilterPrefix, StringComparison.InvariantCultureIgnoreCase))
+					{
+						FileList.FileNames[Idx] = FileList.FileNames[Idx].Substring(FilterPrefix.Length);
+					}
+					else
+					{
+						CommandUtils.LogError("Referenced external file is not under local root: {0}", FileList.FileNames[Idx]);
+					}
+				}
+
+				// Write the resulting file list out to disk
+				string OutputFileListPath = StaticGetBuildDependenciesPath(HostPlatform, TargetPlatform);
+				UnrealBuildTool.Utils.WriteClass<UnrealBuildTool.ExternalFileList>(FileList, OutputFileListPath, "");
+				AddBuildProduct(OutputFileListPath);
+			}
+		}
+
         public override UE4Build.BuildAgenda GetAgenda(GUBP bp)
         {
             if (!bp.ActivePlatforms.Contains(TargetPlatform))
@@ -1664,6 +1696,11 @@ public class GUBP : BuildCommand
 
             return Agenda;
         }
+
+		public static string StaticGetBuildDependenciesPath(UnrealTargetPlatform HostPlatform, UnrealTargetPlatform TargetPlatform)
+		{
+			return CommandUtils.CombinePaths(CommandUtils.CmdEnv.LocalRoot, "Engine/Saved/BuildDependencies/" + TargetPlatform.ToString() + StaticGetHostPlatformSuffix(HostPlatform) + ".xml");
+		}
     }
 
 	public static class HeadersNode
diff --git a/Engine/Source/Programs/AutomationTool/Scripts/RocketBuild.Automation.cs b/Engine/Source/Programs/AutomationTool/Scripts/RocketBuild.Automation.cs
--- a/Engine/Source/Programs/AutomationTool/Scripts/RocketBuild.Automation.cs
+++ b/Engine/Source/Programs/AutomationTool/Scripts/RocketBuild.Automation.cs
@@ -103,16 +103,6 @@ namespace Rocket
 				// Add a dependency on this platform finishing
 				RocketAggregateNode AggNode = (RocketAggregateNode)bp.FindNode(RocketAggregateNode.StaticGetFullName());
 
-				// Get file lists for each of the target platforms
-				foreach(UnrealTargetPlatform CodeTargetPlatform in CodeTargetPlatforms)
-				{
-					UnrealTargetPlatform SourceHostPlatform = GetSourceHostPlatform(bp, HostPlatform, CodeTargetPlatform);
-					if(!bp.HasNode(GetExternalFileListNode.StaticGetFullName(SourceHostPlatform, CodeTargetPlatform)))
-					{
-						bp.AddNode(new GetExternalFileListNode(SourceHostPlatform, CodeTargetPlatform));
-					}
-				}
-
 				// Copy to output folders
 				AggNode.AddDependency(bp.AddNode(new CopyInstalledEditorNode(bp, HostPlatform, InstallDir)));
 				AggNode.AddDependency(bp.AddNode(new CopyInstalledPlatformsNode(bp, HostPlatform, TargetPlatforms, CodeTargetPlatforms, InstallDir, CurrentFeaturePacks, CurrentTemplates)));
@@ -570,7 +560,6 @@ namespace Rocket
 			{
 				UnrealTargetPlatform SourceHostPlatform = GetSourceHostPlatform(bp, HostPlatform, CodeTargetPlatform);
 				AddDependency(GUBP.GamePlatformMonolithicsNode.StaticGetFullName(SourceHostPlatform, bp.Branch.BaseEngineProject, CodeTargetPlatform, Precompiled: true));
-				AddDependency(GetExternalFileListNode.StaticGetFullName(SourceHostPlatform, CodeTargetPlatform));
 			}
 
 			// Add win64 tools on Mac, to get the win64 build of IPP
@@ -613,7 +602,7 @@ namespace Rocket
 			foreach(UnrealTargetPlatform CodeTargetPlatform in CodeTargetPlatforms)
 			{
 				UnrealTargetPlatform SourceHostPlatform = GetSourceHostPlatform(bp, HostPlatform, CodeTargetPlatform);
-				string FileListPath = GetExternalFileListNode.StaticGetFileListPath(SourceHostPlatform, CodeTargetPlatform);
+				string FileListPath = GUBP.GamePlatformMonolithicsNode.StaticGetBuildDependenciesPath(SourceHostPlatform, CodeTargetPlatform);
 				AddFilesToFilter(Filter, UnrealBuildTool.Utils.ReadClass<UnrealBuildTool.ExternalFileList>(FileListPath).FileNames, FileFilterType.Include);
 			}
 
@@ -688,66 +677,6 @@ namespace Rocket
 		}
 	}
 
-	class GetExternalFileListNode : GUBP.HostPlatformNode
-	{
-		UnrealTargetPlatform TargetPlatform;
-
-		public GetExternalFileListNode(UnrealTargetPlatform InHostPlatform, UnrealTargetPlatform InTargetPlatform)
-			: base(InHostPlatform)
-		{
-			TargetPlatform = InTargetPlatform;
-			AgentSharingGroup = "UE4_" + TargetPlatform.ToString() + "_Mono" + StaticGetHostPlatformSuffix(HostPlatform);
-		}
-		public static string StaticGetFullName(UnrealTargetPlatform HostPlatform, UnrealTargetPlatform TargetPlatform)
-		{
-            return "UE4_" + TargetPlatform.ToString() + "_GetExternalFiles" + StaticGetHostPlatformSuffix(HostPlatform);
-		}
-		public override string GetFullName()
-		{
-			return StaticGetFullName(HostPlatform, TargetPlatform);
-		}
-		public override void DoBuild(GUBP bp)
-		{
-			// Prepare an agenda of targets that need to be buildable under Rocket.
-			UE4Build.BuildAgenda Agenda = new UE4Build.BuildAgenda();
-
-			// Add the game dependencies for UE4Game on all the target platforms. We don't link libraries with their dependencies, so we can use UE4Game and -buildrocket 
-			// to pick up all the default plugins.
-			Agenda.AddTargets(new string[] { "UE4Game" }, TargetPlatform, UnrealTargetConfiguration.Development, null, false, false, false, "-precompile");
-			Agenda.AddTargets(new string[] { "UE4Game" }, TargetPlatform, UnrealTargetConfiguration.Shipping, null, false, false, false, "-precompile");
-
-			// Read the file list
-			string FileListPath = new UE4Build(bp).GenerateExternalFileList(Agenda);
-			UnrealBuildTool.ExternalFileList FileList = UnrealBuildTool.Utils.ReadClass<UnrealBuildTool.ExternalFileList>(FileListPath);
-
-			// Make all the paths relative to the root
-			string FilterPrefix = CommandUtils.CombinePaths(PathSeparator.Slash, CommandUtils.CmdEnv.LocalRoot).TrimEnd('/') + "/";
-			for(int Idx = 0; Idx < FileList.FileNames.Count; Idx++)
-			{
-				if(FileList.FileNames[Idx].StartsWith(FilterPrefix, StringComparison.InvariantCultureIgnoreCase))
-				{
-					FileList.FileNames[Idx] = FileList.FileNames[Idx].Substring(FilterPrefix.Length);
-				}
-				else
-				{
-					CommandUtils.LogError("Referenced external file is not under local root: {0}", FileList.FileNames[Idx]);
-				}
-			}
-
-			// Write the resulting file list out to disk
-			string OutputFileListPath = StaticGetFileListPath(HostPlatform, TargetPlatform);
-			UnrealBuildTool.Utils.WriteClass<UnrealBuildTool.ExternalFileList>(FileList, OutputFileListPath, "");
-
-			// Add it to the build products
-			BuildProducts = new List<string>();
-			AddBuildProduct(OutputFileListPath);
-		}
-		public static string StaticGetFileListPath(UnrealTargetPlatform HostPlatform, UnrealTargetPlatform TargetPlatform)
-		{
-			return CommandUtils.CombinePaths(CommandUtils.CmdEnv.LocalRoot, "Engine/Saved/ExternalFileList/" + TargetPlatform.ToString() + StaticGetHostPlatformSuffix(HostPlatform) + ".xml");
-		}
-	}
-
 	public class CopyInstalledDerivedDataNode : GUBP.HostPlatformNode
 	{
 		string TargetPlatforms;
