Add the BuildInstalledEngine aggregate as a requirement for the promotable; it messes up the promotable label trigger if everything isn't triggered afterwards.

diff --git a/Engine/Source/Programs/AutomationTool/Scripts/RocketBuild.Automation.cs b/Engine/Source/Programs/AutomationTool/Scripts/RocketBuild.Automation.cs
--- a/Engine/Source/Programs/AutomationTool/Scripts/RocketBuild.Automation.cs
+++ b/Engine/Source/Programs/AutomationTool/Scripts/RocketBuild.Automation.cs
@@ -111,6 +111,7 @@ namespace Rocket
 
 				// Add the aggregate node for the entire install
 				AggNode.AddDependency(bp.AddNode(new BuildInstalledEngineNode(HostPlatform, InstallDir, SymbolsDir)));
+				PromotableNode.AddDependency(BuildInstalledEngineNode.StaticGetFullName(HostPlatform));
 
 				// Add a node for GitHub promotions
 				if(HostPlatform == UnrealTargetPlatform.Win64)
@@ -568,8 +569,6 @@ namespace Rocket
 				AddDependency(GUBP.ToolsNode.StaticGetFullName(UnrealTargetPlatform.Win64));
 				AddDependency(GUBP.ToolsForCompileNode.StaticGetFullName(UnrealTargetPlatform.Win64));
 			}
-
-			AgentSharingGroup = "BuildInstall" + StaticGetHostPlatformSuffix(HostPlatform);
 		}
 
 		public static UnrealTargetPlatform GetSourceHostPlatform(GUBP bp, UnrealTargetPlatform HostPlatform, UnrealTargetPlatform TargetPlatform)
