Remove the "Distill" program, which isn't being used any more.

diff --git a/Engine/Source/Programs/AutomationTool/Scripts/GUBP.Automation.cs b/Engine/Source/Programs/AutomationTool/Scripts/GUBP.Automation.cs
--- a/Engine/Source/Programs/AutomationTool/Scripts/GUBP.Automation.cs
+++ b/Engine/Source/Programs/AutomationTool/Scripts/GUBP.Automation.cs
@@ -1150,7 +1150,6 @@ public class GUBP : BuildCommand
                 Agenda.DotNetProjects.AddRange(
                     new string[] 
 			    {
-				    CombinePaths(@"Engine\Source\Programs\Distill\Distill.csproj"),		  
                     CombinePaths(@"Engine\Source\Programs\NotForLicensees\CrashReportServer\CrashReportCommon\CrashReportCommon.csproj"),
 					CombinePaths(@"Engine\Source\Programs\NotForLicensees\CrashReportServer\CrashReportReceiver\CrashReportReceiver.csproj"),
 					CombinePaths(@"Engine\Source\Programs\NotForLicensees\CrashReportServer\CrashReportProcess\CrashReportProcess.csproj"),
diff --git a/Engine/Source/Programs/Distill/Distill.csproj b/Engine/Source/Programs/Distill/Distill.csproj
deleted file mode 100644
diff --git a/Engine/Source/Programs/Distill/Private/Crypto.cs b/Engine/Source/Programs/Distill/Private/Crypto.cs
deleted file mode 100644
diff --git a/Engine/Source/Programs/Distill/Private/Distill.cs b/Engine/Source/Programs/Distill/Private/Distill.cs
deleted file mode 100644
diff --git a/Engine/Source/Programs/Distill/Private/FTP.cs b/Engine/Source/Programs/Distill/Private/FTP.cs
deleted file mode 100644
diff --git a/Engine/Source/Programs/Distill/Private/FileUtils.cs b/Engine/Source/Programs/Distill/Private/FileUtils.cs
deleted file mode 100644
diff --git a/Engine/Source/Programs/Distill/Private/Selection.cs b/Engine/Source/Programs/Distill/Private/Selection.cs
deleted file mode 100644
diff --git a/Engine/Source/Programs/Distill/Private/Settings.cs b/Engine/Source/Programs/Distill/Private/Settings.cs
deleted file mode 100644
diff --git a/Engine/Source/Programs/Distill/Private/TOC.cs b/Engine/Source/Programs/Distill/Private/TOC.cs
deleted file mode 100644
diff --git a/Engine/Source/Programs/Distill/Private/Zip.cs b/Engine/Source/Programs/Distill/Private/Zip.cs
deleted file mode 100644
diff --git a/Engine/Source/Programs/Distill/Properties/AssemblyInfo.cs b/Engine/Source/Programs/Distill/Properties/AssemblyInfo.cs
deleted file mode 100644
diff --git a/Engine/Source/Programs/UnrealBuildTool/System/ProjectFileGenerator.cs b/Engine/Source/Programs/UnrealBuildTool/System/ProjectFileGenerator.cs
--- a/Engine/Source/Programs/UnrealBuildTool/System/ProjectFileGenerator.cs
+++ b/Engine/Source/Programs/UnrealBuildTool/System/ProjectFileGenerator.cs
@@ -511,9 +511,6 @@ namespace UnrealBuildTool
 					// Add automation.csproj files to the master project
 					AddAutomationModules(ProgramsFolder);
 
-					// Add Distill to the master project
-					ProgramsFolder.ChildProjects.Add(AddSimpleCSharpProject("Distill"));
-
 					// Add DotNETUtilities to the master project
 					ProgramsFolder.ChildProjects.Add(AddSimpleCSharpProject("DotNETCommon/DotNETUtilities", bShouldBuildForAllSolutionTargets: true, bForceDevelopmentConfiguration: true));
 
