Better sorting for the nodes to display in EC. Primarily sorts by frequency then by execution order, but tries to keep the same nodes on paralllel paths together (by grouping nodes which return the same thing from GUBPNode.GetDisplayGroupName()).

diff --git a/Engine/Source/Programs/AutomationTool/Scripts/GUBP.Automation.cs b/Engine/Source/Programs/AutomationTool/Scripts/GUBP.Automation.cs
--- a/Engine/Source/Programs/AutomationTool/Scripts/GUBP.Automation.cs
+++ b/Engine/Source/Programs/AutomationTool/Scripts/GUBP.Automation.cs
@@ -364,6 +364,10 @@ public class GUBP : BuildCommand
         {
             throw new AutomationException("Unimplemented GetFullName.");
         }
+		public virtual string GetDisplayGroupName()
+		{
+			return GetFullName();
+		}
         public virtual string GameNameIfAnyForTempStorage()
         {
             return "";
@@ -617,6 +621,12 @@ public class GUBP : BuildCommand
         {
             HostPlatform = InHostPlatform;
         }
+		public override string GetDisplayGroupName()
+		{
+			string Name = GetFullName();
+			string Suffix = GetHostPlatformSuffix();
+			return Name.EndsWith(Suffix)? Name.Substring(0, Name.Length - Suffix.Length) : Name;
+		}
         public static string StaticGetHostPlatformSuffix(UnrealTargetPlatform InHostPlatform, UnrealTargetPlatform InAgentPlatform = UnrealTargetPlatform.Unknown)
         {
             if (InHostPlatform == UnrealTargetPlatform.Mac)
@@ -1518,9 +1528,9 @@ public class GUBP : BuildCommand
             }
             else
             {
-                if (TargetPlatform != InHostPlatform && bp.HasNode(GamePlatformMonolithicsNode.StaticGetFullName(InHostPlatform, bp.Branch.BaseEngineProject, InHostPlatform)))
+                if (TargetPlatform != InHostPlatform && bp.HasNode(GamePlatformMonolithicsNode.StaticGetFullName(InHostPlatform, bp.Branch.BaseEngineProject, InHostPlatform, Precompiled: Precompiled)))
                 {
-                    AddPseudodependency(GamePlatformMonolithicsNode.StaticGetFullName(InHostPlatform, bp.Branch.BaseEngineProject, InHostPlatform));
+                    AddPseudodependency(GamePlatformMonolithicsNode.StaticGetFullName(InHostPlatform, bp.Branch.BaseEngineProject, InHostPlatform, Precompiled: Precompiled));
                 }
             }
 			if(!WithXp && !Precompiled && HasPrecompiledTargets(GameProj, HostPlatform, TargetPlatform))
@@ -1535,6 +1545,11 @@ public class GUBP : BuildCommand
 
         }
 
+		public override string GetDisplayGroupName()
+		{
+			return GameProj.GameName + "_Monolithics" + (Precompiled? "_Precompiled" : "");
+		}
+
         public static string StaticGetFullName(UnrealTargetPlatform InHostPlatform, BranchInfo.BranchUProject InGameProj, UnrealTargetPlatform InTargetPlatform, bool WithXp = false, bool Precompiled = false)
         {
 			string Name = InGameProj.GameName;
@@ -5017,7 +5032,6 @@ public class GUBP : BuildCommand
 		}
 	}
 
-
     [Help("Runs one, several or all of the GUBP nodes")]
     [Help(typeof(UE4Build))]
     [Help("NoMac", "Toggle to exclude the Mac host platform, default is Win64+Mac+Linux")]
@@ -6198,7 +6212,6 @@ public class GUBP : BuildCommand
         GUBPNodesControllingTriggerDotName = new Dictionary<string, string>();
 
         var FullNodeList = new Dictionary<string, string>();
-        var FullNodeListSortKey = new Dictionary<string, int>();
         var FullNodeDirectDependencies = new Dictionary<string, string>();
 		var FullNodeDependedOnBy = new Dictionary<string, string>();
 		var FullNodeDependentPromotions = new Dictionary<string, string>();		
@@ -6206,7 +6219,6 @@ public class GUBP : BuildCommand
         {
             Log("******* {0} GUBP Nodes", GUBPNodes.Count);
             var SortedNodes = TopologicalSort(new HashSet<string>(GUBPNodes.Keys), LocalOnly: true, DoNotConsiderCompletion: true);
-            int Count = 0;
             foreach (var Node in SortedNodes)
             {
                 string Note = GetControllingTriggerDotName(Node);
@@ -6233,8 +6245,6 @@ public class GUBP : BuildCommand
                     Log("  {0}: {1}      {2}", Node, Note, All);
                     FullNodeList.Add(Node, Note);
                     FullNodeDirectDependencies.Add(Node, All);
-                    FullNodeListSortKey.Add(Node, Count);					
-                    Count++;
                 }
                 else
                 {
@@ -6242,6 +6252,7 @@ public class GUBP : BuildCommand
                 }
             }
         }
+		Dictionary<string, int> FullNodeListSortKey = GetDisplayOrder(FullNodeList.Keys.ToList(), FullNodeDirectDependencies, GUBPNodes);
 
         bool bOnlyNode = false;
         bool bRelatedToNode = false;
@@ -7427,6 +7438,86 @@ public class GUBP : BuildCommand
 
         PrintRunTime();
     }
+
+	/// <summary>
+	/// Sorts a list of nodes to display in EC. The default order is based on execution order and agent groups, whereas this function arranges nodes by
+	/// frequency then execution order, while trying to group nodes on parallel paths (eg. Mac/Windows editor nodes) together.
+	/// </summary>
+	static Dictionary<string, int> GetDisplayOrder(List<string> NodeNames, Dictionary<string, string> InitialNodeDependencyNames, Dictionary<string, GUBPNode> GUBPNodes)
+	{
+		// Split the nodes into separate lists for each frequency
+		SortedDictionary<int, List<string>> NodesByFrequency = new SortedDictionary<int,List<string>>();
+		foreach(string NodeName in NodeNames)
+		{
+			List<string> NodesByThisFrequency;
+			if(!NodesByFrequency.TryGetValue(GUBPNodes[NodeName].DependentCISFrequencyQuantumShift(), out NodesByThisFrequency))
+			{
+				NodesByThisFrequency = new List<string>();
+				NodesByFrequency.Add(GUBPNodes[NodeName].DependentCISFrequencyQuantumShift(), NodesByThisFrequency);
+			}
+			NodesByThisFrequency.Add(NodeName);
+		}
+
+		// Find a list of nodes in each display group. If the group name matches the node name, put that node at the front of the list.
+		Dictionary<string, string> DisplayGroups = new Dictionary<string,string>();
+		foreach(string NodeName in NodeNames)
+		{
+			string GroupName = GUBPNodes[NodeName].GetDisplayGroupName();
+			if(!DisplayGroups.ContainsKey(GroupName))
+			{
+				DisplayGroups.Add(GroupName, NodeName);
+			}
+			else if(GroupName == NodeName)
+			{
+				DisplayGroups[GroupName] = NodeName + " " + DisplayGroups[GroupName];
+			}
+			else
+			{
+				DisplayGroups[GroupName] = DisplayGroups[GroupName] + " " + NodeName;
+			}
+		}
+
+		// Build a list of ordering dependencies, putting all Mac nodes after Windows nodes with the same names.
+		Dictionary<string, string> NodeDependencyNames = new Dictionary<string,string>(InitialNodeDependencyNames);
+		foreach(KeyValuePair<string, string> DisplayGroup in DisplayGroups)
+		{
+			string[] GroupNodes = DisplayGroup.Value.Split(' ');
+			for(int Idx = 1; Idx < GroupNodes.Length; Idx++)
+			{
+				NodeDependencyNames[GroupNodes[Idx]] += " " + GroupNodes[0];
+			}
+		}
+
+		// Add nodes for each frequency into the master list, trying to match up different host platforms along the way
+		HashSet<string> VisitedNodes = new HashSet<string>();
+		Dictionary<string, int> SortedNodes = new Dictionary<string,int>();
+		foreach(List<string> NodesByThisFrequency in NodesByFrequency.Values)
+		{
+			foreach(string FirstNodeName in NodesByThisFrequency)
+			{
+				string[] GroupNodeNames = DisplayGroups[GUBPNodes[FirstNodeName].GetDisplayGroupName()].Split(' ');
+				foreach(string GroupNodeName in GroupNodeNames)
+				{
+					AddNodeAndDependencies(GroupNodeName, NodeDependencyNames, VisitedNodes, SortedNodes);
+				}
+			}
+		}
+		return SortedNodes;
+	}
+
+	static void AddNodeAndDependencies(string NodeName, Dictionary<string, string> NodeDependencyNames, HashSet<string> VisitedNodes, Dictionary<string, int> SortedNodes)
+	{
+		if(!VisitedNodes.Contains(NodeName))
+		{
+			VisitedNodes.Add(NodeName);
+			foreach(string NodeDependencyName in NodeDependencyNames[NodeName].Split(new char[]{ ' ' }, StringSplitOptions.RemoveEmptyEntries))
+			{
+				AddNodeAndDependencies(NodeDependencyName, NodeDependencyNames, VisitedNodes, SortedNodes);
+			}
+			SortedNodes.Add(NodeName, SortedNodes.Count);
+		}
+	}
+
     string StartedTempStorageSuffix = "_Started";
     string FailedTempStorageSuffix = "_Failed";
     string SucceededTempStorageSuffix = "_Succeeded";
