Don't build non-precompiled targets as often; they aren't the ones we need to ship.

diff --git a/Engine/Source/Programs/AutomationTool/Scripts/GUBP.Automation.cs b/Engine/Source/Programs/AutomationTool/Scripts/GUBP.Automation.cs
--- a/Engine/Source/Programs/AutomationTool/Scripts/GUBP.Automation.cs
+++ b/Engine/Source/Programs/AutomationTool/Scripts/GUBP.Automation.cs
@@ -1583,7 +1583,7 @@ public class GUBP : BuildCommand
         public override int CISFrequencyQuantumShift(GUBP bp)
         {            
             int Result = base.CISFrequencyQuantumShift(bp);
-            if(GameProj.GameName != bp.Branch.BaseEngineProject.GameName)
+            if(GameProj.GameName != bp.Branch.BaseEngineProject.GameName || !Precompiled)
             {
                 Result += 3; //only every 80m
             }
