Separate out node groups depending on the frequency they run.

diff --git a/Engine/Source/Programs/AutomationTool/Scripts/GUBP.Automation.cs b/Engine/Source/Programs/AutomationTool/Scripts/GUBP.Automation.cs
--- a/Engine/Source/Programs/AutomationTool/Scripts/GUBP.Automation.cs
+++ b/Engine/Source/Programs/AutomationTool/Scripts/GUBP.Automation.cs
@@ -7458,41 +7458,42 @@ public class GUBP : BuildCommand
 			NodesByThisFrequency.Add(NodeName);
 		}
 
-		// Find a list of nodes in each display group. If the group name matches the node name, put that node at the front of the list.
-		Dictionary<string, string> DisplayGroups = new Dictionary<string,string>();
-		foreach(string NodeName in NodeNames)
+		// Build the output list by scanning each frequency in order
+		HashSet<string> VisitedNodes = new HashSet<string>();
+		Dictionary<string, int> SortedNodes = new Dictionary<string,int>();
+		foreach(List<string> NodesByThisFrequency in NodesByFrequency.Values)
 		{
-			string GroupName = GUBPNodes[NodeName].GetDisplayGroupName();
-			if(!DisplayGroups.ContainsKey(GroupName))
+			// Find a list of nodes in each display group. If the group name matches the node name, put that node at the front of the list.
+			Dictionary<string, string> DisplayGroups = new Dictionary<string,string>();
+			foreach(string NodeName in NodesByThisFrequency)
 			{
-				DisplayGroups.Add(GroupName, NodeName);
-			}
-			else if(GroupName == NodeName)
-			{
-				DisplayGroups[GroupName] = NodeName + " " + DisplayGroups[GroupName];
-			}
-			else
-			{
-				DisplayGroups[GroupName] = DisplayGroups[GroupName] + " " + NodeName;
+				string GroupName = GUBPNodes[NodeName].GetDisplayGroupName();
+				if(!DisplayGroups.ContainsKey(GroupName))
+				{
+					DisplayGroups.Add(GroupName, NodeName);
+				}
+				else if(GroupName == NodeName)
+				{
+					DisplayGroups[GroupName] = NodeName + " " + DisplayGroups[GroupName];
+				}
+				else
+				{
+					DisplayGroups[GroupName] = DisplayGroups[GroupName] + " " + NodeName;
+				}
 			}
-		}
 
-		// Build a list of ordering dependencies, putting all Mac nodes after Windows nodes with the same names.
-		Dictionary<string, string> NodeDependencyNames = new Dictionary<string,string>(InitialNodeDependencyNames);
-		foreach(KeyValuePair<string, string> DisplayGroup in DisplayGroups)
-		{
-			string[] GroupNodes = DisplayGroup.Value.Split(' ');
-			for(int Idx = 1; Idx < GroupNodes.Length; Idx++)
+			// Build a list of ordering dependencies, putting all Mac nodes after Windows nodes with the same names.
+			Dictionary<string, string> NodeDependencyNames = new Dictionary<string,string>(InitialNodeDependencyNames);
+			foreach(KeyValuePair<string, string> DisplayGroup in DisplayGroups)
 			{
-				NodeDependencyNames[GroupNodes[Idx]] += " " + GroupNodes[0];
+				string[] GroupNodes = DisplayGroup.Value.Split(' ');
+				for(int Idx = 1; Idx < GroupNodes.Length; Idx++)
+				{
+					NodeDependencyNames[GroupNodes[Idx]] += " " + GroupNodes[0];
+				}
 			}
-		}
 
-		// Add nodes for each frequency into the master list, trying to match up different host platforms along the way
-		HashSet<string> VisitedNodes = new HashSet<string>();
-		Dictionary<string, int> SortedNodes = new Dictionary<string,int>();
-		foreach(List<string> NodesByThisFrequency in NodesByFrequency.Values)
-		{
+			// Add nodes for each frequency into the master list, trying to match up different groups along the way
 			foreach(string FirstNodeName in NodesByThisFrequency)
 			{
 				string[] GroupNodeNames = DisplayGroups[GUBPNodes[FirstNodeName].GetDisplayGroupName()].Split(' ');
