Decrease frequency of CopyEditor build steps; it isn't useful without DDC anyway.

diff --git a/Engine/Source/Programs/AutomationTool/Scripts/RocketBuild.Automation.cs b/Engine/Source/Programs/AutomationTool/Scripts/RocketBuild.Automation.cs
--- a/Engine/Source/Programs/AutomationTool/Scripts/RocketBuild.Automation.cs
+++ b/Engine/Source/Programs/AutomationTool/Scripts/RocketBuild.Automation.cs
@@ -500,7 +500,7 @@ namespace Rocket
 
 		public override int CISFrequencyQuantumShift(GUBP bp)
 		{
-			return base.CISFrequencyQuantumShift(bp);
+			return base.CISFrequencyQuantumShift(bp) + 2;
 		}
 
 		public static string StaticGetFullName(UnrealTargetPlatform InHostPlatform)
