Allow platform automation to modify cook options.
UEPLAT-685 WIP

diff --git a/Engine/Source/Programs/AutomationTool/Platform.cs b/Engine/Source/Programs/AutomationTool/Platform.cs
--- a/Engine/Source/Programs/AutomationTool/Platform.cs
+++ b/Engine/Source/Programs/AutomationTool/Platform.cs
@@ -290,6 +290,27 @@ namespace AutomationTool
 			throw new AutomationException("{0} does not yet implement GetCookPlatform.", PlatformType);
 		}
 
+        /// <summary>
+        /// Gets extra cook commandline arguments for this platform.
+        /// </summary>
+        /// <param name="Params"> ProjectParams </param>
+        /// <returns>Cook platform string.</returns>
+        public virtual string GetCookExtraCommandLine(ProjectParams Params)
+        {
+            return ""; 
+        }
+
+        /// <summary>
+        /// Gets extra maps needed on this platform.
+        /// </summary>
+        /// <returns>extra maps</returns>
+        public virtual List<string> GetCookExtraMaps()
+        {
+            return new List<string>();
+        }
+
+
+
         /// <summary>
         /// Gets editor cook platform name for this platform. Cooking the editor is not useful, but this is used to fill the derived data cache
         /// </summary>
diff --git a/Engine/Source/Programs/AutomationTool/Scripts/CookCommand.Automation.cs b/Engine/Source/Programs/AutomationTool/Scripts/CookCommand.Automation.cs
--- a/Engine/Source/Programs/AutomationTool/Scripts/CookCommand.Automation.cs
+++ b/Engine/Source/Programs/AutomationTool/Scripts/CookCommand.Automation.cs
@@ -193,6 +193,19 @@ public partial class Project : CommandUtils
                     CommandletParams += " ";
                     CommandletParams += FormatedAdditionalCookerParams;
                 }
+
+                if (!Params.NoClient)
+                {
+                    var MapsList = Maps.ToList(); 
+                    foreach (var ClientPlatform in Params.ClientTargetPlatforms)
+                    {
+                        var DataPlatform = Params.GetCookedDataPlatformForClientTarget(ClientPlatform);
+                        CommandletParams += (Params.GetTargetPlatformInstance(DataPlatform).GetCookExtraCommandLine(Params));
+                        MapsList.AddRange((Params.GetTargetPlatformInstance(ClientPlatform).GetCookExtraMaps()));
+                    }
+                    Maps = MapsList.ToArray();
+                }
+
                 CookCommandlet(Params.RawProjectPath, Params.UE4Exe, Maps, Dirs, InternationalizationPreset, Cultures, CombineCommandletParams(PlatformsToCook.ToArray()), CommandletParams);
             }
 			catch (Exception Ex)
