Launch on Android now restarts application if APK did not require updating
#ue4
#android

diff --git a/Engine/Source/Programs/AutomationTool/Android/AndroidPlatform.Automation.cs b/Engine/Source/Programs/AutomationTool/Android/AndroidPlatform.Automation.cs
--- a/Engine/Source/Programs/AutomationTool/Android/AndroidPlatform.Automation.cs
+++ b/Engine/Source/Programs/AutomationTool/Android/AndroidPlatform.Automation.cs
@@ -525,6 +525,21 @@ public class AndroidPlatform : Platform
 				{
 					if (InstalledResult.Output.Substring(5).Trim() == APKLastUpdateTime)
 						bNeedAPKInstall = false;
+
+					// Stop the previously running copy (uninstall/install did this before)
+					InstalledResult = RunAdbCommand(Params, "shell am force-stop " + PackageName, null, ERunOptions.AppMustExist);
+					if (InstalledResult.Output.Contains("Error"))
+					{
+						// force-stop not supported (Android < 3.0) so check if package is actually running
+						// Note: cannot use grep here since it may not be installed on device
+						InstalledResult = RunAdbCommand(Params, "shell ps", null, ERunOptions.AppMustExist);
+						if (InstalledResult.Output.Contains(PackageName))
+						{
+							// it is actually running so use the slow way to kill it (uninstall and reinstall)
+							bNeedAPKInstall = true;
+						}
+
+					}
 				}
 			}
 		}
@@ -710,7 +725,7 @@ public class AndroidPlatform : Platform
 				{
 					while (DeployCommands.Count > DeployMaxParallelCommands / 2)
 					{
-						Thread.Sleep(10);
+						Thread.Sleep(1);
 						DeployCommands.RemoveWhere(
 							delegate(ProcessResult r)
 							{
