Fix for CL 2500167, Take care of an edge case where the variable is null.

diff --git a/Engine/Source/Programs/AutomationTool/Scripts/CookCommand.Automation.cs b/Engine/Source/Programs/AutomationTool/Scripts/CookCommand.Automation.cs
--- a/Engine/Source/Programs/AutomationTool/Scripts/CookCommand.Automation.cs
+++ b/Engine/Source/Programs/AutomationTool/Scripts/CookCommand.Automation.cs
@@ -196,7 +196,7 @@ public partial class Project : CommandUtils
 
                 if (!Params.NoClient)
                 {
-                    var MapsList = Maps.ToList(); 
+                    var MapsList = Maps == null ? new List<string>() :  Maps.ToList(); 
                     foreach (var ClientPlatform in Params.ClientTargetPlatforms)
                     {
                         var DataPlatform = Params.GetCookedDataPlatformForClientTarget(ClientPlatform);
