build script, maybe fix for job agent lock

diff --git a/Engine/Source/Programs/AutomationTool/Scripts/GUBP.Automation.cs b/Engine/Source/Programs/AutomationTool/Scripts/GUBP.Automation.cs
--- a/Engine/Source/Programs/AutomationTool/Scripts/GUBP.Automation.cs
+++ b/Engine/Source/Programs/AutomationTool/Scripts/GUBP.Automation.cs
@@ -6889,7 +6889,10 @@ public class GUBP : BuildCommand
                         throw new AutomationException("Node {1} is sticky but has agent requirements.", NodeToDo);
                     }
                     string Procedure = GUBPNodes[NodeToDo].ECProcedure();
-
+                    if(GUBPNodes[NodeToDo].IsSticky() && NodeToDo == LastSticky)
+                    {
+                        Procedure = Procedure + "_Release";
+                    }
                     string Args = String.Format("{0}, subprocedure => '{1}', parallel => '{2}', jobStepName => '{3}', actualParameter => [{{actualParameterName => 'NodeName', value =>'{4}'}}",
                         BaseArgs, Procedure, DoParallel ? 1 : 0, NodeToDo, NodeToDo);
                     string ProcedureParams = GUBPNodes[NodeToDo].ECProcedureParams();
