build script, add linux editor and tools to shared promotions

diff --git a/Engine/Source/Programs/AutomationTool/Scripts/GUBP.Automation.cs b/Engine/Source/Programs/AutomationTool/Scripts/GUBP.Automation.cs
--- a/Engine/Source/Programs/AutomationTool/Scripts/GUBP.Automation.cs
+++ b/Engine/Source/Programs/AutomationTool/Scripts/GUBP.Automation.cs
@@ -920,13 +920,13 @@ public class GUBP : BuildCommand
 			AddDependency(RootEditorNode.StaticGetFullName(UnrealTargetPlatform.Win64));
 			AddDependency(ToolsForCompileNode.StaticGetFullName(UnrealTargetPlatform.Win64));
 		}
-		public static string StaticGetFullName(UnrealTargetPlatform InHostPlatform)
+		public static string StaticGetFullName()
 		{
 			return "RootEditor_Linux";
 		}
 		public override string GetFullName()
 		{
-			return StaticGetFullName(HostPlatform);
+			return StaticGetFullName();
 		}
 		public override int CISFrequencyQuantumShift(GUBP bp)
 		{
@@ -1035,7 +1035,7 @@ public class GUBP : BuildCommand
 		public ToolsCrossCompileNode(UnrealTargetPlatform InHostPlatform)
 			: base(InHostPlatform)
 		{
-			AddPseudodependency(RootEditorCrossCompileLinuxNode.StaticGetFullName(HostPlatform));
+			AddPseudodependency(RootEditorCrossCompileLinuxNode.StaticGetFullName());
 			AgentSharingGroup = "ToolsCrossCompileGroup" + StaticGetHostPlatformSuffix(HostPlatform);
 		}
 		public static string StaticGetFullName(UnrealTargetPlatform InHostPlatform)
@@ -2074,6 +2074,11 @@ public class GUBP : BuildCommand
                         }
                     }
                 }
+                if(HostPlatform == UnrealTargetPlatform.Win64 && bp.ActivePlatforms.Contains(UnrealTargetPlatform.Linux))
+                {
+                    AddDependency(RootEditorCrossCompileLinuxNode.StaticGetFullName());
+                    AddDependency(ToolsCrossCompileNode.StaticGetFullName(HostPlatform));
+                }
             }
 			AddDependency(MakeFeaturePacksNode.StaticGetFullName(MakeFeaturePacksNode.GetDefaultBuildPlatform(bp)));
         }
