PO files downloaded from OneSky are now backed up and submitted to P4. This will create a submission history in P4 such that we can more easily track down anomalies related to OneSky in our localization pipeline.

diff --git a/Engine/Source/Programs/AutomationTool/Scripts/Localisation.Automation.cs b/Engine/Source/Programs/AutomationTool/Scripts/Localisation.Automation.cs
--- a/Engine/Source/Programs/AutomationTool/Scripts/Localisation.Automation.cs
+++ b/Engine/Source/Programs/AutomationTool/Scripts/Localisation.Automation.cs
@@ -24,6 +24,8 @@ class Localise : BuildCommand
         }
     }
 
+    private int OneSkyDownloadedPOChangeList = 0;
+
 	public override void ExecuteBuild()
 	{
 		var EditorExe = CombinePaths(CmdEnv.LocalRoot, @"Engine/Binaries/Win64/UE4Editor-Cmd.exe");
@@ -41,6 +43,12 @@ class Localise : BuildCommand
         var oneSkyService = new OneSkyService(OneSkyConfig.ApiKey, OneSkyConfig.ApiSecret);
         var projectGroup = GetProjectGroup(oneSkyService, "Unreal Engine");
 
+        // Create changelist for backed up POs from OneSky.
+        if (P4Enabled)
+        {
+            OneSkyDownloadedPOChangeList = P4.CreateChange(P4Env.Client, "OneSky downloaded PO backup.");
+        }
+
         // Export all text from OneSky
         ExportProjectToDirectory(oneSkyService, projectGroup, "Engine");
         ExportProjectToDirectory(oneSkyService, projectGroup, "Editor");
@@ -48,6 +56,13 @@ class Localise : BuildCommand
         ExportProjectToDirectory(oneSkyService, projectGroup, "PropertyNames");
         ExportProjectToDirectory(oneSkyService, projectGroup, "ToolTips");
 
+        // Submit changelist for backed up POs from OneSky.
+        if (P4Enabled)
+        {
+            int SubmittedChangeList;
+            P4.Submit(OneSkyDownloadedPOChangeList, out SubmittedChangeList);
+        }
+
 		// Setup editor arguments for SCC.
 		string EditorArguments = String.Empty;
 		if (P4Enabled)
@@ -152,7 +167,7 @@ class Localise : BuildCommand
         return appProject;
     }
 
-    private static void ExportFileToDirectory(UploadedFile file, DirectoryInfo destination, IEnumerable<string> cultures)
+    private void ExportFileToDirectory(UploadedFile file, DirectoryInfo destination, IEnumerable<string> cultures)
     {
         foreach (var culture in cultures)
         {
@@ -175,6 +190,13 @@ class Localise : BuildCommand
                         memoryStream.CopyTo(fileStream);
                         Console.WriteLine("[SUCCESS] Exporting: " + exportFile.FullName + " Locale: " + culture);
                     }
+                    FileInfo exportFileCopy = new FileInfo(Path.Combine(exportFile.DirectoryName, Path.GetFileNameWithoutExtension(exportFile.Name) + "_FromOneSky" + exportFile.Extension));
+                    // Add/check out backed up POs from OneSky.
+                    if (P4Enabled)
+                    {
+                        UE4Build.AddBuildProductsToChangelist(OneSkyDownloadedPOChangeList, new List<string>() {exportFileCopy.FullName} );
+                    }
+                    File.Copy(exportFile.FullName, exportFileCopy.FullName, true);
                 }
                 else if (exportTranslationState == UploadedFile.ExportTranslationState.NoContent)
                 {
