Merging

//depot/UE4-UT/...

to //depot/UE4/...

diff --git a/Engine/Source/Programs/AutomationTool/Mac/MacPlatform.Automation.cs b/Engine/Source/Programs/AutomationTool/Mac/MacPlatform.Automation.cs
--- a/Engine/Source/Programs/AutomationTool/Mac/MacPlatform.Automation.cs
+++ b/Engine/Source/Programs/AutomationTool/Mac/MacPlatform.Automation.cs
@@ -80,12 +80,12 @@ public class MacPlatform : Platform
 				if (SC.StageTargetConfigurations.Contains(UnrealTargetConfiguration.Test))
 				{
 					SC.StageFiles(StagedFileType.NonUFS, CombinePaths(SC.LocalRoot, "Engine/Binaries", SC.PlatformDir, "UE4-Mac-Test.app"));
-					SC.StageFiles(StagedFileType.NonUFS, CombinePaths(SC.ProjectRoot, "Binaries", SC.PlatformDir), "UE4-" + SC.ShortProjectName + "Test.dylib");
+					SC.StageFiles(StagedFileType.NonUFS, CombinePaths(SC.ProjectRoot, "Binaries", SC.PlatformDir), "UE4-" + SC.ShortProjectName + "-Test.dylib");
 				}
 				if (SC.StageTargetConfigurations.Contains(UnrealTargetConfiguration.Shipping))
 				{
 					SC.StageFiles(StagedFileType.NonUFS, CombinePaths(SC.LocalRoot, "Engine/Binaries", SC.PlatformDir, "UE4-Mac-Shipping.app"));
-					SC.StageFiles(StagedFileType.NonUFS, CombinePaths(SC.ProjectRoot, "Binaries", SC.PlatformDir), "UE4-" + SC.ShortProjectName + "Shipping.dylib");
+					SC.StageFiles(StagedFileType.NonUFS, CombinePaths(SC.ProjectRoot, "Binaries", SC.PlatformDir), "UE4-" + SC.ShortProjectName + "-Shipping.dylib");
 				}
 
 				SC.StageFiles(StagedFileType.NonUFS, CombinePaths(SC.LocalRoot, "Engine/Plugins"), "UE4-*.dylib", true);
