Allow preventing the creation of a UE4CommandLine.txt if the target platform doesn't need one. Windows stages a bootstrap executable which directly launches the game with the project argument, and now that it can parse UE4CommandLine.txt too (it used to be staged even though it was always ignored), it causes the final command line to duplicate the project argument.

diff --git a/Engine/Source/Programs/AutomationTool/Platform.cs b/Engine/Source/Programs/AutomationTool/Platform.cs
--- a/Engine/Source/Programs/AutomationTool/Platform.cs
+++ b/Engine/Source/Programs/AutomationTool/Platform.cs
@@ -476,6 +476,14 @@ namespace AutomationTool
 			return true;
 		}
 
+		/// <summary>
+		/// Determines whether we should stage a UE4CommandLine.txt for this platform 
+		/// </summary>
+		public virtual bool ShouldStageCommandLine(ProjectParams Params, DeploymentContext SC)
+		{
+			return true;
+		}
+
         /// <summary>
         /// Only relevant for the mac and PC at the moment. Example calling the Mac platform with PS4 as an arg will return false. Can't compile or cook for the PS4 on the mac.
         /// </summary>
diff --git a/Engine/Source/Programs/AutomationTool/Scripts/CopyBuildToStagingDirectory.Automation.cs b/Engine/Source/Programs/AutomationTool/Scripts/CopyBuildToStagingDirectory.Automation.cs
--- a/Engine/Source/Programs/AutomationTool/Scripts/CopyBuildToStagingDirectory.Automation.cs
+++ b/Engine/Source/Programs/AutomationTool/Scripts/CopyBuildToStagingDirectory.Automation.cs
@@ -993,6 +993,11 @@ public partial class Project : CommandUtils
 			File.Delete(IntermediateCmdLineFile);
 		}
 
+		if(!SC.StageTargetPlatform.ShouldStageCommandLine(Params, SC))
+		{
+			return;
+		}
+
 		if (!string.IsNullOrEmpty(Params.StageCommandline) || !string.IsNullOrEmpty(Params.RunCommandline))
 		{
 			string FileHostParams = " ";
diff --git a/Engine/Source/Programs/AutomationTool/Win/WinPlatform.Automation.cs b/Engine/Source/Programs/AutomationTool/Win/WinPlatform.Automation.cs
--- a/Engine/Source/Programs/AutomationTool/Win/WinPlatform.Automation.cs
+++ b/Engine/Source/Programs/AutomationTool/Win/WinPlatform.Automation.cs
@@ -195,7 +195,8 @@ public abstract class BaseWinPlatform : Platform
 					StageExecutable("exe", SC, CombinePaths(SC.LocalRoot, "Engine/Binaries", SC.PlatformDir), Path.GetFileNameWithoutExtension(SourceFile) + ".", true, null, null, false, WorkingFileType);
 					if(Exe == Exes[0] && !Params.NoBootstrapExe)
 					{
-						StageBootstrapExecutable(SC, SourceFile, CombinePaths("Engine", "Binaries", SC.PlatformDir, ExeFileName), String.Format("..\\..\\..\\{0}\\{0}.uproject", SC.ShortProjectName));
+						string CommandLine = ShouldStageCommandLine(Params, SC)? "" : String.Format("..\\..\\..\\{0}\\{0}.uproject", SC.ShortProjectName);
+						StageBootstrapExecutable(SC, SourceFile, CombinePaths("Engine", "Binaries", SC.PlatformDir, ExeFileName), CommandLine);
 					}
 				}
 				else
@@ -357,6 +358,11 @@ public abstract class BaseWinPlatform : Platform
 		return ExecutableNames;
 	}
 
+	public override bool ShouldStageCommandLine(ProjectParams Params, DeploymentContext SC)
+	{
+		return !String.IsNullOrEmpty(Params.StageCommandline) || !String.IsNullOrEmpty(Params.RunCommandline) || (!Params.IsCodeBasedProject && Params.NoBootstrapExe);
+	}
+
 	public override List<string> GetDebugFileExtentions()
 	{
 		return new List<string> { ".pdb", ".map" };
