Fix handling of paths containing spaces for some P4 commands.

diff --git a/Engine/Source/Programs/AutomationTool/P4Utils.cs b/Engine/Source/Programs/AutomationTool/P4Utils.cs
--- a/Engine/Source/Programs/AutomationTool/P4Utils.cs
+++ b/Engine/Source/Programs/AutomationTool/P4Utils.cs
@@ -2197,7 +2197,7 @@ namespace AutomationTool
 		{
 			CheckP4Enabled();
 			string Output;
-			string Command = "fstat " + Filename;
+			string Command = "fstat " + CommandUtils.MakePathSafeToUseWithCommandLine(Filename);
 			if (!LogP4Output(out Output, Command))
 			{
 				throw new P4Exception("p4.exe {0} failed.", Command);
@@ -2263,7 +2263,7 @@ namespace AutomationTool
 			{
 				var CmdLine = String.Format("{0} -c {1} -t {2} {3}",
 					(Stat.Action != P4Action.None) ? "reopen" : "open",
-					Changelist, FileAttributesToString(Attributes | Stat.Attributes), Filename);
+					Changelist, FileAttributesToString(Attributes | Stat.Attributes), CommandUtils.MakePathSafeToUseWithCommandLine(Filename));
 				LogP4(CmdLine);
 			}
 		}
