Submit build receipts to P4 as part of promotions.

diff --git a/Engine/Source/Programs/AutomationTool/UE4Build.cs b/Engine/Source/Programs/AutomationTool/UE4Build.cs
--- a/Engine/Source/Programs/AutomationTool/UE4Build.cs
+++ b/Engine/Source/Programs/AutomationTool/UE4Build.cs
@@ -4,6 +4,7 @@ using System.Collections.Generic;
 using System.Text;
 using System.IO;
 using System.Xml;
+using System.Diagnostics;
 
 namespace AutomationTool
 {
@@ -54,7 +55,7 @@ namespace AutomationTool
 			}
 		}
 
-		void PrepareManifest(string ManifestName)
+		void PrepareManifest(string ManifestName, bool bAddReceipt)
 		{
 			if (FileExists(ManifestName) == false)
 			{
@@ -79,10 +80,22 @@ namespace AutomationTool
 			UnrealBuildTool.BuildManifest Manifest = ReadManifest(ManifestName);
 			foreach (string Item in Manifest.BuildProducts)
 			{
-				PrepareBuildProduct(Item);
+				if(bAddReceipt && IsBuildReceipt(Item))
+				{
+					AddBuildProduct(Item);
+				}
+				else
+				{
+					PrepareBuildProduct(Item);
+				}
 			}
 		}
 
+		static bool IsBuildReceipt(string FileName)
+		{
+			return Path.GetDirectoryName(FileName).Replace('\\', '/').EndsWith("/Receipts", StringComparison.InvariantCultureIgnoreCase);
+		}
+
 		void AddBuildProductsFromManifest(string ManifestName)
 		{
 			if (FileExists(ManifestName) == false)
@@ -175,7 +188,7 @@ namespace AutomationTool
 
 			RunUBT(CmdEnv, UBTExecutable: UBTExecutable, Project: ProjectName, Target: TargetName, Platform: Platform.ToString(), Config: Config, AdditionalArgs: "-generatemanifest -nobuilduht -xgeexport" + AddArgs, EnvVars: EnvVars);
 
-			PrepareManifest(UBTManifest);
+			PrepareManifest(UBTManifest, true);
 
 			Result.Platform = Platform;
 			Result.Config = Config;
@@ -316,7 +329,7 @@ namespace AutomationTool
 
 				RunUBT(CmdEnv, UBTExecutable: UBTExecutable, Project: ProjectName, Target: TargetName, Platform: TargetPlatform.ToString(), Config: Config, AdditionalArgs: "-generatemanifest" + AddArgs, EnvVars: EnvVars);
 
-				PrepareManifest(UBTManifest);
+				PrepareManifest(UBTManifest, false);
 			}
 
 			RunUBT(CmdEnv, UBTExecutable: UBTExecutable, Project: ProjectName, Target: TargetName, Platform: TargetPlatform.ToString(), Config: Config, AdditionalArgs: AddArgs, EnvVars: EnvVars);
@@ -576,6 +589,7 @@ namespace AutomationTool
 			return Values;
 		}
 
+		[DebuggerDisplay("{TargetName} {Platform} {Config}")]
 		public class BuildTarget
 		{
 			/// Unreal project name, if needed
@@ -1345,13 +1359,13 @@ namespace AutomationTool
 			Log("Adding {0} build products to changelist {1}...", Files.Count, WorkingCL);
 			foreach (var File in Files)
 			{
-				P4.Sync("-f -k " + File + "#head"); // sync the file without overwriting local one
+				P4.Sync("-f -k " + CommandUtils.MakePathSafeToUseWithCommandLine(File) + "#head"); // sync the file without overwriting local one
 				if (!FileExists(File))
 				{
 					throw new AutomationException("BUILD FAILED {0} was a build product but no longer exists", File);
 				}
 
-				P4.ReconcileNoDeletes(WorkingCL, File);
+				P4.ReconcileNoDeletes(WorkingCL, CommandUtils.MakePathSafeToUseWithCommandLine(File));
 
 				// Change file type on binary files to be always writeable.
 				var FileStats = P4.FStat(File);
@@ -1382,7 +1396,7 @@ namespace AutomationTool
 		/// <returns>True if this is a Windows build product. False otherwise.</returns>
 		private static bool IsBuildProduct(string File, P4FileStat FileStats)
 		{
-			if(FileStats.Type == P4FileType.Binary)
+			if(FileStats.Type == P4FileType.Binary || IsBuildReceipt(File))
 			{
 				return true;
 			}
diff --git a/Engine/Source/Programs/UnrealBuildTool/Configuration/UEBuildTarget.cs b/Engine/Source/Programs/UnrealBuildTool/Configuration/UEBuildTarget.cs
--- a/Engine/Source/Programs/UnrealBuildTool/Configuration/UEBuildTarget.cs
+++ b/Engine/Source/Programs/UnrealBuildTool/Configuration/UEBuildTarget.cs
@@ -1330,6 +1330,8 @@ namespace UnrealBuildTool
 				Manifest.AddBuildProduct(BuildProduct.Path);
 			}
 
+			Manifest.AddBuildProduct(BuildReceipt.GetDefaultPath(ProjectDirectory, TargetName, Platform, Configuration));
+
 			if (UEBuildConfiguration.bCleanProject)
 			{
 				CleanTarget(Manifest);
