build script, start gating cooks

diff --git a/Engine/Source/Programs/AutomationTool/Scripts/GUBP.Automation.cs b/Engine/Source/Programs/AutomationTool/Scripts/GUBP.Automation.cs
--- a/Engine/Source/Programs/AutomationTool/Scripts/GUBP.Automation.cs
+++ b/Engine/Source/Programs/AutomationTool/Scripts/GUBP.Automation.cs
@@ -3746,7 +3746,7 @@ public class GUBP : BuildCommand
         if (!GUBPNodes.ContainsKey(Node))
         {
             throw new AutomationException("Node {0} not found", Node);
-        }
+        }        
         GUBPNodes[Node].RemovePseudodependency(Dep);
     }
 
@@ -5767,9 +5767,6 @@ public class GUBP : BuildCommand
 													RemovePseudodependencyFromNode(
 														CookNode.StaticGetFullName(HostPlatform, NonCodeProject, CookedPlatform),
 														WaitForTestShared.StaticGetFullName());
-													RemovePseudodependencyFromNode(
-														CookNode.StaticGetFullName(HostPlatform, NonCodeProject, CookedPlatform),
-														CookNode.StaticGetFullName(HostPlatform, Branch.BaseEngineProject, CookedPlatform));
 													string BuildAgentSharingGroup = "";
 													if (Options.bSeparateGamePromotion)
 													{
@@ -6008,9 +6005,6 @@ public class GUBP : BuildCommand
 											RemovePseudodependencyFromNode(
 												CookNode.StaticGetFullName(HostPlatform, CodeProj, CookedPlatform),
 												WaitForTestShared.StaticGetFullName());
-											RemovePseudodependencyFromNode(
-												CookNode.StaticGetFullName(HostPlatform, CodeProj, CookedPlatform),
-												CookNode.StaticGetFullName(HostPlatform, Branch.BaseEngineProject, CookedPlatform));
 											string BuildAgentSharingGroup = "";
 											
 											if (Options.bSeparateGamePromotion)
