Disable Rocket from promotable for now; causing all build products to be added into the promotable. Will get working properly after lockdown.

diff --git a/Engine/Source/Programs/AutomationTool/Scripts/RocketBuild.Automation.cs b/Engine/Source/Programs/AutomationTool/Scripts/RocketBuild.Automation.cs
--- a/Engine/Source/Programs/AutomationTool/Scripts/RocketBuild.Automation.cs
+++ b/Engine/Source/Programs/AutomationTool/Scripts/RocketBuild.Automation.cs
@@ -75,7 +75,7 @@ namespace Rocket
 			if(!bp.BranchOptions.bNoInstalledEngine)
 			{
 				// Get the promotable node, so we can add additional dependencies to it
-				GUBP.GUBPNode PromotableNode = bp.FindNode(GUBP.SharedAggregatePromotableNode.StaticGetFullName());
+//				GUBP.GUBPNode PromotableNode = bp.FindNode(GUBP.SharedAggregatePromotableNode.StaticGetFullName());
 
 				// Get the output paths for the installed engine
 				string InstallDir = CommandUtils.CombinePaths(CommandUtils.CmdEnv.LocalRoot, "LocalBuilds", "Engine", CommandUtils.GetGenericPlatformName(HostPlatform));
@@ -111,7 +111,7 @@ namespace Rocket
 
 				// Add the aggregate node for the entire install
 				AggNode.AddDependency(bp.AddNode(new BuildInstalledEngineNode(HostPlatform, InstallDir, SymbolsDir)));
-				PromotableNode.AddDependency(BuildInstalledEngineNode.StaticGetFullName(HostPlatform));
+//				PromotableNode.AddDependency(BuildInstalledEngineNode.StaticGetFullName(HostPlatform));
 
 				// Add a node for GitHub promotions
 				if(HostPlatform == UnrealTargetPlatform.Win64)
@@ -120,7 +120,7 @@ namespace Rocket
 					if(CommandUtils.FileExists(CommandUtils.CombinePaths(CommandUtils.CmdEnv.LocalRoot, GitConfigRelativePath)))
 					{
 						bp.AddNode(new BuildGitPromotable(bp, HostPlatform, GitConfigRelativePath));
-						PromotableNode.AddDependency(BuildGitPromotable.StaticGetFullName(HostPlatform));
+//						PromotableNode.AddDependency(BuildGitPromotable.StaticGetFullName(HostPlatform));
 					}
 				}
 			}
