brutal resolve performance fixes

diff --git a/Engine/Source/Programs/AutomationTool/P4Utils.cs b/Engine/Source/Programs/AutomationTool/P4Utils.cs
--- a/Engine/Source/Programs/AutomationTool/P4Utils.cs
+++ b/Engine/Source/Programs/AutomationTool/P4Utils.cs
@@ -1335,7 +1335,14 @@ namespace AutomationTool
 				{
 					throw new P4Exception("Change {0} was not pending.", CL);
 				}
-
+                bool isClPending = false;
+                if (ChangeFiles(CL, out isClPending, false).Count == 0)
+                {
+                    CommandUtils.Log(TraceEventType.Information, "No edits left to commit after brutal submit resolve. Assuming another build committed same changes already and exiting as success.");
+                    DeleteChange(CL);
+                    // No changes to submit, no need to retry.
+                    return;
+                }
 				string CmdOutput;
 				if (!LogP4Output(out CmdOutput, String.Format("submit -c {0}", CL)))
 				{
@@ -1370,6 +1377,7 @@ namespace AutomationTool
                     if (AnyIssue)
                     {
                         string Work = CmdOutput;
+                        HashSet<string> AlreadyDone = new HashSet<string>();
                         while (Work.Length > 0)
                         {
                             string SlashSlashStr = "//";
@@ -1391,8 +1399,12 @@ namespace AutomationTool
                             if (MinMatch > Work.Length)
                             {
                                 break;
-                            }
+                            }                            
                             string File = Work.Substring(0, MinMatch).Trim();
+                            if (AlreadyDone.Contains(File))
+                            {
+                                continue;
+                            }
                             if (File.IndexOf(SlashSlashStr) != File.LastIndexOf(SlashSlashStr))
                             {
                                 // this is some other line about the same line, we ignore it, removing the first // so we advance
@@ -1401,21 +1413,12 @@ namespace AutomationTool
                             else
                             {
                                 Work = Work.Substring(MinMatch);
-
 								CommandUtils.Log(TraceEventType.Information, "Brutal 'resolve' on {0} to force submit.\n", File);
-
 								Revert(CL, "-k " + CommandUtils.MakePathSafeToUseWithCommandLine(File));  // revert the file without overwriting the local one
 								Sync("-f -k " + CommandUtils.MakePathSafeToUseWithCommandLine(File + "#head"), false); // sync the file without overwriting local one
 								ReconcileNoDeletes(CL, CommandUtils.MakePathSafeToUseWithCommandLine(File));  // re-check out, if it changed, or add
                                 DidSomething = true;
-								bool isClPending = false;
-								if (ChangeFiles(CL, out isClPending).Count == 0)
-								{
-									CommandUtils.Log(TraceEventType.Information, "No edits left to commit after brutal submit resolve. Assuming another build committed same changes already and exiting as success.");
-									DeleteChange(CL);
-									// No changes to submit, no need to retry.
-									return;
-								}
+                                AlreadyDone.Add(File);															
                             }
                         }
                     }
@@ -1610,11 +1613,11 @@ namespace AutomationTool
 		/// </summary>
 		/// <param name="CL">Changelist to get the specification from.</param>
 		/// <returns>Specification of the changelist.</returns>
-		public string ChangeOutput(int CL)
+		public string ChangeOutput(int CL, bool AllowSpew = true)
 		{
 			CheckP4Enabled();
 			string CmdOutput;
-			if (LogP4Output(out CmdOutput, String.Format("change -o {0}", CL)))
+			if (LogP4Output(out CmdOutput, String.Format("change -o {0}", CL), AllowSpew: AllowSpew))
 			{
 				return CmdOutput;
 			}
@@ -1627,10 +1630,10 @@ namespace AutomationTool
 		/// <param name="CL">Changelist id.</param>
 		/// <param name="Pending">Whether it is a pending changelist.</param>
 		/// <returns>Returns whether the changelist exists.</returns>
-		public bool ChangeExists(int CL, out bool Pending)
+		public bool ChangeExists(int CL, out bool Pending, bool AllowSpew = true)
 		{
 			CheckP4Enabled();
-			string CmdOutput = ChangeOutput(CL);
+			string CmdOutput = ChangeOutput(CL, AllowSpew);
 			Pending = false;
 			if (CmdOutput.Length > 0)
 			{
@@ -1670,14 +1673,14 @@ namespace AutomationTool
 		/// <param name="CL">Changelist to get the files from.</param>
 		/// <param name="Pending">Whether the changelist is a pending one.</param>
 		/// <returns>List of the files contained in the changelist.</returns>
-		public List<string> ChangeFiles(int CL, out bool Pending)
+		public List<string> ChangeFiles(int CL, out bool Pending, bool AllowSpew = true)
 		{
 			CheckP4Enabled();
 			var Result = new List<string>();
 
-			if (ChangeExists(CL, out Pending))
+			if (ChangeExists(CL, out Pending, AllowSpew))
 			{
-				string CmdOutput = ChangeOutput(CL);
+				string CmdOutput = ChangeOutput(CL, AllowSpew);
 				if (CmdOutput.Length > 0)
 				{
 
