Allow HTML5  to download maps on the fly.

UPLAT-685 (Remove Virtual File System to reduce run-time memory overhead and allow for larger worlds in HTML5)

diff --git a/Engine/Source/Developer/HTML5/HTML5PlatformEditor/Classes/HTML5TargetSettings.h b/Engine/Source/Developer/HTML5/HTML5PlatformEditor/Classes/HTML5TargetSettings.h
--- a/Engine/Source/Developer/HTML5/HTML5PlatformEditor/Classes/HTML5TargetSettings.h
+++ b/Engine/Source/Developer/HTML5/HTML5PlatformEditor/Classes/HTML5TargetSettings.h
@@ -9,6 +9,20 @@
 #include "HTML5TargetSettings.generated.h"
 
 
+
+USTRUCT()
+struct FHTML5LevelTransitions
+{
+	GENERATED_USTRUCT_BODY()
+
+	UPROPERTY(EditAnywhere, Category = HTML5_LevelTransitions, Meta = (DisplayName = "From Map"))
+	FFilePath MapFrom;
+
+	UPROPERTY(EditAnywhere, Category = HTML5_LevelTransitions, Meta = (DisplayName = "To Map"))
+	FFilePath MapTo;
+};
+
+
 /**
  * Implements the settings for the HTML5 target platform.
  */
@@ -19,7 +33,6 @@ class HTML5PLATFORMEDITOR_API UHTML5TargetSettings
 public:
 
 	GENERATED_UCLASS_BODY()
-
  	/**
  	 * Setting to control HTML5 Heap size (in Development)
  	 */
@@ -37,4 +50,16 @@ public:
  	 */
  	UPROPERTY(GlobalConfig, EditAnywhere, Category=Memory, Meta = (DisplayName = "Port to use when deploying game from the editor", ClampMin="49152", ClampMax="65535"))
  	int32 DeployServerPort;
+	/**
+	 * Use a loading level and download maps during transitions.                                                                     
+	 */
+	UPROPERTY(GlobalConfig, EditAnywhere, Category = Packaging, Meta = (DisplayName = "Download Maps on the fly [experimental]"))
+	bool UseAsyncLevelLoading;
+
+	/**
+	 * Generate Delta Pak files for these level transitions.
+	 */
+	UPROPERTY(GlobalConfig, EditAnywhere, Category = Packaging, Meta = (DisplayName = "Level transitions for delta paks [experimental]"))
+	TArray<FHTML5LevelTransitions> LevelTransitions;
+
 };
diff --git a/Engine/Source/Programs/AutomationTool/HTML5/HTML5.Automation.csproj b/Engine/Source/Programs/AutomationTool/HTML5/HTML5.Automation.csproj
--- a/Engine/Source/Programs/AutomationTool/HTML5/HTML5.Automation.csproj
+++ b/Engine/Source/Programs/AutomationTool/HTML5/HTML5.Automation.csproj
@@ -48,6 +48,7 @@
       <Link>Properties\MetaData.cs</Link>
     </Compile>
     <Compile Include="HTML5Platform.Automation.cs" />
+    <Compile Include="HTML5Platform.PakFiles.Automation.cs" />
     <Compile Include="Properties\AssemblyInfo.cs" />
   </ItemGroup>
   <ItemGroup>
diff --git a/Engine/Source/Programs/AutomationTool/HTML5/HTML5Platform.Automation.cs b/Engine/Source/Programs/AutomationTool/HTML5/HTML5Platform.Automation.cs
--- a/Engine/Source/Programs/AutomationTool/HTML5/HTML5Platform.Automation.cs
+++ b/Engine/Source/Programs/AutomationTool/HTML5/HTML5Platform.Automation.cs
@@ -12,7 +12,7 @@ using System.Threading;
 
 public class HTML5Platform : Platform
 {
- 	public HTML5Platform()
+    public HTML5Platform()
 		: base(UnrealTargetPlatform.HTML5)
 	{
 	}
@@ -21,87 +21,124 @@ public class HTML5Platform : Platform
 	{
 		Log("Package {0}", Params.RawProjectPath);
         
-		string PackagePath = Path.Combine(Path.GetDirectoryName(Params.RawProjectPath), "Binaries", "HTML5");
-		if (!Directory.Exists(PackagePath))
-		{
-			Directory.CreateDirectory(PackagePath);
-		}
-		string FinalDataLocation = Path.Combine(PackagePath, Params.ShortProjectName) + ".data";
+        string PackagePath = Path.Combine(Path.GetDirectoryName(Params.RawProjectPath), "Binaries", "HTML5");
+        if (!Directory.Exists(PackagePath))
+        {
+            Directory.CreateDirectory(PackagePath);
+        }
+        string FinalDataLocation = Path.Combine(PackagePath, Params.ShortProjectName) + ".data";
 
-		    // we need to operate in the root
-            using (new PushedDirectory(Path.Combine(Params.BaseStageDirectory, "HTML5")))
-            {
-                string PythonPath = HTML5SDKInfo.PythonPath();
-                string PackagerPath = HTML5SDKInfo.EmscriptenPackager();
 
-                string CmdLine = string.Format("\"{0}\" \"{1}\" --preload . --js-output=\"{1}.js\"", PackagerPath, FinalDataLocation);
-                RunAndLog(CmdEnv, PythonPath, CmdLine);
-            }
+        bool UseAsyncLevelLoading = false;
+        var ConfigCache = new UnrealBuildTool.ConfigCacheIni(UnrealTargetPlatform.HTML5, "Engine", Path.GetDirectoryName(Params.RawProjectPath), CombinePaths(CmdEnv.LocalRoot, "Engine"));
+        ConfigCache.GetBool("/Script/HTML5PlatformEditor.HTML5TargetSettings", "UseAsyncLevelLoading", out UseAsyncLevelLoading);
 
-            // copy the "Executable" to the package directory
-            string GameExe = Path.GetFileNameWithoutExtension(Params.ProjectGameExeFilename);
-            if (Params.ClientConfigsToBuild[0].ToString() != "Development")
-            {
-                GameExe += "-HTML5-" + Params.ClientConfigsToBuild[0].ToString();
-            }
-			GameExe += ".js";
-			
-			// ensure the ue4game binary exists, if applicable
-			string FullGameExePath = Path.Combine(Path.GetDirectoryName(Params.ProjectGameExeFilename), GameExe);
-			if (!SC.IsCodeBasedProject && !FileExists_NoExceptions(FullGameExePath))
-			{
-				Log("Failed to find game application " + FullGameExePath);
-				AutomationTool.ErrorReporter.Error("Stage Failed.", (int)AutomationTool.ErrorCodes.Error_MissingExecutable);
-				throw new AutomationException("Could not find application {0}. You may need to build the UE4 project with your target configuration and platform.", FullGameExePath);
-			}
+        if (UseAsyncLevelLoading)
+        {
+            HTMLPakAutomation PakAutomation = new HTMLPakAutomation(Params, SC);
 
-            if (Path.Combine(Path.GetDirectoryName(Params.ProjectGameExeFilename), GameExe) != Path.Combine(PackagePath, GameExe))
-            {
-                File.Copy(Path.Combine(Path.GetDirectoryName(Params.ProjectGameExeFilename), GameExe), Path.Combine(PackagePath, GameExe), true);
-                File.Copy(Path.Combine(Path.GetDirectoryName(Params.ProjectGameExeFilename), GameExe) + ".mem", Path.Combine(PackagePath, GameExe) + ".mem", true);
-            }
-            File.SetAttributes(Path.Combine(PackagePath, GameExe), FileAttributes.Normal);
-            File.SetAttributes(Path.Combine(PackagePath, GameExe) + ".mem", FileAttributes.Normal);
+            // Create Necessary Paks. 
+            PakAutomation.CreateEnginePak();
+            PakAutomation.CreateGamePak();
+            PakAutomation.CreateContentDirectoryPak();
 
-            // put the HTML file to the package directory
-            string TemplateFile = Path.Combine(CombinePaths(CmdEnv.LocalRoot, "Engine"), "Build", "HTML5", "Game.html.template");
-            string OutputFile = Path.Combine(PackagePath, (Params.ClientConfigsToBuild[0].ToString() != "Development" ? (Params.ShortProjectName + "-HTML5-" + Params.ClientConfigsToBuild[0].ToString()) : Params.ShortProjectName)) + ".html";
+            // Create Emscripten Package from Necessary Paks. - This will be the VFS. 
+            PakAutomation.CreateEmscriptenDataPackage(PackagePath, FinalDataLocation);
 
-            // find Heap Size.
-            ulong HeapSize;
-            var ConfigCache = new UnrealBuildTool.ConfigCacheIni(UnrealTargetPlatform.HTML5, "Engine", Path.GetDirectoryName(Params.RawProjectPath), CombinePaths(CmdEnv.LocalRoot, "Engine"));
+            // Create All Map Paks which  will be downloaded on the fly. 
+            PakAutomation.CreateMapPak();
 
-            int ConfigHeapSize = 0;
-			// Valuer set by Editor UI
-			var bGotHeapSize = ConfigCache.GetInt32("/Script/HTML5PlatformEditor.HTML5TargetSettings", "HeapSize" + Params.ClientConfigsToBuild[0].ToString(), out ConfigHeapSize);
+            // Create Delta Paks if setup.
+            List<string> Paks = new List<string>();
+            ConfigCache.GetArray("/Script/HTML5PlatformEditor.HTML5TargetSettings", "LevelTransitions", out Paks);
 
-			// Fallback if the previous method failed
-            if (!bGotHeapSize && !ConfigCache.GetInt32("BuildSettings", "HeapSize" + Params.ClientConfigsToBuild[0].ToString(), out ConfigHeapSize)) // in Megs.
+            if (Paks != null)
             {
-                // we couldn't find a per config heap size, look for a common one.
-                if (!ConfigCache.GetInt32("BuildSettings", "HeapSize", out ConfigHeapSize))
+                foreach (var Pak in Paks)
                 {
-                    ConfigHeapSize = Params.IsCodeBasedProject ? 1024 : 512;
-                    Log("Could not find Heap Size setting in .ini for Client config {0}", Params.ClientConfigsToBuild[0].ToString());
+                    var Matched = Regex.Matches(Pak, "\"[^\"]+\"", RegexOptions.IgnoreCase);
+                    string MapFrom = Path.GetFileNameWithoutExtension(Matched[0].ToString().Replace("\"", ""));
+                    string MapTo = Path.GetFileNameWithoutExtension(Matched[1].ToString().Replace("\"", ""));
+                    PakAutomation.CreateDeltaMapPaks(MapFrom, MapTo);
                 }
             }
+        }
+        else
+        {
+            // we need to operate in the root
+            using (new PushedDirectory(Path.Combine(Params.BaseStageDirectory, "HTML5")))
+            {
+                string PythonPath = HTML5SDKInfo.PythonPath();
+                string PackagerPath = HTML5SDKInfo.EmscriptenPackager();
 
-            HeapSize = (ulong)ConfigHeapSize * 1024L * 1024L; // convert to bytes.
-            Log("Setting Heap size to {0} Mb ", ConfigHeapSize);
+                string CmdLine = string.Format("\"{0}\" \"{1}\" --preload . --js-output=\"{1}.js\"", PackagerPath, FinalDataLocation);
+                RunAndLog(CmdEnv, PythonPath, CmdLine);
+            }
+        }
+
+        // copy the "Executable" to the package directory
+        string GameExe = Path.GetFileNameWithoutExtension(Params.ProjectGameExeFilename);
+        if (Params.ClientConfigsToBuild[0].ToString() != "Development")
+        {
+            GameExe += "-HTML5-" + Params.ClientConfigsToBuild[0].ToString();
+        }
+        GameExe += ".js";
+			
+        // ensure the ue4game binary exists, if applicable
+        string FullGameExePath = Path.Combine(Path.GetDirectoryName(Params.ProjectGameExeFilename), GameExe);
+        if (!SC.IsCodeBasedProject && !FileExists_NoExceptions(FullGameExePath))
+        {
+	        Log("Failed to find game application " + FullGameExePath);
+	        AutomationTool.ErrorReporter.Error("Stage Failed.", (int)AutomationTool.ErrorCodes.Error_MissingExecutable);
+	        throw new AutomationException("Could not find application {0}. You may need to build the UE4 project with your target configuration and platform.", FullGameExePath);
+        }
+
+        if (Path.Combine(Path.GetDirectoryName(Params.ProjectGameExeFilename), GameExe) != Path.Combine(PackagePath, GameExe))
+        {
+            File.Copy(Path.Combine(Path.GetDirectoryName(Params.ProjectGameExeFilename), GameExe), Path.Combine(PackagePath, GameExe), true);
+            File.Copy(Path.Combine(Path.GetDirectoryName(Params.ProjectGameExeFilename), GameExe) + ".mem", Path.Combine(PackagePath, GameExe) + ".mem", true);
+        }
+        File.SetAttributes(Path.Combine(PackagePath, GameExe), FileAttributes.Normal);
+        File.SetAttributes(Path.Combine(PackagePath, GameExe) + ".mem", FileAttributes.Normal);
+
+        // put the HTML file to the package directory
+        string TemplateFile = Path.Combine(CombinePaths(CmdEnv.LocalRoot, "Engine"), "Build", "HTML5", "Game.html.template");
+        string OutputFile = Path.Combine(PackagePath, (Params.ClientConfigsToBuild[0].ToString() != "Development" ? (Params.ShortProjectName + "-HTML5-" + Params.ClientConfigsToBuild[0].ToString()) : Params.ShortProjectName)) + ".html";
+
+        // find Heap Size.
+        ulong HeapSize;
+
+        int ConfigHeapSize = 0;
+        // Valuer set by Editor UI
+        var bGotHeapSize = ConfigCache.GetInt32("/Script/HTML5PlatformEditor.HTML5TargetSettings", "HeapSize" + Params.ClientConfigsToBuild[0].ToString(), out ConfigHeapSize);
+
+        // Fallback if the previous method failed
+        if (!bGotHeapSize && !ConfigCache.GetInt32("BuildSettings", "HeapSize" + Params.ClientConfigsToBuild[0].ToString(), out ConfigHeapSize)) // in Megs.
+        {
+            // we couldn't find a per config heap size, look for a common one.
+            if (!ConfigCache.GetInt32("BuildSettings", "HeapSize", out ConfigHeapSize))
+            {
+                ConfigHeapSize = Params.IsCodeBasedProject ? 1024 : 512;
+                Log("Could not find Heap Size setting in .ini for Client config {0}", Params.ClientConfigsToBuild[0].ToString());
+            }
+        }
 
+        HeapSize = (ulong)ConfigHeapSize * 1024L * 1024L; // convert to bytes.
+        Log("Setting Heap size to {0} Mb ", ConfigHeapSize);
 
-            GenerateFileFromTemplate(TemplateFile, OutputFile, Params.ShortProjectName, Params.ClientConfigsToBuild[0].ToString(), Params.StageCommandline, !Params.IsCodeBasedProject, HeapSize);
 
-            // copy the jstorage files to the binaries directory
-            string JSDir = Path.Combine(CombinePaths(CmdEnv.LocalRoot, "Engine"), "Build", "HTML5");
-            string OutDir = PackagePath;
-            File.Copy(JSDir + "/json2.js", OutDir + "/json2.js", true);
-            File.SetAttributes(OutDir + "/json2.js", FileAttributes.Normal);
-            File.Copy(JSDir + "/jStorage.js", OutDir + "/jStorage.js", true);
-            File.SetAttributes(OutDir + "/jStorage.js", FileAttributes.Normal);
-            File.Copy(JSDir + "/moz_binarystring.js", OutDir + "/moz_binarystring.js", true);
-            File.SetAttributes(OutDir + "/moz_binarystring.js", FileAttributes.Normal);
-            PrintRunTime();
+        GenerateFileFromTemplate(TemplateFile, OutputFile, Params.ShortProjectName, Params.ClientConfigsToBuild[0].ToString(), Params.StageCommandline, !Params.IsCodeBasedProject, HeapSize);
+
+        // copy the jstorage files to the binaries directory
+        string JSDir = Path.Combine(CombinePaths(CmdEnv.LocalRoot, "Engine"), "Build", "HTML5");
+        string OutDir = PackagePath;
+        File.Copy(JSDir + "/json2.js", OutDir + "/json2.js", true);
+        File.SetAttributes(OutDir + "/json2.js", FileAttributes.Normal);
+        File.Copy(JSDir + "/jStorage.js", OutDir + "/jStorage.js", true);
+        File.SetAttributes(OutDir + "/jStorage.js", FileAttributes.Normal);
+        File.Copy(JSDir + "/moz_binarystring.js", OutDir + "/moz_binarystring.js", true);
+        File.SetAttributes(OutDir + "/moz_binarystring.js", FileAttributes.Normal);
+        PrintRunTime();
 	}
 
 
@@ -211,6 +248,23 @@ public class HTML5Platform : Platform
 		SC.ArchiveFiles(PackagePath, Path.GetFileName("jStorage.js"));
 		SC.ArchiveFiles(PackagePath, Path.GetFileName("moz_binarystring.js"));
 		SC.ArchiveFiles(PackagePath, Path.GetFileName(OutputFile));
+        
+
+        bool UseAsyncLevelLoading = false;
+        var ConfigCache = new UnrealBuildTool.ConfigCacheIni(UnrealTargetPlatform.HTML5, "Engine", Path.GetDirectoryName(Params.RawProjectPath), CombinePaths(CmdEnv.LocalRoot, "Engine"));
+        ConfigCache.GetBool("/Script/HTML5PlatformEditor.HTML5TargetSettings", "UseAsyncLevelLoading", out UseAsyncLevelLoading);
+
+        if (UseAsyncLevelLoading)
+        {
+            // find all paks.
+            string[] Files = Directory.GetFiles(Path.Combine(PackagePath, Params.ShortProjectName), "*",SearchOption.AllDirectories);
+            foreach(string PakFile in Files)
+            {
+                var DestPak = PakFile.Replace(PackagePath,"");
+                SC.ArchivedFiles.Add(PakFile, DestPak);
+            }
+        }
+
 	}
 
 	public override ProcessResult RunClient(ERunOptions ClientRunFlags, string ClientApp, string ClientCmdLine, ProjectParams Params)
@@ -315,6 +369,22 @@ public class HTML5Platform : Platform
 		return "HTML5";
 	}
 
+    public override string GetCookExtraCommandLine(ProjectParams Params)
+    {
+        var ConfigCache = new UnrealBuildTool.ConfigCacheIni(UnrealTargetPlatform.HTML5, "Engine", Path.GetDirectoryName(Params.RawProjectPath), CombinePaths(CmdEnv.LocalRoot, "Engine"));
+        bool UseAsyncLevelLoading = false;
+        ConfigCache.GetBool("/Script/HTML5PlatformEditor.HTML5TargetSettings", "UseAsyncLevelLoading", out UseAsyncLevelLoading);
+
+        return UseAsyncLevelLoading ? " -GenerateDependenciesForMaps " : "";
+    }
+
+    public override List<string> GetCookExtraMaps()
+    {
+        var Maps = new List<string>();
+        Maps.Add("/Engine/Maps/Loading");
+        return Maps; 
+    }
+
 	public override bool DeployPakInternalLowerCaseFilenames()
 	{
 		return false;
@@ -322,7 +392,11 @@ public class HTML5Platform : Platform
 
     public override PakType RequiresPak(ProjectParams Params)
     {
-        return PakType.Always;
+        var ConfigCache = new UnrealBuildTool.ConfigCacheIni(UnrealTargetPlatform.HTML5, "Engine", Path.GetDirectoryName(Params.RawProjectPath), CombinePaths(CmdEnv.LocalRoot, "Engine"));
+        bool UseAsyncLevelLoading = false; 
+        ConfigCache.GetBool("/Script/HTML5PlatformEditor.HTML5TargetSettings", "UseAsyncLevelLoading", out UseAsyncLevelLoading);
+
+        return UseAsyncLevelLoading ? PakType.Never : PakType.Always;
     }
 
 	public override string GetPlatformPakCommandLine()
diff --git a/Engine/Source/Programs/AutomationTool/HTML5/HTML5Platform.PakFiles.Automation.cs b/Engine/Source/Programs/AutomationTool/HTML5/HTML5Platform.PakFiles.Automation.cs
new file mode 100644
--- /dev/null
+++ b/Engine/Source/Programs/AutomationTool/HTML5/HTML5Platform.PakFiles.Automation.cs
@@ -0,0 +1,337 @@
+﻿// Copyright 1998-2015 Epic Games, Inc. All Rights Reserved.
+using System;
+using System.Collections.Generic;
+using System.Text.RegularExpressions;
+using System.Linq;
+using System.Text;
+using System.IO;
+using AutomationTool;
+using UnrealBuildTool;
+using System.Diagnostics;
+
+
+public class HTMLPakAutomation
+{
+    ProjectParams Params;
+
+    DeploymentContext SC;
+
+    string PakOrderFileLocation;
+
+    Dictionary<string, object> DependencyJson;
+
+    public HTMLPakAutomation(ProjectParams InParam, DeploymentContext InSC)
+    {
+        Params = InParam;
+        SC = InSC;
+
+        var PakOrderFileLocationBase = CommandUtils.CombinePaths(SC.ProjectRoot, "Build", SC.FinalCookPlatform, "FileOpenOrder");
+        PakOrderFileLocation = CommandUtils.CombinePaths(PakOrderFileLocationBase, "GameOpenOrder.log");
+        if (!CommandUtils.FileExists_NoExceptions(PakOrderFileLocation))
+        {
+            // Use a fall back, it doesn't matter if this file exists or not. GameOpenOrder.log is preferred however
+            PakOrderFileLocation = CommandUtils.CombinePaths(PakOrderFileLocationBase, "EditorOpenOrder.log");
+        }
+
+        string PakPath = Path.Combine(new string[] { Path.GetDirectoryName(Params.RawProjectPath), "Binaries", "HTML5", SC.ShortProjectName});
+        if (Directory.Exists(PakPath))
+        {
+            Directory.Delete(PakPath,true);
+        }
+
+        // read in the json file. 
+        string JsonFile = CommandUtils.CombinePaths(new string[] { SC.ProjectRoot, "Saved", "Cooked", "HTML5", Params.ShortProjectName, "MapDependencyGraph.json" });
+        string text = File.ReadAllText(JsonFile);
+
+        DependencyJson = fastJSON.JSON.Instance.ToObject<Dictionary<string, object>>(text);
+
+    }
+
+    /// <summary>
+    /// Create Emscripten Data package with Necessary Packages. 
+    /// 
+    /// </summary>
+    public void CreateEmscriptenDataPackage(string PackagePath, string FinalDataLocation)
+    {
+        // we need to operate in the root
+        using (new PushedDirectory(Path.Combine(PackagePath)))
+        {
+            string PythonPath = HTML5SDKInfo.PythonPath();
+            string PackagerPath = HTML5SDKInfo.EmscriptenPackager();
+
+            string CmdLine = string.Format("\"{0}\" \"{1}\" --preload {2}\\ --js-output=\"{1}.js\"", PackagerPath, FinalDataLocation, SC.ShortProjectName);
+            CommandUtils.RunAndLog(CommandUtils.CmdEnv, PythonPath, CmdLine);
+        }
+    }
+
+    /// <summary>
+    /// Create a Pak from Engine Contents. 
+    /// </summary>
+    public void CreateEnginePak()
+    {
+        string StagedDir = Path.Combine(Params.BaseStageDirectory, "HTML5");// CommandUtils.CombinePaths(SC.ProjectRoot, "Saved", "Cooked", "HTML5");
+    
+         string PakPath = Path.Combine( new string [] { Path.GetDirectoryName(Params.RawProjectPath), "Binaries", "HTML5", SC.ShortProjectName, "Content","Paks"});
+         if (!Directory.Exists(PakPath))
+         {
+             Directory.CreateDirectory(PakPath);
+         }
+
+         // find list of files in the Engine Dir. 
+         string[] files = Directory.GetFiles( StagedDir + "\\Engine\\", "*", SearchOption.AllDirectories);
+
+         Dictionary<string,string> UnrealPakResponseFile = new Dictionary<string,string>();
+
+         // create response file structure. 
+         foreach( string file in files )  
+         {
+             string MountPoint = file.Remove(0, StagedDir.Length + 1);
+             UnrealPakResponseFile[file] = "../../../" +  MountPoint; 
+         }
+
+        RunUnrealPak(UnrealPakResponseFile, PakPath + "\\EnginePak.pak", null, "", true);
+    }
+
+
+    /// <summary>
+    /// Create a pak from a minimum set of game files
+    /// </summary>
+    public void CreateGamePak()
+    {
+        // create a pak of game config and first level files. 
+
+        string StagedDir = Path.Combine(Params.BaseStageDirectory, "HTML5");// CommandUtils.CombinePaths(SC.ProjectRoot, "Saved", "Cooked", "HTML5");
+
+        string PakPath = Path.Combine(new string[] { Path.GetDirectoryName(Params.RawProjectPath), "Binaries", "HTML5", Params.ShortProjectName, "Content", "Paks" });
+        if (!Directory.Exists(PakPath))
+        {
+            Directory.CreateDirectory(PakPath);
+        }
+
+        List<string> AllFiles = new List<string>();
+        try
+        {
+            string[] SlateFiles = Directory.GetFiles(StagedDir + "\\" + Params.ShortProjectName + "\\Content\\Slate", "*", SearchOption.AllDirectories);
+            AllFiles.AddRange(SlateFiles);
+        }
+        catch ( Exception /*Ex*/)
+        {
+            // not found. 
+        }
+
+        try
+        {
+            string[] Resources = Directory.GetFiles(StagedDir + "\\" + Params.ShortProjectName + "\\Content\\Resources", "*", SearchOption.AllDirectories);
+            AllFiles.AddRange(Resources);
+        }
+        catch (Exception /*Ex*/)
+        {
+            // not found. 
+        }
+
+        string[] ConfigFiles = Directory.GetFiles(StagedDir + "\\" + Params.ShortProjectName + "\\Config", "*", SearchOption.AllDirectories);
+
+        AllFiles.AddRange(ConfigFiles);
+
+        string[] TopLevelFiles = Directory.GetFiles(StagedDir + "\\" + Params.ShortProjectName + "\\", "*", SearchOption.TopDirectoryOnly);
+
+        AllFiles.AddRange(TopLevelFiles);
+
+        Dictionary<string, string> UnrealPakResponseFile = new Dictionary<string, string>();
+
+        // create response file structure. 
+        foreach (string file in AllFiles)
+        {
+            string MountPoint = file.Remove(0, StagedDir.Length + 1);
+            UnrealPakResponseFile[file] = "../../../" + MountPoint;
+        }
+
+        RunUnrealPak(UnrealPakResponseFile, PakPath + "\\" + Params.ShortProjectName + ".pak", null, "", true);
+        
+    }
+
+    /// <summary>
+    /// Create ContentDirectory Pak.
+    /// </summary>
+    public void CreateContentDirectoryPak()
+    {
+        object[] Value = (object[])DependencyJson["ContentDirectoryAssets"];
+        List<string> ContentDirectory = Array.ConvertAll<object, string>(Value, Convert.ToString).ToList();
+        CreatePak("ContentDirectoryAssets", ContentDirectory);
+    }
+
+
+    /// <summary>
+    /// Create Maps Paks based on JSON file created during Cooking. (-GenerateDependenciesForMaps)
+    /// </summary>
+    public void CreateMapPak()
+    {
+        // read in the json file. 
+
+        object[] Value = (object[])DependencyJson["ContentDirectoryAssets"];
+        List<string> ContentDirectory = Array.ConvertAll<object, string>(Value, Convert.ToString).ToList();
+
+        foreach (string Pak in DependencyJson.Keys)
+        {
+            object[] ObjectList = (object[])DependencyJson[Pak];
+            List<string> Assets = Array.ConvertAll<object, string>(ObjectList, Convert.ToString).ToList();
+
+            // add itself.
+            Assets.Add(Pak);
+
+            if (Pak != "ContentDirectoryAssets")
+            {
+                // remove content directory assets. 
+                Assets = Assets.Except(ContentDirectory).ToList();
+                string[] Tokens = Pak.Split('/');
+                CreatePak(Tokens[Tokens.Length - 1], Assets);
+            }
+        }
+
+    }
+
+    /// <summary>
+    /// Create Delta Paks based on Level Transitions. 
+    /// </summary>
+    /// <param name="MapFrom"></param>
+    /// <param name="MapTo"></param>
+    public void CreateDeltaMapPaks(string MapFrom, string MapTo)
+    {
+        foreach (string NameKey in DependencyJson.Keys)
+        {
+            if (NameKey.Contains(MapFrom))
+            { 
+                MapFrom = NameKey;
+            }
+            if (NameKey.Contains(MapTo))
+            { 
+                MapTo = NameKey;
+            }
+        }
+
+        List<string> MapFromAssets = Array.ConvertAll<object, string>(DependencyJson[MapFrom] as object[], Convert.ToString).ToList();
+        MapFromAssets.Add(MapFrom);
+        List<string> MapToAssets = Array.ConvertAll<object, string>(DependencyJson[MapTo] as object[], Convert.ToString).ToList();
+        MapToAssets.Add(MapTo);
+
+        object[] Value = (object[])DependencyJson["ContentDirectoryAssets"];
+        List<string> ContentDirectory = Array.ConvertAll<object, string>(Value, Convert.ToString).ToList();
+
+        // Delta. 
+        List<string> DeltaAssets = MapToAssets.Except(MapFromAssets).Except(ContentDirectory).ToList();
+
+        // create a pak file for this delta. 
+        string[] Tokens_MapFrom = MapFrom.Split('/');
+        string[] Tokens_MapTo = MapTo.Split('/');
+
+        string DeltaPakName = Tokens_MapFrom[Tokens_MapFrom.Length - 1] + "_" + Tokens_MapTo[Tokens_MapTo.Length - 1];
+        CreatePak(DeltaPakName, DeltaAssets);
+    }
+
+    /// <summary>
+    /// Writes a pak response file to disk
+    /// </summary>
+    /// <param name="Filename"></param>
+    /// <param name="ResponseFile"></param>
+    private static void WritePakResponseFile(string Filename, Dictionary<string, string> ResponseFile, bool Compressed)
+    {
+        using (var Writer = new StreamWriter(Filename, false, new System.Text.UTF8Encoding(true)))
+        {
+            foreach (var Entry in ResponseFile)
+            {
+                string Line = String.Format("\"{0}\" \"{1}\"", Entry.Key, Entry.Value);
+                if (Compressed)
+                {
+                    Line += " -compress";
+                }
+                Writer.WriteLine(Line);
+            }
+        }
+    }
+
+    /// <summary>
+    /// Create a Pak file
+    /// </summary>
+    private void RunUnrealPak(Dictionary<string, string> UnrealPakResponseFile, string OutputLocation, string EncryptionKeys, string PlatformOptions, bool Compressed)
+    {
+        if (UnrealPakResponseFile.Count < 1)
+        {
+            return;
+        }
+        string PakName = Path.GetFileNameWithoutExtension(OutputLocation);
+        string UnrealPakResponseFileName = CommandUtils.CombinePaths(CommandUtils.CmdEnv.LogFolder, "PakList_" + PakName + ".txt");
+        WritePakResponseFile(UnrealPakResponseFileName, UnrealPakResponseFile, Compressed);
+
+        var UnrealPakExe = CommandUtils.CombinePaths(CommandUtils.CmdEnv.LocalRoot, "Engine/Binaries/Win64/UnrealPak.exe");
+
+        string CmdLine = CommandUtils.MakePathSafeToUseWithCommandLine(OutputLocation) + " -create=" + CommandUtils.MakePathSafeToUseWithCommandLine(UnrealPakResponseFileName);
+        if (!String.IsNullOrEmpty(EncryptionKeys))
+        {
+            CmdLine += " -sign=" + CommandUtils.MakePathSafeToUseWithCommandLine(EncryptionKeys);
+        }
+        if (GlobalCommandLine.Installed)
+        {
+            CmdLine += " -installed";
+        }
+        CmdLine += " -order=" + CommandUtils.MakePathSafeToUseWithCommandLine(PakOrderFileLocation);
+        if (GlobalCommandLine.UTF8Output)
+        {
+            CmdLine += " -UTF8Output";
+        }
+        CmdLine += PlatformOptions;
+        CommandUtils.RunAndLog(CommandUtils.CmdEnv, UnrealPakExe, CmdLine, Options: CommandUtils.ERunOptions.Default | CommandUtils.ERunOptions.UTF8Output);
+    }
+
+    /// <summary>
+    /// Create a Pak.
+    /// </summary>
+    /// <param name="MapFrom"></param>
+    /// <param name="MapTo"></param>
+    private void CreatePak(string Name, List<string> Assets)
+    {
+        string CookedDir = Path.Combine(Params.BaseStageDirectory, "HTML5");
+
+        string PakPath = Path.Combine(new string[] { Path.GetDirectoryName(Params.RawProjectPath), "Binaries", "HTML5", Params.ShortProjectName, "Content", "Paks" });
+        if (!Directory.Exists(PakPath))
+        {
+            Directory.CreateDirectory(PakPath);
+        }
+
+        Assets.Add(Name); 
+        List<string> Files = new List<string>();
+        foreach( string Asset in Assets)
+        {
+            string GameDir = "/Game/";
+            if (Asset.StartsWith(GameDir))
+            {
+                string PathOnDiskasset  = CommandUtils.CombinePaths(CookedDir, Params.ShortProjectName, "Content", Asset.Remove(0, GameDir.Length) + ".uasset");
+                string PathOnDiskmap    = CommandUtils.CombinePaths(CookedDir, Params.ShortProjectName, "Content", Asset.Remove(0, GameDir.Length)+ ".umap");
+
+                if (File.Exists(PathOnDiskmap))
+                {
+                    Files.Add(PathOnDiskmap);
+                }
+                else if (File.Exists(PathOnDiskasset))
+                {
+                    Files.Add(PathOnDiskasset);
+                }
+                else
+                {
+                    System.Console.WriteLine(Asset + " not found, skipping !!"); 
+                }
+            }
+        }
+
+        // we have all the files. 
+        Dictionary<string, string> UnrealPakResponseFile = new Dictionary<string, string>();
+
+        // create response file structure. 
+        foreach (string file in Files)
+        {
+            string MountPoint = file.Remove(0, CookedDir.Length + 1);
+            UnrealPakResponseFile[file] ="../../../" +  MountPoint;
+        }
+        RunUnrealPak(UnrealPakResponseFile, PakPath + "\\"+ Name + ".pak", null, "", true);
+    }
+
+}
\ No newline at end of file
diff --git a/Engine/Source/Runtime/Core/Core.Build.cs b/Engine/Source/Runtime/Core/Core.Build.cs
--- a/Engine/Source/Runtime/Core/Core.Build.cs
+++ b/Engine/Source/Runtime/Core/Core.Build.cs
@@ -141,6 +141,7 @@ public class Core : ModuleRules
         {
             AddThirdPartyPrivateStaticDependencies(Target, "SDL2");
             PrivateDependencyModuleNames.Add("HTML5JS");
+            PrivateDependencyModuleNames.Add("MapPakDownloader");
         }
 
         if ( UEBuildConfiguration.bCompileICU == true ) 
diff --git a/Engine/Source/Runtime/Core/Private/HTML5/HTML5PlatformMisc.cpp b/Engine/Source/Runtime/Core/Private/HTML5/HTML5PlatformMisc.cpp
--- a/Engine/Source/Runtime/Core/Private/HTML5/HTML5PlatformMisc.cpp
+++ b/Engine/Source/Runtime/Core/Private/HTML5/HTML5PlatformMisc.cpp
@@ -18,6 +18,10 @@
 #include <SDL.h>
 #include <ctime>
 
+#include "MapPakDownloaderModule.h"
+#include "MapPakDownloader.h"
+
+
 void FHTML5Misc::PlatformInit()
 {
 	// Identity.
@@ -158,3 +162,14 @@ void FHTML5Misc::SetCrashHandler(void(* CrashHandler)(const FGenericCrashContext
 {
 	GHTML5CrashHandler = CrashHandler;
 }
+
+const void FHTML5Misc::PreLoadMap(FString& Map, FString& LastMap, void* DynData)
+{
+	static TSharedPtr<FMapPakDownloader> Downloader = FModuleManager::GetModulePtr<IMapPakDownloaderModule>("MapPakDownloader")->GetDownloader();
+	Downloader->Cache(Map, LastMap, DynData);
+}
+
+void FHTML5Misc::PlatformPostInit(bool ShowSplashScreen /*= false*/)
+{
+	FModuleManager::Get().LoadModule("MapPakDownloader");
+}
diff --git a/Engine/Source/Runtime/Core/Public/HTML5/HTML5PlatformMisc.h b/Engine/Source/Runtime/Core/Public/HTML5/HTML5PlatformMisc.h
--- a/Engine/Source/Runtime/Core/Public/HTML5/HTML5PlatformMisc.h
+++ b/Engine/Source/Runtime/Core/Public/HTML5/HTML5PlatformMisc.h
@@ -18,6 +18,7 @@
 struct CORE_API FHTML5Misc : public FGenericPlatformMisc
 {
 	static void PlatformInit();
+	static void PlatformPostInit(bool ShowSplashScreen = false);
 	static class GenericApplication* CreateApplication();
 	static uint32 GetKeyMap(uint16* KeyCodes, FString* KeyNames, uint32 MaxMappings);
 	static uint32 GetCharKeyMap(uint16* KeyCodes, FString* KeyNames, uint32 MaxMappings);
@@ -85,6 +86,8 @@ struct CORE_API FHTML5Misc : public FGenericPlatformMisc
 	{
 		wprintf(TEXT("%ls"), Str);
 	}
+
+	static const void PreLoadMap(FString&, FString&, void*);
 };
 
 typedef FHTML5Misc FPlatformMisc;
diff --git a/Engine/Source/Runtime/HTML5/MapPakDownloader/MapPakDownloader.Build.cs b/Engine/Source/Runtime/HTML5/MapPakDownloader/MapPakDownloader.Build.cs
new file mode 100644
--- /dev/null
+++ b/Engine/Source/Runtime/HTML5/MapPakDownloader/MapPakDownloader.Build.cs
@@ -0,0 +1,20 @@
+// Copyright 1998-2015 Epic Games, Inc. All Rights Reserved.
+
+namespace UnrealBuildTool.Rules
+{
+	public class MapPakDownloader: ModuleRules
+	{
+        public MapPakDownloader(TargetInfo Target)
+		{
+    		PublicDependencyModuleNames.AddRange(
+				new string[]
+				{
+					"Core",
+                    "CoreUObject",
+                    "Engine",
+					// ... add other public dependencies that you statically link with here ...
+				}
+				);
+    	}
+	}
+}
\ No newline at end of file
diff --git a/Engine/Source/Runtime/HTML5/MapPakDownloader/Private/MapPakDownloader.cpp b/Engine/Source/Runtime/HTML5/MapPakDownloader/Private/MapPakDownloader.cpp
new file mode 100644
--- /dev/null
+++ b/Engine/Source/Runtime/HTML5/MapPakDownloader/Private/MapPakDownloader.cpp
@@ -0,0 +1,167 @@
+// Copyright 1998-2015 Epic Games, Inc. All Rights Reserved.
+// Precompiled header include. Can't add anything above this line.
+
+#include "MapPakDownloaderModulePrivatePCH.h"
+#include "MapPakDownloader.h"
+#include "emscripten.h"
+#include "Misc/Guid.h"
+
+DECLARE_DELEGATE_OneParam(FDelegateFString, FString)
+DECLARE_DELEGATE_OneParam(FDelegateInt32, int32)
+
+// avoid IHttpModule
+class FEmscriptenHttpFileRequest
+{
+public:
+	typedef int32 Handle;
+private:
+
+	FString FileName;
+	FString Url;
+
+	FDelegateFString OnLoadCallBack;
+	FDelegateInt32 OnErrorCallBack;
+	FDelegateInt32 OnProgressCallBack;
+
+public:
+
+	FEmscriptenHttpFileRequest()
+	{}
+	
+	static void OnLoad(uint32 Var, void* Arg, const ANSICHAR* FileName)
+	{
+		FEmscriptenHttpFileRequest* Request = reinterpret_cast<FEmscriptenHttpFileRequest*>(Arg);
+		Request->OnLoadCallBack.ExecuteIfBound(FString(ANSI_TO_TCHAR(FileName)));
+	}
+	
+	static void OnError(uint32 Var,void* Arg, int32 ErrorCode)
+	{
+		FEmscriptenHttpFileRequest* Request = reinterpret_cast<FEmscriptenHttpFileRequest*>(Arg);
+		Request->OnErrorCallBack.ExecuteIfBound(ErrorCode);
+	}
+
+	static void OnProgress(uint32 Var,void* Arg, int32 Progress)
+	{
+		FEmscriptenHttpFileRequest* Request = reinterpret_cast<FEmscriptenHttpFileRequest*>(Arg);
+		Request->OnProgressCallBack.ExecuteIfBound(Progress);
+	}
+
+	void Process()
+	{
+		UE_LOG(LogMapPakDownloader, Warning, TEXT("Starting Download for %s"), *FileName);
+
+		emscripten_async_wget2(
+			TCHAR_TO_ANSI(*(Url + FString(TEXT("?rand=")) + FGuid::NewGuid().ToString())),
+			TCHAR_TO_ANSI(*FileName),
+			TCHAR_TO_ANSI(TEXT("GET")),
+			TCHAR_TO_ANSI(TEXT("")),
+			this,
+			&FEmscriptenHttpFileRequest::OnLoad,
+			&FEmscriptenHttpFileRequest::OnError,
+			&FEmscriptenHttpFileRequest::OnProgress
+			);
+	}
+
+	void SetFileName(const FString& InFileName) { FileName = InFileName; }
+	FString GetFileName() { return FileName;  }
+	
+	void SetUrl(const FString& InUrl) { Url = InUrl; }
+	void SetOnLoadCallBack(const FDelegateFString& CallBack) { OnLoadCallBack = CallBack; }
+	void SetOnErrorCallBack(const FDelegateInt32& CallBack) { OnErrorCallBack = CallBack; }
+	void SetOnProgressCallBack(const FDelegateInt32& CallBack) { OnProgressCallBack = CallBack; }
+
+};
+
+FMapPakDownloader::FMapPakDownloader()
+{}
+
+bool FMapPakDownloader::Init()
+{
+	//  figure out where we are hosted 
+	ANSICHAR *LocationString = (ANSICHAR*)EM_ASM_INT_V({
+
+		var hoststring = "http://" + location.host; 
+		var buffer = Module._malloc(hoststring.length);
+		Module.writeAsciiToMemory(hoststring, buffer);
+		return buffer;
+
+	});
+
+	HostName = FString(ANSI_TO_TCHAR(LocationString));
+
+	PakLocation = FString(FApp::GetGameName()) / FString(TEXT("Content")) / FString(TEXT("Paks")); 
+
+	// Create directory. 
+	IPlatformFile& PlatformFile = FPlatformFileManager::Get().GetPlatformFile();
+	PlatformFile.CreateDirectory(*PakLocation);
+
+	return true;
+}
+
+void FMapPakDownloader::CachePak()
+{
+	UE_LOG(LogMapPakDownloader, Warning, TEXT("Caching Dependencies for %s"), *MapToCache);
+
+	FString PakName = FPackageName::GetShortName(FName(*MapToCache)) + ".pak";
+	FString DeltaPakName = FPackageName::GetShortName(FName(*LastMap)) + TEXT("_") + FPackageName::GetShortName(FName(*MapToCache)) + ".pak";
+
+	FEmscriptenHttpFileRequest* PakRequest = new FEmscriptenHttpFileRequest; // can't use shared ptrs. 
+
+
+	FDelegateFString OnFileDownloaded = FDelegateFString::CreateLambda([=](FString Name){
+
+											UE_LOG(LogMapPakDownloader, Warning, TEXT("%s download complete!"), *PakRequest->GetFileName());
+											UE_LOG(LogMapPakDownloader, Warning, TEXT("Mounting..."), *Name);
+											FCoreDelegates::OnMountPak.Execute(Name, 0);
+											UE_LOG(LogMapPakDownloader, Warning, TEXT("%s Mounted!"), *Name);
+
+											// Get hold of the world.
+											TObjectIterator<UWorld> It;
+											UWorld* ItWorld = *It;
+
+											UE_LOG(LogMapPakDownloader, Warning, TEXT("Travel to %s"), *MapToCache);
+											// Make engine Travel to the cached Map. 
+											GEngine->SetClientTravel(ItWorld, *MapToCache, TRAVEL_Absolute);
+
+											// delete the HTTP request. 
+											delete PakRequest; 
+										});
+
+	FDelegateInt32 OnFileDownloadProgress = FDelegateInt32::CreateLambda([=](int32 Progress){
+												UE_LOG(LogMapPakDownloader, Warning, TEXT(" %s %d%% downloaded"), *PakRequest->GetFileName(), Progress);
+											});
+
+	PakRequest->SetFileName(PakLocation / DeltaPakName);
+	PakRequest->SetUrl(HostName / PakLocation / DeltaPakName);
+	PakRequest->SetOnLoadCallBack(OnFileDownloaded);
+	PakRequest->SetOnProgressCallBack(OnFileDownloadProgress);
+	PakRequest->SetOnErrorCallBack(FDelegateInt32::CreateLambda([=](int32){
+
+											UE_LOG(LogMapPakDownloader, Warning, TEXT("Could not download %s"), *PakRequest->GetFileName());
+											// lets try again with regular map pak
+											PakRequest->SetFileName(PakLocation/PakName);
+											PakRequest->SetUrl(HostName/PakLocation/PakName);
+											PakRequest->SetOnErrorCallBack( 
+												FDelegateInt32::CreateLambda([=](int32){
+												// we can't even find regular maps. fatal. 
+												UE_LOG(LogMapPakDownloader, Fatal, TEXT("Could not find any Map Paks, exiting"), *PakRequest->GetFileName());
+												}
+											));
+											PakRequest->Process();
+
+									}));
+
+	PakRequest->Process();
+}
+
+void FMapPakDownloader::Cache(FString& Map, FString& InLastMap, void* InDynData)
+{
+	MapToCache = Map;
+	LastMap = InLastMap;
+
+	if (!FPackageName::DoesPackageExist(Map))
+	{
+		CachePak();
+		Map = TEXT("/Engine/Maps/Loading");
+	}
+}
diff --git a/Engine/Source/Runtime/HTML5/MapPakDownloader/Private/MapPakDownloaderModule.cpp b/Engine/Source/Runtime/HTML5/MapPakDownloader/Private/MapPakDownloaderModule.cpp
new file mode 100644
--- /dev/null
+++ b/Engine/Source/Runtime/HTML5/MapPakDownloader/Private/MapPakDownloaderModule.cpp
@@ -0,0 +1,52 @@
+// Copyright 1998-2015 Epic Games, Inc. All Rights Reserved.
+
+#include "MapPakDownloaderModulePrivatePCH.h"
+#include "MapPakDownloaderModule.h"
+#include "MapPakDownloader.h"
+
+
+class FMapPakDownloaderModule : public IMapPakDownloaderModule
+{
+	/** IModuleInterface implementation */
+	virtual void StartupModule() override;
+	virtual void ShutdownModule() override;
+	virtual TSharedPtr<FMapPakDownloader> GetDownloader() override;
+
+private: 
+
+	TSharedPtr<FMapPakDownloader> MapPakDownloader;
+
+};
+
+IMPLEMENT_MODULE(FMapPakDownloaderModule, MapPakDownloader)
+
+
+void FMapPakDownloaderModule::StartupModule()
+{
+	// This code will execute after your module is loaded into memory (but after global variables are initialized, of course.)
+	TSharedPtr<FMapPakDownloader> Downloader = MakeShareable(new FMapPakDownloader());
+
+	if (!Downloader->Init())
+	{
+		ShutdownModule();
+	}
+	MapPakDownloader = Downloader;
+}
+
+
+void FMapPakDownloaderModule::ShutdownModule()
+{
+	// This function may be called during shutdown to clean up your module.  For modules that support dynamic reloading,
+	// we call this function before unloading the module.
+}
+
+TSharedPtr<FMapPakDownloader> FMapPakDownloaderModule::GetDownloader()
+{
+	return MapPakDownloader;
+}
+
+
+DEFINE_LOG_CATEGORY(LogMapPakDownloader);
+
+
+
diff --git a/Engine/Source/Runtime/HTML5/MapPakDownloader/Private/MapPakDownloaderModulePrivatePCH.h b/Engine/Source/Runtime/HTML5/MapPakDownloader/Private/MapPakDownloaderModulePrivatePCH.h
new file mode 100644
--- /dev/null
+++ b/Engine/Source/Runtime/HTML5/MapPakDownloader/Private/MapPakDownloaderModulePrivatePCH.h
@@ -0,0 +1,9 @@
+// Copyright 1998-2015 Epic Games, Inc. All Rights Reserved.
+#pragma  once
+#include "Core.h"
+#include "CoreUObject.h"
+#include "Engine.h"
+// You should place include statements to your module's private header files here.  You only need to
+// add includes for headers that are used in most of your module's source files though.
+
+DECLARE_LOG_CATEGORY_EXTERN(LogMapPakDownloader, Log, All);
\ No newline at end of file
diff --git a/Engine/Source/Runtime/HTML5/MapPakDownloader/Public/MapPakDownloader.h b/Engine/Source/Runtime/HTML5/MapPakDownloader/Public/MapPakDownloader.h
new file mode 100644
--- /dev/null
+++ b/Engine/Source/Runtime/HTML5/MapPakDownloader/Public/MapPakDownloader.h
@@ -0,0 +1,47 @@
+// Copyright 1998-2015 Epic Games, Inc. All Rights Reserved.
+
+#pragma once
+
+#include "Core.h"
+
+class FMapPakDownloader
+{
+
+public: 
+
+	FMapPakDownloader();
+	
+	/**
+	 *  Setup hostname and pak directories.
+	 */ 
+	bool Init();
+
+	/**
+	 *   Cache Map if needed and then transition			                     
+	 */
+	void Cache(FString& Url, FString& LastUrl, void* DynData);
+
+private: 
+
+	/**
+	 *   http download pak file and then transition to requested map. 															
+	 */
+	void CachePak();
+
+	// URL being Cached.
+	FString MapToCache; 
+
+	// Last URL. 
+	FString LastMap;
+
+	// Dynamic Data.
+	void* DynData; 
+	
+	// Server where paks are located. 
+	FString HostName;
+
+	// Relative dir of Pak files 
+	FString PakLocation;
+
+};
+
diff --git a/Engine/Source/Runtime/HTML5/MapPakDownloader/Public/MapPakDownloaderModule.h b/Engine/Source/Runtime/HTML5/MapPakDownloader/Public/MapPakDownloaderModule.h
new file mode 100644
--- /dev/null
+++ b/Engine/Source/Runtime/HTML5/MapPakDownloader/Public/MapPakDownloaderModule.h
@@ -0,0 +1,39 @@
+// Copyright 1998-2015 Epic Games, Inc. All Rights Reserved.
+
+#pragma once
+
+#include "ModuleManager.h"
+
+class FMapPakDownloader; 
+/**
+ * The public interface to this module
+ */
+class IMapPakDownloaderModule : public IModuleInterface
+{
+
+public:
+
+	/**
+	 * Singleton-like access to this module's interface.  This is just for convenience!
+	 * Beware of calling this during the shutdown phase, though.  Your module might have been unloaded already.
+	 *
+	 * @return Returns singleton instance, loading the module on demand if needed
+	 */
+	static inline IMapPakDownloaderModule& Get()
+	{
+		return FModuleManager::LoadModuleChecked< IMapPakDownloaderModule >("MapPakDownloader");
+	}
+
+	/**
+	 * Checks to see if this module is loaded and ready.  It is only valid to call Get() if IsAvailable() returns true.
+	 *
+	 * @return True if the module is loaded and ready to use
+	 */
+	static inline bool IsAvailable()
+	{
+		return FModuleManager::Get().IsModuleLoaded( "MapPakDownloader" );
+	}
+
+	virtual TSharedPtr <FMapPakDownloader> GetDownloader() = 0;
+};
+
