Don't submit stuff in the saved folder to Perforce in promotables.

diff --git a/Engine/Source/Programs/AutomationTool/Scripts/GUBP.Automation.cs b/Engine/Source/Programs/AutomationTool/Scripts/GUBP.Automation.cs
--- a/Engine/Source/Programs/AutomationTool/Scripts/GUBP.Automation.cs
+++ b/Engine/Source/Programs/AutomationTool/Scripts/GUBP.Automation.cs
@@ -2386,6 +2386,7 @@ public class GUBP : BuildCommand
 							!Product.Replace('\\', '/').EndsWith("DotNetCommon/MetaData.cs", StringComparison.InvariantCultureIgnoreCase) &&
                             !Product.EndsWith("_Success.log", StringComparison.InvariantCultureIgnoreCase) &&
 							!Product.Replace('\\', '/').Contains("/Intermediate/") &&
+							!Product.Replace('\\', '/').Contains("/Engine/Saved/") &&
 							!Product.EndsWith(".lib") &&
 							!Product.EndsWith(".a") && 
 							!Product.EndsWith(".bc")
