Linux: fix GUBP to work on Linux host.

diff --git a/Engine/Build/Target.cs.template b/Engine/Build/Target.cs.template
--- a/Engine/Build/Target.cs.template
+++ b/Engine/Build/Target.cs.template
@@ -75,7 +75,7 @@ public class {GAME_NAME}Target : TargetRules
 	}
 	public override List<GUBPFormalBuild> GUBP_GetConfigsForFormalBuilds_MonolithicOnly(UnrealTargetPlatform HostPlatform)
     {
-		if (HostPlatform == UnrealTargetPlatform.Win64)
+		if (HostPlatform != UnrealTargetPlatform.Mac)
 		{
 			return new List<GUBPFormalBuild>();			
 		}
diff --git a/Engine/Source/Programs/AutomationTool/Scripts/GUBP.Automation.cs b/Engine/Source/Programs/AutomationTool/Scripts/GUBP.Automation.cs
--- a/Engine/Source/Programs/AutomationTool/Scripts/GUBP.Automation.cs
+++ b/Engine/Source/Programs/AutomationTool/Scripts/GUBP.Automation.cs
@@ -5547,8 +5547,11 @@ public class GUBP : BuildCommand
                     }
                     else if (Proj.Properties.bIsCodeBasedProject)
                     {
-						Branch.NonCodeProjects.Add(Proj);
-						NonCodeProjectNames.Add(Codeless.Key, Codeless.Value);                        
+						if (!Branch.NonCodeProjects.Contains(Proj)) 
+						{ 
+							Branch.NonCodeProjects.Add(Proj);
+							NonCodeProjectNames.Add(Codeless.Key, Codeless.Value);
+						}
                     }
                     else
                     {
@@ -5798,7 +5801,7 @@ public class GUBP : BuildCommand
 									}
                                     if (!bNoAutomatedTesting)
                                     {                                        
-										if (HostPlatform == UnrealTargetPlatform.Mac) continue; //temp hack till mac automated testing works
+										if (HostPlatform == UnrealTargetPlatform.Mac || HostPlatform == UnrealTargetPlatform.Linux) continue; //temp hack till Linux and Mac automated testing works
 										var ThisMonoGameTestNodes = new List<string>();
 										foreach (var Test in GameTests)
 										{
@@ -6035,7 +6038,7 @@ public class GUBP : BuildCommand
 									}
 									if (!bNoAutomatedTesting)
 									{
-										if (HostPlatform == UnrealTargetPlatform.Mac) continue; //temp hack till mac automated testing works
+										if (HostPlatform == UnrealTargetPlatform.Mac || HostPlatform == UnrealTargetPlatform.Linux) continue; //temp hack till Linux and Mac automated testing works
 										var GameTests = Target.Rules.GUBP_GetGameTests_MonolithicOnly(HostPlatform, GetAltHostPlatform(HostPlatform), Plat);
 										var RequiredPlatforms = new List<UnrealTargetPlatform> { Plat };
 										var ThisMonoGameTestNodes = new List<string>();
@@ -6065,7 +6068,7 @@ public class GUBP : BuildCommand
 					{
 						foreach (var GamePlatform in GamePlatforms)
 						{
-							if (HostPlatform == UnrealTargetPlatform.Mac) continue; //temp hack till mac automated testing works
+							if (HostPlatform == UnrealTargetPlatform.Mac || HostPlatform == UnrealTargetPlatform.Linux) continue; //temp hack till Linux and Mac automated testing works
 							var Target = CodeProj.Properties.Targets[TargetRules.TargetType.Game];
 							var ClientServerTests = Target.Rules.GUBP_GetClientServerTests_MonolithicOnly(HostPlatform, GetAltHostPlatform(HostPlatform), ServerPlatform, GamePlatform);
 							var RequiredPlatforms = new List<UnrealTargetPlatform> { ServerPlatform };
diff --git a/Engine/Source/Programs/AutomationTool/Scripts/RocketBuild.Automation.cs b/Engine/Source/Programs/AutomationTool/Scripts/RocketBuild.Automation.cs
--- a/Engine/Source/Programs/AutomationTool/Scripts/RocketBuild.Automation.cs
+++ b/Engine/Source/Programs/AutomationTool/Scripts/RocketBuild.Automation.cs
@@ -129,6 +129,7 @@ namespace Rocket
 		public static List<UnrealTargetPlatform> GetTargetPlatforms(GUBP bp, UnrealTargetPlatform HostPlatform)
 		{
 			List<UnrealTargetPlatform> TargetPlatforms = new List<UnrealTargetPlatform>();
+			bool bSupportsMobilePlatforms = (HostPlatform == UnrealTargetPlatform.Win64 || HostPlatform == UnrealTargetPlatform.Mac);
 			if(!bp.ParseParam("NoPlatforms"))
 			{
 				TargetPlatforms.Add(HostPlatform);
@@ -137,11 +138,11 @@ namespace Rocket
 				{
 					TargetPlatforms.Add(UnrealTargetPlatform.Win32);
 				}
-				if(!bp.ParseParam("NoAndroid"))
+				if(!bp.ParseParam("NoAndroid") && bSupportsMobilePlatforms)
 				{
 					TargetPlatforms.Add(UnrealTargetPlatform.Android);
 				}
-				if(!bp.ParseParam("NoIOS"))
+				if(!bp.ParseParam("NoIOS") && bSupportsMobilePlatforms)
 				{
 					TargetPlatforms.Add(UnrealTargetPlatform.IOS);
 				}
