- Added support for DesktopPlatform to package content and executables together (editor's File | Package)
  - This is a first pass
  - Mac uses data outside the .app so that there's only one copy of the content. This should get a bootstrap executable like Windows has
  - Only tested on Mac so far, although it doesn't handle D3D11 shaders, or some of the icon updating for the bootstrap working
- Added Desktop platform icons

diff --git a/Engine/Source/Developer/TargetPlatform/Private/TargetPlatformManagerModule.cpp b/Engine/Source/Developer/TargetPlatform/Private/TargetPlatformManagerModule.cpp
--- a/Engine/Source/Developer/TargetPlatform/Private/TargetPlatformManagerModule.cpp
+++ b/Engine/Source/Developer/TargetPlatform/Private/TargetPlatformManagerModule.cpp
@@ -904,6 +904,11 @@ protected:
 					PlatformName = TEXT("LinuxServer");
 					PlatformInfo::UpdatePlatformSDKStatus(PlatformName, Status);
 				}
+				else if (PlatformName == TEXT("Desktop"))
+				{
+					// since Desktop is just packaging, we don't need an SDK, and UBT will return INVALID, since it doesn't build for it
+					PlatformInfo::UpdatePlatformSDKStatus(PlatformName, PlatformInfo::EPlatformSDKStatus::Installed);
+				}
 				else
 				{
 					PlatformInfo::UpdatePlatformSDKStatus(PlatformName, Status);
diff --git a/Engine/Source/Programs/AutomationTool/DeploymentContext.cs b/Engine/Source/Programs/AutomationTool/DeploymentContext.cs
--- a/Engine/Source/Programs/AutomationTool/DeploymentContext.cs
+++ b/Engine/Source/Programs/AutomationTool/DeploymentContext.cs
@@ -200,6 +200,10 @@ public class DeploymentContext //: ProjectParams
 	/// </summary>
 	public string ShortProjectName;
 
+	/// <summary>
+	/// If true, multiple platforms are being merged together - some behavior needs to change (but not much)
+	/// </summary>
+	public bool bIsCombiningMultiplePlatforms = false;
 
 	public DeploymentContext(
 		string RawProjectPathOrName,
@@ -352,7 +356,7 @@ public class DeploymentContext //: ProjectParams
                     continue;
                 }
 
-                if (bStripFilesForOtherPlatforms)
+				if (bStripFilesForOtherPlatforms && !bIsCombiningMultiplePlatforms)
                 {
                     bool OtherPlatform = false;
                     foreach (UnrealTargetPlatform Plat in Enum.GetValues(typeof(UnrealTargetPlatform)))
@@ -481,7 +485,7 @@ public class DeploymentContext //: ProjectParams
 			}
 		}
 
-		if (FilesAdded == 0 && !bAllowNone)
+		if (FilesAdded == 0 && !bAllowNone && !bIsCombiningMultiplePlatforms)
 		{
 			AutomationTool.ErrorReporter.Error(String.Format("No files found to deploy for {0} with wildcard {1} and exclusions {2}", InPath, Wildcard, ExcludeWildcard), (int)AutomationTool.ErrorCodes.Error_StageMissingFile);
 			throw new AutomationException("No files found to deploy for {0} with wildcard {1} and exclusions {2}", InPath, Wildcard, ExcludeWildcard);
@@ -535,18 +539,22 @@ public class DeploymentContext //: ProjectParams
 				{
 					continue;
 				}
-				bool OtherPlatform = false;
-				foreach (UnrealTargetPlatform Plat in Enum.GetValues(typeof(UnrealTargetPlatform)))
+
+				if (!bIsCombiningMultiplePlatforms)
 				{
-					if (Plat != StageTargetPlatform.PlatformType && Plat != UnrealTargetPlatform.Unknown && FileToCopy.IndexOf(CommandUtils.CombinePaths("/" + Plat.ToString() + "/"), 0, StringComparison.InvariantCultureIgnoreCase) >= 0)
+					bool OtherPlatform = false;
+					foreach (UnrealTargetPlatform Plat in Enum.GetValues(typeof(UnrealTargetPlatform)))
 					{
-						OtherPlatform = true;
-						break;
+						if (Plat != StageTargetPlatform.PlatformType && Plat != UnrealTargetPlatform.Unknown && FileToCopy.IndexOf(CommandUtils.CombinePaths("/" + Plat.ToString() + "/"), 0, StringComparison.InvariantCultureIgnoreCase) >= 0)
+						{
+							OtherPlatform = true;
+							break;
+						}
+					}
+					if (OtherPlatform)
+					{
+						continue;
 					}
-				}
-				if (OtherPlatform)
-				{
-					continue;
 				}
 
 				string Dest;
@@ -561,7 +569,17 @@ public class DeploymentContext //: ProjectParams
 					Dest = Dest.Substring(1);
 				}
 
-				ArchivedFiles.Add(FileToCopy, Dest);
+				if (ArchivedFiles.ContainsKey(FileToCopy))
+				{
+					if (ArchivedFiles[FileToCopy] != Dest)
+					{
+						throw new AutomationException("Can't archive {0}: it was already in the files to archive with a different destination '{1}'", FileToCopy, Dest);
+					}
+				}
+				else
+				{
+					ArchivedFiles.Add(FileToCopy, Dest);
+				}
 
 				FilesAdded++;
 			}
diff --git a/Engine/Source/Programs/AutomationTool/Desktop/DesktopPlatform.Automation.cs b/Engine/Source/Programs/AutomationTool/Desktop/DesktopPlatform.Automation.cs
--- a/Engine/Source/Programs/AutomationTool/Desktop/DesktopPlatform.Automation.cs
+++ b/Engine/Source/Programs/AutomationTool/Desktop/DesktopPlatform.Automation.cs
@@ -18,14 +18,44 @@ public class DesktopPlatform : Platform
 	{
 	}
 
-	public override void Package(ProjectParams Params, DeploymentContext SC, int WorkingCL)
+	public override bool CanBeCompiled()
 	{
+		return false;
+	}
+
+	private static List<UnrealTargetPlatform> DesktopPlatforms = new List<UnrealTargetPlatform>() { 
+		UnrealTargetPlatform.Win32,
+		UnrealTargetPlatform.Win64,
+		UnrealTargetPlatform.Mac,
+		UnrealTargetPlatform.Linux,
+	};
 
-		PrintRunTime();
+	public override void Package(ProjectParams Params, DeploymentContext SC, int WorkingCL)
+	{
+		SC.bIsCombiningMultiplePlatforms = true;
+		string SavedPlatformDir = SC.PlatformDir;
+		foreach (UnrealTargetPlatform DesktopPlatform in DesktopPlatforms)
+		{
+			Platform SubPlatform = Platform.Platforms[DesktopPlatform];
+			SC.PlatformDir = DesktopPlatform.ToString();
+			SubPlatform.Package(Params, SC, WorkingCL);
+		}
+		SC.PlatformDir = SavedPlatformDir;
+		SC.bIsCombiningMultiplePlatforms = false;
 	}
 
 	public override void GetFilesToArchive(ProjectParams Params, DeploymentContext SC)
 	{
+		SC.bIsCombiningMultiplePlatforms = true;
+		string SavedPlatformDir = SC.PlatformDir;
+		foreach (UnrealTargetPlatform DesktopPlatform in DesktopPlatforms)
+		{
+			Platform SubPlatform = Platform.Platforms[DesktopPlatform];
+			SC.PlatformDir = DesktopPlatform.ToString();
+			SubPlatform.GetFilesToArchive(Params, SC);
+		}
+		SC.PlatformDir = SavedPlatformDir;
+		SC.bIsCombiningMultiplePlatforms = false;
 	}
 
 	public override void GetConnectedDevices(ProjectParams Params, out List<string> Devices)
@@ -33,10 +63,6 @@ public class DesktopPlatform : Platform
 		Devices = new List<string>();
 	}
 
-	public override void Deploy(ProjectParams Params, DeploymentContext SC)
-	{
-	}
-
 	public override ProcessResult RunClient(ERunOptions ClientRunFlags, string ClientApp, string ClientCmdLine, ProjectParams Params)
 	{
 		return null;
@@ -44,6 +70,34 @@ public class DesktopPlatform : Platform
 
 	public override void GetFilesToDeployOrStage(ProjectParams Params, DeploymentContext SC)
 	{
+		SC.bIsCombiningMultiplePlatforms = true;
+		string SavedPlatformDir = SC.PlatformDir;
+		foreach (UnrealTargetPlatform DesktopPlatform in DesktopPlatforms)
+		{
+			Platform SubPlatform = Platform.Platforms[DesktopPlatform];
+			SC.PlatformDir = DesktopPlatform.ToString();
+			SubPlatform.GetFilesToDeployOrStage(Params, SC);
+		}
+		SC.PlatformDir = SavedPlatformDir;
+		SC.bIsCombiningMultiplePlatforms = false;
+	}
+
+	public override void ProcessArchivedProject(ProjectParams Params, DeploymentContext SC)
+	{
+		Console.WriteLine("***************************** PROCESSING ARCHIVED PROJECT ****************");
+
+		SC.bIsCombiningMultiplePlatforms = true;
+		string SavedPlatformDir = SC.PlatformDir;
+		foreach (UnrealTargetPlatform DesktopPlatform in DesktopPlatforms)
+		{
+			Platform SubPlatform = Platform.Platforms[DesktopPlatform];
+			SC.PlatformDir = DesktopPlatform.ToString();
+			SubPlatform.ProcessArchivedProject(Params, SC);
+		}
+		SC.PlatformDir = SavedPlatformDir;
+		SC.bIsCombiningMultiplePlatforms = false;
+
+		Console.WriteLine("***************************** DONE PROCESSING ARCHIVED PROJECT ****************");
 	}
 
 	/// <summary>
diff --git a/Engine/Source/Programs/AutomationTool/Linux/LinuxPlatform.Automation.cs b/Engine/Source/Programs/AutomationTool/Linux/LinuxPlatform.Automation.cs
--- a/Engine/Source/Programs/AutomationTool/Linux/LinuxPlatform.Automation.cs
+++ b/Engine/Source/Programs/AutomationTool/Linux/LinuxPlatform.Automation.cs
@@ -100,7 +100,7 @@ public abstract class BaseLinuxPlatform : Platform
 				    if (!Params.IsCodeBasedProject && Exe == Exes[0])
 				    {
 						// ensure the ue4game binary exists, if applicable
-						if (!SC.IsCodeBasedProject && !FileExists_NoExceptions(Params.ProjectGameExeFilename))
+						if (!SC.IsCodeBasedProject && !FileExists_NoExceptions(Params.ProjectGameExeFilename) && !SC.bIsCombiningMultiplePlatforms)
 						{
 							Log("Failed to find game binary " + Params.ProjectGameExeFilename);
 							AutomationTool.ErrorReporter.Error("Stage Failed.", (int)AutomationTool.ErrorCodes.Error_MissingExecutable);
diff --git a/Engine/Source/Programs/AutomationTool/Mac/MacPlatform.Automation.cs b/Engine/Source/Programs/AutomationTool/Mac/MacPlatform.Automation.cs
--- a/Engine/Source/Programs/AutomationTool/Mac/MacPlatform.Automation.cs
+++ b/Engine/Source/Programs/AutomationTool/Mac/MacPlatform.Automation.cs
@@ -112,7 +112,7 @@ public class MacPlatform : Platform
 
 					string AbsoluteBundlePath = CombinePaths (SC.LocalRoot, AppBundlePath);
 					// ensure the ue4game binary exists, if applicable
-					if (!SC.IsCodeBasedProject && !Directory.Exists(AbsoluteBundlePath))
+					if (!SC.IsCodeBasedProject && !Directory.Exists(AbsoluteBundlePath) && !SC.bIsCombiningMultiplePlatforms)
 					{
 						Log("Failed to find app bundle " + AbsoluteBundlePath);
 						AutomationTool.ErrorReporter.Error("Stage Failed.", (int)AutomationTool.ErrorCodes.Error_MissingExecutable);
@@ -158,54 +158,69 @@ public class MacPlatform : Platform
 	{
 		string ExeName = SC.StageExecutables[0];
 		string BundlePath = CombinePaths(SC.ArchiveDirectory, SC.ShortProjectName + ".app");
-		if (!Directory.Exists(BundlePath))
+
+		if (SC.bIsCombiningMultiplePlatforms)
 		{
-			string SourceBundlePath = CombinePaths(SC.ArchiveDirectory, SC.ShortProjectName, "Binaries", "Mac", ExeName + ".app");
-			if (!Directory.Exists(SourceBundlePath))
+			// when combining multiple platforms, don't merge the content into the .app, use the one in the Binaries directory
+			BundlePath = CombinePaths(SC.ArchiveDirectory, SC.ShortProjectName, "Binaries", "Mac", ExeName + ".app");
+			if (!Directory.Exists(BundlePath))
 			{
-				SourceBundlePath = CombinePaths(SC.ArchiveDirectory, "Engine", "Binaries", "Mac", ExeName + ".app");
+				// if the .app wasn't there, just skip out (we don't require executables when combining)
+				return;
 			}
-			Directory.Move(SourceBundlePath, BundlePath);
 		}
 
 		string TargetPath = CombinePaths(BundlePath, "Contents", "UE4");
-		if (DirectoryExists(TargetPath))
+		if (!SC.bIsCombiningMultiplePlatforms)
 		{
-			Directory.Delete(TargetPath, true);
-		}
+			if (!Directory.Exists(BundlePath))
+			{
+				string SourceBundlePath = CombinePaths(SC.ArchiveDirectory, SC.ShortProjectName, "Binaries", "Mac", ExeName + ".app");
+				if (!Directory.Exists(SourceBundlePath))
+				{
+					SourceBundlePath = CombinePaths(SC.ArchiveDirectory, "Engine", "Binaries", "Mac", ExeName + ".app");
+				}
+				Directory.Move(SourceBundlePath, BundlePath);
+			}
 
-		// First, move all files and folders inside he app bundle
-		string[] StagedFiles = Directory.GetFiles(SC.ArchiveDirectory, "*", SearchOption.TopDirectoryOnly);
-		foreach (string FilePath in StagedFiles)
-		{
-			string TargetFilePath = CombinePaths(TargetPath, Path.GetFileName(FilePath));
-			Directory.CreateDirectory(Path.GetDirectoryName(TargetFilePath));
-			File.Move(FilePath, TargetFilePath);
-		}
+			if (DirectoryExists(TargetPath))
+			{
+				Directory.Delete(TargetPath, true);
+			}
 
-		string[] StagedDirectories = Directory.GetDirectories(SC.ArchiveDirectory, "*", SearchOption.TopDirectoryOnly);
-		foreach (string DirPath in StagedDirectories)
-		{
-			string DirName = Path.GetFileName(DirPath);
-			if (!DirName.EndsWith(".app"))
+			// First, move all files and folders inside he app bundle
+			string[] StagedFiles = Directory.GetFiles(SC.ArchiveDirectory, "*", SearchOption.TopDirectoryOnly);
+			foreach (string FilePath in StagedFiles)
+			{
+				string TargetFilePath = CombinePaths(TargetPath, Path.GetFileName(FilePath));
+				Directory.CreateDirectory(Path.GetDirectoryName(TargetFilePath));
+				File.Move(FilePath, TargetFilePath);
+			}
+
+			string[] StagedDirectories = Directory.GetDirectories(SC.ArchiveDirectory, "*", SearchOption.TopDirectoryOnly);
+			foreach (string DirPath in StagedDirectories)
 			{
-				string TargetDirPath = CombinePaths(TargetPath, DirName);
-				Directory.CreateDirectory(Path.GetDirectoryName(TargetDirPath));
-				Directory.Move(DirPath, TargetDirPath);
+				string DirName = Path.GetFileName(DirPath);
+				if (!DirName.EndsWith(".app"))
+				{
+					string TargetDirPath = CombinePaths(TargetPath, DirName);
+					Directory.CreateDirectory(Path.GetDirectoryName(TargetDirPath));
+					Directory.Move(DirPath, TargetDirPath);
+				}
 			}
 		}
 
 		// Update executable name, icon and entry in Info.plist
-		string UE4GamePath = CombinePaths(SC.ArchiveDirectory, SC.ShortProjectName + ".app", "Contents", "MacOS", ExeName);
+		string UE4GamePath = CombinePaths(BundlePath, "Contents", "MacOS", ExeName);
 		if (ExeName != SC.ShortProjectName && File.Exists(UE4GamePath))
 		{
-			string GameExePath = CombinePaths(SC.ArchiveDirectory, SC.ShortProjectName + ".app", "Contents", "MacOS", SC.ShortProjectName);
+			string GameExePath = CombinePaths(BundlePath, "Contents", "MacOS", SC.ShortProjectName);
 			File.Delete(GameExePath);
 			File.Move(UE4GamePath, GameExePath);
 
-			string DefaultIconPath = CombinePaths(SC.ArchiveDirectory, SC.ShortProjectName + ".app", "Contents", "Resources", "UE4.icns");
-			string CustomIconSrcPath = CombinePaths(SC.ArchiveDirectory, SC.ShortProjectName + ".app", "Contents", "Resources", "Application.icns");
-			string CustomIconDestPath = CombinePaths(SC.ArchiveDirectory, SC.ShortProjectName + ".app", "Contents", "Resources", SC.ShortProjectName + ".icns");
+			string DefaultIconPath = CombinePaths(BundlePath, "Contents", "Resources", "UE4.icns");
+			string CustomIconSrcPath = CombinePaths(BundlePath, "Contents", "Resources", "Application.icns");
+			string CustomIconDestPath = CombinePaths(BundlePath, "Contents", "Resources", SC.ShortProjectName + ".icns");
 			if (File.Exists(CustomIconSrcPath))
 			{
 				File.Delete(DefaultIconPath);
@@ -216,7 +231,7 @@ public class MacPlatform : Platform
 				File.Move(DefaultIconPath, CustomIconDestPath);
 			}
 
-			string InfoPlistPath = CombinePaths(SC.ArchiveDirectory, SC.ShortProjectName + ".app", "Contents", "Info.plist");
+			string InfoPlistPath = CombinePaths(BundlePath, "Contents", "Info.plist");
 			string InfoPlistContents = File.ReadAllText(InfoPlistPath);
 			InfoPlistContents = InfoPlistContents.Replace(ExeName, SC.ShortProjectName);
 			InfoPlistContents = InfoPlistContents.Replace("<string>UE4</string>", "<string>" + SC.ShortProjectName + "</string>");
diff --git a/Engine/Source/Programs/AutomationTool/Platform.cs b/Engine/Source/Programs/AutomationTool/Platform.cs
--- a/Engine/Source/Programs/AutomationTool/Platform.cs
+++ b/Engine/Source/Programs/AutomationTool/Platform.cs
@@ -492,6 +492,15 @@ namespace AutomationTool
             return false;
         }
 
+		/// <summary>
+		/// Allows some platforms to not be compiled, for instance when BuildCookRun -build is performed
+		/// </summary>
+		/// <returns><c>true</c> if this instance can be compiled; otherwise, <c>false</c>.</returns>
+		public virtual bool CanBeCompiled()
+		{
+			return true;
+		}
+
 		public virtual bool RetrieveDeployedManifests(ProjectParams Params, DeploymentContext SC, out List<string> UFSManifests, out List<string> NonUFSManifests)
 		{
 			UFSManifests = null;
diff --git a/Engine/Source/Programs/AutomationTool/UE4Build.cs b/Engine/Source/Programs/AutomationTool/UE4Build.cs
--- a/Engine/Source/Programs/AutomationTool/UE4Build.cs
+++ b/Engine/Source/Programs/AutomationTool/UE4Build.cs
@@ -666,6 +666,12 @@ namespace AutomationTool
 			/// <param name="InAddArgs">Specifies additional arguments for UBT</param>
 			public void AddTarget(string TargetName, UnrealBuildTool.UnrealTargetPlatform InPlatform, UnrealBuildTool.UnrealTargetConfiguration InConfiguration, string InUprojectPath = null, bool bForceMonolithic = false, bool bForceNonUnity = false, bool bForceDebugInfo = false, string InAddArgs = "")
 			{
+				// Is this platform a compilable target?
+				if (!Platform.Platforms[InPlatform].CanBeCompiled())
+				{
+					return;
+				}
+
 				Targets.Add(new BuildTarget()
 				{
 					TargetName = TargetName,
@@ -691,6 +697,12 @@ namespace AutomationTool
 			/// <param name="bForceDebugInfo">Force debug info even in development builds</param>
 			public void AddTargets(string[] TargetNames, UnrealBuildTool.UnrealTargetPlatform InPlatform, UnrealBuildTool.UnrealTargetConfiguration InConfiguration, string InUprojectPath = null, bool bForceMonolithic = false, bool bForceNonUnity = false, bool bForceDebugInfo = false, string InAddArgs = "")
 			{
+				// Is this platform a compilable target?
+				if (!Platform.Platforms[InPlatform].CanBeCompiled())
+				{
+					return;
+				}
+
 				foreach (var Target in TargetNames)
 				{
 					Targets.Add(new BuildTarget()
@@ -1335,15 +1347,18 @@ namespace AutomationTool
 				Log("Build products *******");
 				if (BuildProductFiles.Count < 1)
 				{
-					throw new AutomationException("BUILD FAILED: no build products??");
+					Log("No build products were made");
 				}
-				foreach (var Product in BuildProductFiles)
+				else
 				{
-					if (!FileExists(Product))
+					foreach (var Product in BuildProductFiles)
 					{
-						throw new AutomationException("BUILD FAILED {0} was a build product but no longer exists", Product);
+						if (!FileExists(Product))
+						{
+							throw new AutomationException("BUILD FAILED {0} was a build product but no longer exists", Product);
+						}
+						Log(Product);
 					}
-					Log(Product);
 				}
 				Log("End Build products *******");
 			}
diff --git a/Engine/Source/Programs/AutomationTool/Win/WinPlatform.Automation.cs b/Engine/Source/Programs/AutomationTool/Win/WinPlatform.Automation.cs
--- a/Engine/Source/Programs/AutomationTool/Win/WinPlatform.Automation.cs
+++ b/Engine/Source/Programs/AutomationTool/Win/WinPlatform.Automation.cs
@@ -185,7 +185,7 @@ public abstract class BaseWinPlatform : Platform
 					string SourceFile = CombinePaths(SC.LocalRoot, "Engine", "Binaries", SC.PlatformDir, ExeFileName);
 
 					// ensure the ue4game binary exists, if applicable
-					if (!SC.IsCodeBasedProject && !FileExists_NoExceptions(SourceFile))
+					if (!SC.IsCodeBasedProject && !FileExists_NoExceptions(SourceFile) && !SC.bIsCombiningMultiplePlatforms)
 					{
 						Log("Failed to find game executable " + SourceFile);
 						AutomationTool.ErrorReporter.Error("Stage Failed.", (int)AutomationTool.ErrorCodes.Error_MissingExecutable);
@@ -223,28 +223,33 @@ public abstract class BaseWinPlatform : Platform
 			string IntermediateFile = CombinePaths(IntermediateDir, ExeName);
 			File.Copy(InputFile, IntermediateFile, true);
 	
-			// Get the icon from the build directory if possible
-			GroupIconResource GroupIcon = null;
-			if(InternalUtils.SafeFileExists(CombinePaths(SC.ProjectRoot, "Build/Windows/Application.ico")))
+			// currently the icon updating doesn't run under mono
+			if (UnrealBuildTool.BuildHostPlatform.Current.Platform == UnrealTargetPlatform.Win64 ||
+				UnrealBuildTool.BuildHostPlatform.Current.Platform == UnrealTargetPlatform.Win32)
 			{
-				GroupIcon = GroupIconResource.FromIco(CombinePaths(SC.ProjectRoot, "Build/Windows/Application.ico"));
-			}
-			if(GroupIcon == null)
-			{
-				GroupIcon = GroupIconResource.FromExe(TargetFile);
-			}
+				// Get the icon from the build directory if possible
+				GroupIconResource GroupIcon = null;
+				if(InternalUtils.SafeFileExists(CombinePaths(SC.ProjectRoot, "Build/Windows/Application.ico")))
+				{
+					GroupIcon = GroupIconResource.FromIco(CombinePaths(SC.ProjectRoot, "Build/Windows/Application.ico"));
+				}
+				if(GroupIcon == null)
+				{
+					GroupIcon = GroupIconResource.FromExe(TargetFile);
+				}
 
-			// Update the resources in the new file
-			using(ModuleResourceUpdate Update = new ModuleResourceUpdate(IntermediateFile, true))
-			{
-				const int IconResourceId = 101;
-				if(GroupIcon != null) Update.SetIcons(IconResourceId, GroupIcon);
+				// Update the resources in the new file
+				using(ModuleResourceUpdate Update = new ModuleResourceUpdate(IntermediateFile, true))
+				{
+					const int IconResourceId = 101;
+					if(GroupIcon != null) Update.SetIcons(IconResourceId, GroupIcon);
 
-				const int ExecFileResourceId = 201;
-				Update.SetData(ExecFileResourceId, ResourceType.RawData, Encoding.Unicode.GetBytes(StagedRelativeTargetPath + "\0"));
+					const int ExecFileResourceId = 201;
+					Update.SetData(ExecFileResourceId, ResourceType.RawData, Encoding.Unicode.GetBytes(StagedRelativeTargetPath + "\0"));
 
-				const int ExecArgsResourceId = 202;
-				Update.SetData(ExecArgsResourceId, ResourceType.RawData, Encoding.Unicode.GetBytes(StagedArguments + "\0"));
+					const int ExecArgsResourceId = 202;
+					Update.SetData(ExecArgsResourceId, ResourceType.RawData, Encoding.Unicode.GetBytes(StagedArguments + "\0"));
+				}
 			}
 
 			// Copy it to the staging directory
