- Added the automation scripts to the AutomationTool_Mono solution
- Fixed the Desktop.Automation.csproj GUID to not overlap Android

diff --git a/Engine/Source/Programs/AutomationTool/Desktop/Desktop.Automation.csproj b/Engine/Source/Programs/AutomationTool/Desktop/Desktop.Automation.csproj
--- a/Engine/Source/Programs/AutomationTool/Desktop/Desktop.Automation.csproj
+++ b/Engine/Source/Programs/AutomationTool/Desktop/Desktop.Automation.csproj
@@ -1,10 +1,10 @@
-﻿<?xml version="1.0" encoding="utf-8"?>
+<?xml version="1.0" encoding="utf-8"?>
 <Project ToolsVersion="4.0" DefaultTargets="Build" xmlns="http://schemas.microsoft.com/developer/msbuild/2003">
   <Import Project="$(MSBuildExtensionsPath)\$(MSBuildToolsVersion)\Microsoft.Common.props" Condition="Exists('$(MSBuildExtensionsPath)\$(MSBuildToolsVersion)\Microsoft.Common.props')" />
   <PropertyGroup>
     <Configuration Condition=" '$(Configuration)' == '' ">Debug</Configuration>
     <Platform Condition=" '$(Platform)' == '' ">AnyCPU</Platform>
-    <ProjectGuid>{A9CD1A06-764A-4C53-86CE-2F34B94F0283}</ProjectGuid>
+    <ProjectGuid>{A9CD1A06-764A-86CE-4C53-2F34B94F0283}</ProjectGuid>
     <OutputType>Library</OutputType>
     <AppDesignerFolder>Properties</AppDesignerFolder>
     <RootNamespace>Desktop.Automation</RootNamespace>
