Fix path slash error with UAT & -manifests only creating chunk0 on Mac (and probably Linux)

diff --git a/Engine/Source/Programs/AutomationTool/Scripts/CopyBuildToStagingDirectory.Automation.cs b/Engine/Source/Programs/AutomationTool/Scripts/CopyBuildToStagingDirectory.Automation.cs
--- a/Engine/Source/Programs/AutomationTool/Scripts/CopyBuildToStagingDirectory.Automation.cs
+++ b/Engine/Source/Programs/AutomationTool/Scripts/CopyBuildToStagingDirectory.Automation.cs
@@ -831,7 +831,7 @@ public partial class Project : CommandUtils
 			bool bAddedToChunk = false;
 			for (int ChunkIndex = 0; !bAddedToChunk && ChunkIndex < ChunkResponseFiles.Length; ++ChunkIndex)
 			{
-				if (ChunkResponseFiles[ChunkIndex].Contains(StagingFile.Key))
+				if (ChunkResponseFiles[ChunkIndex].Contains(StagingFile.Key) || ChunkResponseFiles[ChunkIndex].Contains(StagingFile.Key.Replace('/','\\')))
 				{
 					PakResponseFiles[ChunkIndex].Add(StagingFile.Key, StagingFile.Value);
 					bAddedToChunk = true;
