Use build receipts to figure out what needs to be staged for Windows.

diff --git a/Engine/Source/Programs/AutomationTool/DeploymentContext.cs b/Engine/Source/Programs/AutomationTool/DeploymentContext.cs
--- a/Engine/Source/Programs/AutomationTool/DeploymentContext.cs
+++ b/Engine/Source/Programs/AutomationTool/DeploymentContext.cs
@@ -51,6 +51,11 @@ public class DeploymentContext //: ProjectParams
 	/// </summary>
 	public List<UnrealTargetConfiguration> StageTargetConfigurations;
 
+	/// <summary>
+	/// Receipts for the build targets that should be staged
+	/// </summary>
+	public List<BuildReceipt> StageTargetReceipts;
+
 	/// <summary>
 	///  this is the root directory that contains the engine: d:\a\UE4\
 	/// </summary>
@@ -214,6 +219,7 @@ public class DeploymentContext //: ProjectParams
 		Platform InSourcePlatform,
         Platform InTargetPlatform,
 		List<UnrealTargetConfiguration> InTargetConfigurations,
+		IEnumerable<BuildReceipt> InStageTargetReceipts,
 		List<String> InStageExecutables,
 		bool InServer,
 		bool InCooked,
@@ -232,6 +238,7 @@ public class DeploymentContext //: ProjectParams
         CookSourcePlatform = InSourcePlatform;
 		StageTargetPlatform = InTargetPlatform;
 		StageTargetConfigurations = new List<UnrealTargetConfiguration>(InTargetConfigurations);
+		StageTargetReceipts = new List<BuildReceipt>(InStageTargetReceipts);
 		StageExecutables = InStageExecutables;
         IsCodeBasedProject = ProjectUtils.IsCodeBasedUProjectFile(RawProjectPath);
 		ShortProjectName = ProjectUtils.GetShortProjectName(RawProjectPath);
@@ -321,6 +328,67 @@ public class DeploymentContext //: ProjectParams
 		ProjectArgForCommandLines = ProjectArgForCommandLines.Replace("\\", "/");
 	}
 
+	public void StageFile(StagedFileType FileType, string InputPath, string OutputPath = null, bool bRemap = true)
+	{
+		InputPath = InputPath.Replace(Path.AltDirectorySeparatorChar, Path.DirectorySeparatorChar);
+
+		if(OutputPath == null)
+		{
+			if(InputPath.StartsWith(ProjectRoot, StringComparison.InvariantCultureIgnoreCase))
+			{
+				OutputPath = CommandUtils.CombinePaths(RelativeProjectRootForStage, InputPath.Substring(ProjectRoot.Length).TrimStart('/', '\\'));
+			}
+			else if (InputPath.StartsWith(LocalRoot, StringComparison.InvariantCultureIgnoreCase))
+			{
+				OutputPath = CommandUtils.CombinePaths(InputPath.Substring(LocalRoot.Length).TrimStart('/', '\\'));
+			}
+			else
+			{
+				throw new AutomationException("Can't deploy {0} because it doesn't start with {1} or {2}", InputPath, ProjectRoot, LocalRoot);
+			}
+		}
+
+		if(bRemap)
+		{
+			OutputPath = StageTargetPlatform.Remap(OutputPath);
+		}
+
+		if (FileType == StagedFileType.UFS)
+		{
+			AddUniqueStagingFile(UFSStagingFiles, InputPath, OutputPath);
+		}
+		else if (FileType == StagedFileType.NonUFS)
+		{
+			AddUniqueStagingFile(NonUFSStagingFiles, InputPath, OutputPath);
+		}
+		else if (FileType == StagedFileType.DebugNonUFS)
+		{
+			AddUniqueStagingFile(NonUFSStagingFilesDebug, InputPath, OutputPath);
+		}
+	}
+
+	public void StageFilesInReceipt(BuildReceipt Receipt)
+	{
+		// Stage all the build products needed at runtime
+		foreach(BuildProduct BuildProduct in Receipt.BuildProducts)
+		{
+			if(BuildProduct.Type == BuildProductType.Executable || BuildProduct.Type == BuildProductType.DynamicLibrary || BuildProduct.Type == BuildProductType.RequiredResource)
+			{
+				StageFile(StagedFileType.NonUFS, BuildProduct.Path);
+			}
+			else if(BuildProduct.Type == BuildProductType.SymbolFile)
+			{
+				StageFile(StagedFileType.DebugNonUFS, BuildProduct.Path);
+			}
+		}
+
+		// Also stage any additional runtime dependencies, like ThirdParty DLLs
+		foreach(RuntimeDependency RuntimeDependency in Receipt.RuntimeDependencies)
+		{
+			StageFile(StagedFileType.NonUFS, RuntimeDependency.Path, RuntimeDependency.StagePath);
+		}
+	}
+
     public int StageFiles(StagedFileType FileType, string InPath, string Wildcard = "*", bool bRecursive = true, string[] ExcludeWildcard = null, string NewPath = null, bool bAllowNone = false, bool bRemap = true, string NewName = null, bool bAllowNotForLicenseesFiles = true, bool bStripFilesForOtherPlatforms = true)
 	{
 		int FilesAdded = 0;
diff --git a/Engine/Source/Programs/AutomationTool/ProjectParams.cs b/Engine/Source/Programs/AutomationTool/ProjectParams.cs
--- a/Engine/Source/Programs/AutomationTool/ProjectParams.cs
+++ b/Engine/Source/Programs/AutomationTool/ProjectParams.cs
@@ -290,6 +290,7 @@ namespace AutomationTool
 			this.NoBootstrapExe = InParams.NoBootstrapExe;
             this.Prebuilt = InParams.Prebuilt;
             this.RunTimeoutSeconds = InParams.RunTimeoutSeconds;
+			this.bIsCodeBasedProject = InParams.bIsCodeBasedProject;
 		}
 
 		/// <summary>
diff --git a/Engine/Source/Programs/AutomationTool/Scripts/CopyBuildToStagingDirectory.Automation.cs b/Engine/Source/Programs/AutomationTool/Scripts/CopyBuildToStagingDirectory.Automation.cs
--- a/Engine/Source/Programs/AutomationTool/Scripts/CopyBuildToStagingDirectory.Automation.cs
+++ b/Engine/Source/Programs/AutomationTool/Scripts/CopyBuildToStagingDirectory.Automation.cs
@@ -1360,6 +1360,29 @@ public partial class Project : CommandUtils
                 }
             }
 
+			string EngineDir = CommandUtils.CombinePaths(CommandUtils.CmdEnv.LocalRoot, "Engine");
+
+			List<BuildReceipt> TargetsToStage = new List<BuildReceipt>();
+			foreach(string Target in ListToProcess)
+			{
+				foreach(UnrealTargetConfiguration Config in ConfigsToProcess)
+				{
+					string ReceiptBaseDir = Params.IsCodeBasedProject? Path.GetDirectoryName(Params.RawProjectPath) : EngineDir;
+
+					// Read the receipt
+					string ReceiptFileName = BuildReceipt.GetDefaultPath(ReceiptBaseDir, Target, StagePlatform, Config, "");
+					if(!File.Exists(ReceiptFileName))
+					{
+						throw new AutomationException("Missing receipt '{0}'. Check that this target has been built.", Path.GetFileName(ReceiptFileName));
+					}
+
+					// Convert the paths to absolute
+					BuildReceipt Receipt = BuildReceipt.Read(ReceiptFileName);
+					Receipt.ExpandPathVariables(EngineDir, Path.GetDirectoryName(Params.RawProjectPath));
+					TargetsToStage.Add(Receipt);
+				}
+			}
+
 			//@todo should pull StageExecutables from somewhere else if not cooked
             var SC = new DeploymentContext(Params.RawProjectPath, CmdEnv.LocalRoot,
                 StageDirectory,
@@ -1368,6 +1391,7 @@ public partial class Project : CommandUtils
 				Params.GetTargetPlatformInstance(CookedDataPlatform),
 				Params.GetTargetPlatformInstance(StagePlatform),
 				ConfigsToProcess,
+				TargetsToStage,
 				ExecutablesToStage,
 				InDedicatedServer,
 				Params.Cook || Params.CookOnTheFly,
diff --git a/Engine/Source/Programs/AutomationTool/Win/WinPlatform.Automation.cs b/Engine/Source/Programs/AutomationTool/Win/WinPlatform.Automation.cs
--- a/Engine/Source/Programs/AutomationTool/Win/WinPlatform.Automation.cs
+++ b/Engine/Source/Programs/AutomationTool/Win/WinPlatform.Automation.cs
@@ -35,188 +35,42 @@ public abstract class BaseWinPlatform : Platform
 			StageExecutable("exe", SC, CommandUtils.CombinePaths(SC.LocalRoot, "Engine/Binaries", SC.PlatformDir), "CrashReportClient.");
 		}
 
-		//todo we need to support shipping and test executables
-		//todo this should all be partially based on UBT manifests and not hard coded
-		//monolithic assumption
-		StageExecutable("dll", SC, CommandUtils.CombinePaths(SC.LocalRoot, "Engine/Binaries/ThirdParty/Ogg", SC.PlatformDir), "*.", true);
-		StageExecutable("dll", SC, CommandUtils.CombinePaths(SC.LocalRoot, "Engine/Binaries/ThirdParty/Vorbis", SC.PlatformDir), "*.", true);
-		string PhysXVer = "VS" + WindowsPlatform.GetVisualStudioCompilerVersionName();
-		string ApexVer = "VS" + WindowsPlatform.GetVisualStudioCompilerVersionName();
-		string PhysXMaskForDebugConfiguration = Params.bDebugBuildsActuallyUseDebugCRT ? "*DEBUG*.*" : "*PROFILE*.*";
-		if (SC.StageTargetConfigurations.Contains(UnrealTargetConfiguration.Debug))
+		// Loop through all the things to stage
+		foreach(BuildReceipt Receipt in SC.StageTargetReceipts)
 		{
-			StageExecutable("dll", SC, CommandUtils.CombinePaths(SC.LocalRoot, "Engine/Binaries/ThirdParty/PhysX/APEX-1.3", SC.PlatformDir, ApexVer), PhysXMaskForDebugConfiguration, true);
-			StageExecutable("dll", SC, CommandUtils.CombinePaths(SC.LocalRoot, "Engine/Binaries/ThirdParty/PhysX/PhysX-3.3", SC.PlatformDir, PhysXVer), PhysXMaskForDebugConfiguration, true);
-			StageExecutable("dll", SC, CommandUtils.CombinePaths(SC.LocalRoot, "Engine/Binaries/ThirdParty/PhysX/PhysX-3.3", SC.PlatformDir, PhysXVer), "nvToolsExt*.", true);
+			SC.StageFilesInReceipt(Receipt);
 		}
-		if (SC.StageTargetConfigurations.Any(x => x != UnrealTargetConfiguration.Debug))
-		{
-			StageExecutable("dll", SC, CommandUtils.CombinePaths(SC.LocalRoot, "Engine/Binaries/ThirdParty/PhysX/APEX-1.3", SC.PlatformDir, ApexVer), "*.", true, new string[] { "*DEBUG*.*", "*CHECKED*.*" });
-			StageExecutable("dll", SC, CommandUtils.CombinePaths(SC.LocalRoot, "Engine/Binaries/ThirdParty/PhysX/PhysX-3.3", SC.PlatformDir, PhysXVer), "*.", true, new string[] { "*DEBUG*.*", "*CHECKED*.*" });
-		}
-
-		if (Params.bUsesSteam)
-		{
-			string SteamVersion = "Steamv132";
 
-			// Check that the TPS directory exists. We don't distribute binaries for Steam in Rocket.
-			if (Directory.Exists(CommandUtils.CombinePaths(SC.LocalRoot, "Engine/Binaries/ThirdParty/Steamworks/" + SteamVersion)))
-			{
-				if (SC.StageTargetPlatform.PlatformType == UnrealTargetPlatform.Win32)
-				{
-					StageExecutable("dll", SC, CommandUtils.CombinePaths(SC.LocalRoot, "Engine/Binaries/ThirdParty/Steamworks/" + SteamVersion, SC.PlatformDir), "steam_api.");
-					if (SC.DedicatedServer)
-					{
-						StageExecutable("dll", SC, CommandUtils.CombinePaths(SC.LocalRoot, "Engine/Binaries/ThirdParty/Steamworks/" + SteamVersion, SC.PlatformDir), "steamclient.");
-						StageExecutable("dll", SC, CommandUtils.CombinePaths(SC.LocalRoot, "Engine/Binaries/ThirdParty/Steamworks/" + SteamVersion, SC.PlatformDir), "tier0_s.");
-						StageExecutable("dll", SC, CommandUtils.CombinePaths(SC.LocalRoot, "Engine/Binaries/ThirdParty/Steamworks/" + SteamVersion, SC.PlatformDir), "vstdlib_s.");
-					}
-				}
-				else
-				{
-					StageExecutable("dll", SC, CommandUtils.CombinePaths(SC.LocalRoot, "Engine/Binaries/ThirdParty/Steamworks/" + SteamVersion, SC.PlatformDir), "steam_api64.");
-					if (SC.DedicatedServer)
-					{
-						StageExecutable("dll", SC, CommandUtils.CombinePaths(SC.LocalRoot, "Engine/Binaries/ThirdParty/Steamworks/" + SteamVersion, SC.PlatformDir), "steamclient64.");
-						StageExecutable("dll", SC, CommandUtils.CombinePaths(SC.LocalRoot, "Engine/Binaries/ThirdParty/Steamworks/" + SteamVersion, SC.PlatformDir), "tier0_s64.");
-						StageExecutable("dll", SC, CommandUtils.CombinePaths(SC.LocalRoot, "Engine/Binaries/ThirdParty/Steamworks/" + SteamVersion, SC.PlatformDir), "vstdlib_s64.");
-					}
-				}
-			}
-		}
-
-        // Leap files
-        StageExecutable("dll", SC, CommandUtils.CombinePaths(SC.LocalRoot, "Engine/Binaries/ThirdParty/Leap", SC.PlatformDir), "Leap.", false, null, null, true);
-        
 		// Copy the splash screen, windows specific
 		SC.StageFiles(StagedFileType.NonUFS, CombinePaths(SC.ProjectRoot, "Content/Splash"), "Splash.bmp", false, null, null, true);
 
-		// CEF3 files
-		if (Params.bUsesCEF3)
+		// Stage the bootstrap executable
+		if(!Params.NoBootstrapExe)
 		{
-			SC.StageFiles(StagedFileType.NonUFS, CombinePaths(SC.LocalRoot, "Engine/Binaries/ThirdParty/CEF3", SC.PlatformDir), "*", true, null, null, true);
-			StageExecutable("exe", SC, CommandUtils.CombinePaths(SC.LocalRoot, "Engine/Binaries", SC.PlatformDir), "UnrealCEFSubProcess.", false, null, null, true);
-		}
-
-        if (Params.StageNonMonolithic)
-        {
-            StageExecutable("dll", SC, CommandUtils.CombinePaths(SC.LocalRoot, "Engine/Binaries/ThirdParty/ICU/icu4c-53_1", SC.PlatformDir, "VS2013"), "*.");
-
-            if (SC.DedicatedServer)
-            {
-                if (SC.StageTargetConfigurations.Contains(UnrealTargetConfiguration.Development))
-                {
-                    StageExecutable("exe", SC, CommandUtils.CombinePaths(SC.LocalRoot, "Engine/Binaries", SC.PlatformDir), "UE4Server.");
-                }
-
-                if (SC.StageTargetConfigurations.Contains(UnrealTargetConfiguration.Test))
-                {
-                    StageExecutable("exe", SC, CommandUtils.CombinePaths(SC.LocalRoot, "Engine/Binaries", SC.PlatformDir), "UE4Server*-Test.");
-                }
-
-                if (SC.StageTargetConfigurations.Contains(UnrealTargetConfiguration.Shipping))
-                {
-                    StageExecutable("exe", SC, CommandUtils.CombinePaths(SC.LocalRoot, "Engine/Binaries", SC.PlatformDir), "UE4Server*-Shipping.");
-                }
-
-                StageExecutable("dll", SC, CommandUtils.CombinePaths(SC.LocalRoot, "Engine/Binaries", SC.PlatformDir), "UE4Server-*.");
-                StageExecutable("dll", SC, CommandUtils.CombinePaths(SC.LocalRoot, "Engine/Plugins"), "UE4Server-*.", true);
-                StageExecutable("dll", SC, CommandUtils.CombinePaths(SC.ProjectRoot, "Binaries", SC.PlatformDir), "UE4Server-*.");
-                StageExecutable("dll", SC, CommandUtils.CombinePaths(SC.ProjectRoot, "Plugins"), "UE4Server-*.", true);
-            }
-            else
-            {
-                if (SC.StageTargetConfigurations.Contains(UnrealTargetConfiguration.Development))
-                {
-                    StageExecutable("exe", SC, CommandUtils.CombinePaths(SC.LocalRoot, "Engine/Binaries", SC.PlatformDir), "UE4.");
-                }
-
-                if (SC.StageTargetConfigurations.Contains(UnrealTargetConfiguration.Test))
-                {
-                    StageExecutable("exe", SC, CommandUtils.CombinePaths(SC.LocalRoot, "Engine/Binaries", SC.PlatformDir), "UE4*-Test.");
-                }
-
-                if (SC.StageTargetConfigurations.Contains(UnrealTargetConfiguration.Shipping))
-                {
-                    StageExecutable("exe", SC, CommandUtils.CombinePaths(SC.LocalRoot, "Engine/Binaries", SC.PlatformDir), "UE4*-Shipping.");
-                }
-
-                StageExecutable("dll", SC, CommandUtils.CombinePaths(SC.LocalRoot, "Engine/Binaries", SC.PlatformDir), "UE4-*.");
-                StageExecutable("dll", SC, CommandUtils.CombinePaths(SC.LocalRoot, "Engine/Plugins"), "UE4-*.", true);
-                StageExecutable("dll", SC, CommandUtils.CombinePaths(SC.ProjectRoot, "Binaries", SC.PlatformDir), "UE4-*.");
-                StageExecutable("dll", SC, CommandUtils.CombinePaths(SC.ProjectRoot, "Plugins"), "UE4-*.", true);
-            }
-        }
-        else
-        {
-			List<string> Exes = GetExecutableNames(SC);
-
-			// the first exe is the "main" one, the rest are marked as debug files
-			StagedFileType WorkingFileType = StagedFileType.NonUFS;
-				
-		    foreach (var Exe in Exes)
-            {
-				if (Exe.StartsWith(CombinePaths(SC.RuntimeProjectRootDir, "Binaries", SC.PlatformDir)))
-				{
-					string BaseExeFileName = Path.GetFileName(Exe);
-					if (!String.IsNullOrEmpty(Params.OverrideMinimumOS) && (Params.OverrideMinimumOS == "WinXP"))
-					{
-						BaseExeFileName = Path.GetFileNameWithoutExtension(Exe) + "_xp" + Path.GetExtension(Exe);
-					}
-
-					// remap the project root.
-					string SourceFile = CombinePaths(SC.ProjectRoot, "Binaries", SC.PlatformDir, BaseExeFileName);
-					StageExecutable("exe", SC, Path.GetDirectoryName(SourceFile), Path.GetFileNameWithoutExtension(SourceFile) + ".", true, null, CommandUtils.CombinePaths(SC.RelativeProjectRootForStage, "Binaries", SC.PlatformDir), false, WorkingFileType);
-					if(Exe == Exes[0] && !Params.NoBootstrapExe)
-					{
-						StageBootstrapExecutable(SC, SourceFile, CombinePaths(SC.RelativeProjectRootForStage, "Binaries", SC.PlatformDir, BaseExeFileName), "");
-					}
-				}
-				else if (Exe.StartsWith(CombinePaths(SC.RuntimeRootDir, "Engine/Binaries", SC.PlatformDir)))
+			foreach(BuildReceipt Receipt in SC.StageTargetReceipts)
+			{
+				BuildProduct Executable = Receipt.BuildProducts.FirstOrDefault(x => x.Type == BuildProductType.Executable);
+				if(Executable != null)
 				{
-					// Select the appropriate windows executable to stage, xp support or not.
-					string ExeFileName = Path.GetFileName(Exe);
-					if (!String.IsNullOrEmpty(Params.OverrideMinimumOS) && (Params.OverrideMinimumOS == "WinXP"))
-					{
-						ExeFileName = Path.GetFileNameWithoutExtension(ExeFileName) + "_xp" + Path.GetExtension(ExeFileName);
-					}
-
-					// keep it in the engine directory.
-					string SourceFile = CombinePaths(SC.LocalRoot, "Engine", "Binaries", SC.PlatformDir, ExeFileName);
-
-					// ensure the ue4game binary exists, if applicable
-					if (!SC.IsCodeBasedProject && !FileExists_NoExceptions(SourceFile) && !SC.bIsCombiningMultiplePlatforms)
+					string BootstrapArguments = "";
+					if(!SC.IsCodeBasedProject && !ShouldStageCommandLine(Params, SC))
 					{
-						Log("Failed to find game executable " + SourceFile);
-						AutomationTool.ErrorReporter.Error("Stage Failed.", (int)AutomationTool.ErrorCodes.Error_MissingExecutable);
-						throw new AutomationException("Could not find exe {0}. You may need to build the UE4 project with your target configuration and platform.", SourceFile);
+						BootstrapArguments = String.Format("..\\..\\..\\{0}\\{0}.uproject", SC.ShortProjectName);
 					}
 
-					StageExecutable("exe", SC, CombinePaths(SC.LocalRoot, "Engine/Binaries", SC.PlatformDir), Path.GetFileNameWithoutExtension(SourceFile) + ".", true, null, null, false, WorkingFileType);
-					if(Exe == Exes[0] && !Params.NoBootstrapExe)
-					{
-						string CommandLine = ShouldStageCommandLine(Params, SC)? "" : String.Format("..\\..\\..\\{0}\\{0}.uproject", SC.ShortProjectName);
-						StageBootstrapExecutable(SC, SourceFile, CombinePaths("Engine", "Binaries", SC.PlatformDir, ExeFileName), CommandLine);
-					}
+					string BootstrapExeName = (SC.StageTargetConfigurations.Count == 1)? (Receipt.GetProperty("TargetName", SC.ShortProjectName) + ".exe") : Path.GetFileName(Executable.Path);
+					StageBootstrapExecutable(SC, BootstrapExeName, Executable.Path, SC.NonUFSStagingFiles[Executable.Path], BootstrapArguments);
 				}
-				else
-				{
-					throw new AutomationException("Can't stage the exe {0} because it doesn't start with {1} or {2}", Exe, CombinePaths(SC.RuntimeProjectRootDir, "Binaries", SC.PlatformDir), CombinePaths(SC.RuntimeRootDir, "Engine/Binaries", SC.PlatformDir));
-				}
-				// the first exe is the "main" one, the rest are marked as debug files
-				WorkingFileType = StagedFileType.DebugNonUFS;
 			}
 		}
 	}
 
-	void StageBootstrapExecutable(DeploymentContext SC, string TargetFile, string StagedRelativeTargetPath, string StagedArguments)
+	void StageBootstrapExecutable(DeploymentContext SC, string ExeName, string TargetFile, string StagedRelativeTargetPath, string StagedArguments)
 	{
 		string InputFile = CombinePaths(SC.LocalRoot, "Engine", "Binaries", SC.PlatformDir, String.Format("BootstrapPackagedGame-{0}-Shipping.exe", SC.PlatformDir));
 		if(InternalUtils.SafeFileExists(InputFile))
 		{
 			// Create the new bootstrap program
-			string ExeName = String.Format("{0}.exe", SC.ShortProjectName);
-
 			string IntermediateDir = CombinePaths(SC.ProjectRoot, "Intermediate", "Staging");
 			InternalUtils.SafeCreateDirectory(IntermediateDir);
 
diff --git a/Engine/Source/Programs/UnrealBuildTool/Configuration/UEBuildBinary.cs b/Engine/Source/Programs/UnrealBuildTool/Configuration/UEBuildBinary.cs
--- a/Engine/Source/Programs/UnrealBuildTool/Configuration/UEBuildBinary.cs
+++ b/Engine/Source/Programs/UnrealBuildTool/Configuration/UEBuildBinary.cs
@@ -696,6 +696,8 @@ namespace UnrealBuildTool
 		public override BuildReceipt MakeReceipt(IUEToolChain ToolChain)
 		{
 			BuildReceipt Receipt = base.MakeReceipt(ToolChain);
+
+			// Set the IsPrecompiled flag on all the build products if we're not actually building this binary
 			if(!Config.bAllowCompilation)
 			{
 				foreach(BuildProduct BuildProduct in Receipt.BuildProducts)
@@ -703,6 +705,10 @@ namespace UnrealBuildTool
 					BuildProduct.IsPrecompiled = true;
 				}
 			}
+
+			// Add runtime dependencies for all the modules in this binary, and build up a list of all the referenced modules
+			Dictionary<string, UEBuildModule> ReferencedModules = new Dictionary<string,UEBuildModule>();
+			List<UEBuildModule> OrderedModules = new List<UEBuildModule>();
 			foreach (string ModuleName in ModuleNames)
 			{
 				UEBuildModule Module = Target.GetModuleByName(ModuleName); 
@@ -710,6 +716,20 @@ namespace UnrealBuildTool
 				{
 					Receipt.RuntimeDependencies.Add(new RuntimeDependency(RuntimeDependency));
 				}
+				Module.GetAllDependencyModules(ReferencedModules, OrderedModules, true, false, true);
+			}
+
+			// Add runtime dependencies for all the referenced external modules. These may be introduce dependencies for the binary without actually being listed for inclusion in it.
+			foreach(UEBuildModule OrderedModule in OrderedModules)
+			{
+				UEBuildExternalModule ExternalModule = OrderedModule as UEBuildExternalModule;
+				if(ExternalModule != null)
+				{
+					foreach(RuntimeDependency RuntimeDependency in ExternalModule.RuntimeDependencies)
+					{
+						Receipt.RuntimeDependencies.Add(new RuntimeDependency(RuntimeDependency));
+					}
+				}
 			}
 			return Receipt;
 		}
diff --git a/Engine/Source/Programs/UnrealBuildTool/Configuration/UEBuildTarget.cs b/Engine/Source/Programs/UnrealBuildTool/Configuration/UEBuildTarget.cs
--- a/Engine/Source/Programs/UnrealBuildTool/Configuration/UEBuildTarget.cs
+++ b/Engine/Source/Programs/UnrealBuildTool/Configuration/UEBuildTarget.cs
@@ -1347,11 +1347,21 @@ namespace UnrealBuildTool
 		private void PrepareReceipt(IUEToolChain ToolChain)
 		{
 			Receipt = new BuildReceipt();
+
+			// Add the target properties
+			Receipt.SetProperty("TargetName", TargetName);
+			Receipt.SetProperty("Platform", Platform.ToString());
+			Receipt.SetProperty("Configuration", Configuration.ToString());
+			Receipt.SetProperty("Precompile", bPrecompile.ToString());
+
+			// Merge all the binary receipts into this
 			foreach(UEBuildBinary Binary in AppBinaries)
 			{
 				BuildReceipt BinaryReceipt = Binary.MakeReceipt(ToolChain);
-				Receipt.Append(BinaryReceipt);
+				Receipt.Merge(BinaryReceipt);
 			}
+
+			// Convert all the paths to use variables for the engine and project root
 			Receipt.InsertStandardPathVariables(BuildConfiguration.RelativeEnginePath, ProjectDirectory);
 		}
 
diff --git a/Engine/Source/Programs/UnrealBuildTool/System/BuildReceipt.cs b/Engine/Source/Programs/UnrealBuildTool/System/BuildReceipt.cs
--- a/Engine/Source/Programs/UnrealBuildTool/System/BuildReceipt.cs
+++ b/Engine/Source/Programs/UnrealBuildTool/System/BuildReceipt.cs
@@ -10,6 +10,26 @@ using System.ComponentModel;
 
 namespace UnrealBuildTool
 {
+	[Serializable]
+	public class BuildProperty
+	{
+		[XmlAttribute]
+		public string Name;
+
+		[XmlAttribute]
+		public string Value;
+
+		private BuildProperty()
+		{
+		}
+
+		public BuildProperty(string InName, string InValue)
+		{
+			Name = InName;
+			Value = InValue;
+		}
+	}
+
 	public enum BuildProductType
 	{
 		Executable,
@@ -97,6 +117,9 @@ namespace UnrealBuildTool
 	[Serializable]
 	public class BuildReceipt
 	{
+		[XmlArrayItem("Property")]
+		public List<BuildProperty> Properties = new List<BuildProperty>();
+
 		[XmlArrayItem("BuildProduct")]
 		public List<BuildProduct> BuildProducts = new List<BuildProduct>();
 
@@ -126,6 +149,42 @@ namespace UnrealBuildTool
 			}
 		}
 
+		/// <summary>
+		/// Sets a property with the given name
+		/// </summary>
+		/// <param name="Name">Name of the property; case sensitive.</param>
+		/// <param name="Value">Value for the property</param>
+		public void SetProperty(string Name, string Value)
+		{
+			BuildProperty Property = Properties.FirstOrDefault(x => x.Name == Name);
+			if(Property == null)
+			{
+				Properties.Add(new BuildProperty(Name, Value));
+			}
+			else
+			{
+				Property.Value = Value;
+			}
+		}
+
+		/// <summary>
+		/// Gets the value associated with a property
+		/// </summary>
+		/// <param name="Name">Name of the property; case sensitive.</param>
+		/// <param name="DefaultValue">Default value for the property if it's not found</param>
+		public string GetProperty(string Name, string DefaultValue)
+		{
+			BuildProperty Property = Properties.FirstOrDefault(x => x.Name == Name);
+			if(Property == null)
+			{
+				return DefaultValue;
+			}
+			else
+			{
+				return Property.Value;
+			}
+		}
+
 		/// <summary>
 		/// Adds a build product to the receipt. Does not check whether it already exists.
 		/// </summary>
@@ -153,10 +212,10 @@ namespace UnrealBuildTool
 		}
 
 		/// <summary>
-		/// Appends another receipt to this one.
+		/// Merges another receipt to this one.
 		/// </summary>
-		/// <param name="Other">Receipt which should be appended</param>
-		public void Append(BuildReceipt Other)
+		/// <param name="Other">Receipt which should be merged</param>
+		public void Merge(BuildReceipt Other)
 		{
 			foreach(BuildProduct OtherBuildProduct in Other.BuildProducts)
 			{
@@ -164,7 +223,10 @@ namespace UnrealBuildTool
 			}
 			foreach(RuntimeDependency OtherRuntimeDependency in Other.RuntimeDependencies)
 			{
-				RuntimeDependencies.Add(OtherRuntimeDependency);
+				if(!RuntimeDependencies.Any(x => x.Path == OtherRuntimeDependency.Path && x.StagePath == OtherRuntimeDependency.StagePath))
+				{
+					RuntimeDependencies.Add(OtherRuntimeDependency);
+				}
 			}
 		}
 
