HTML5: Don't create per map paks when doing "Launch On".

diff --git a/Engine/Source/Programs/AutomationTool/HTML5/HTML5Platform.Automation.cs b/Engine/Source/Programs/AutomationTool/HTML5/HTML5Platform.Automation.cs
--- a/Engine/Source/Programs/AutomationTool/HTML5/HTML5Platform.Automation.cs
+++ b/Engine/Source/Programs/AutomationTool/HTML5/HTML5Platform.Automation.cs
@@ -19,7 +19,7 @@ public class HTML5Platform : Platform
 
 	public override void Package(ProjectParams Params, DeploymentContext SC, int WorkingCL)
 	{
-		Log("Package {0}", Params.RawProjectPath);
+        Log("Package {0}", Params.RawProjectPath);
         
         string PackagePath = Path.Combine(Path.GetDirectoryName(Params.RawProjectPath), "Binaries", "HTML5");
         if (!Directory.Exists(PackagePath))
@@ -28,12 +28,9 @@ public class HTML5Platform : Platform
         }
         string FinalDataLocation = Path.Combine(PackagePath, Params.ShortProjectName) + ".data";
 
+        var ConfigCache = new UnrealBuildTool.ConfigCacheIni(UnrealTargetPlatform.HTML5, "Engine", Path.GetDirectoryName(Params.RawProjectPath), CommandUtils.CombinePaths(CommandUtils.CmdEnv.LocalRoot, "Engine"));
 
-        bool UseAsyncLevelLoading = false;
-        var ConfigCache = new UnrealBuildTool.ConfigCacheIni(UnrealTargetPlatform.HTML5, "Engine", Path.GetDirectoryName(Params.RawProjectPath), CombinePaths(CmdEnv.LocalRoot, "Engine"));
-        ConfigCache.GetBool("/Script/HTML5PlatformEditor.HTML5TargetSettings", "UseAsyncLevelLoading", out UseAsyncLevelLoading);
-
-        if (UseAsyncLevelLoading)
+        if (HTMLPakAutomation.CanCreateMapPaks(Params))
         {
             HTMLPakAutomation PakAutomation = new HTMLPakAutomation(Params, SC);
 
@@ -250,11 +247,7 @@ public class HTML5Platform : Platform
 		SC.ArchiveFiles(PackagePath, Path.GetFileName(OutputFile));
         
 
-        bool UseAsyncLevelLoading = false;
-        var ConfigCache = new UnrealBuildTool.ConfigCacheIni(UnrealTargetPlatform.HTML5, "Engine", Path.GetDirectoryName(Params.RawProjectPath), CombinePaths(CmdEnv.LocalRoot, "Engine"));
-        ConfigCache.GetBool("/Script/HTML5PlatformEditor.HTML5TargetSettings", "UseAsyncLevelLoading", out UseAsyncLevelLoading);
-
-        if (UseAsyncLevelLoading)
+        if (HTMLPakAutomation.CanCreateMapPaks(Params))
         {
             // find all paks.
             string[] Files = Directory.GetFiles(Path.Combine(PackagePath, Params.ShortProjectName), "*",SearchOption.AllDirectories);
@@ -371,11 +364,7 @@ public class HTML5Platform : Platform
 
     public override string GetCookExtraCommandLine(ProjectParams Params)
     {
-        var ConfigCache = new UnrealBuildTool.ConfigCacheIni(UnrealTargetPlatform.HTML5, "Engine", Path.GetDirectoryName(Params.RawProjectPath), CombinePaths(CmdEnv.LocalRoot, "Engine"));
-        bool UseAsyncLevelLoading = false;
-        ConfigCache.GetBool("/Script/HTML5PlatformEditor.HTML5TargetSettings", "UseAsyncLevelLoading", out UseAsyncLevelLoading);
-
-        return UseAsyncLevelLoading ? " -GenerateDependenciesForMaps " : "";
+        return HTMLPakAutomation.CanCreateMapPaks(Params) ? " -GenerateDependenciesForMaps " : ""; 
     }
 
     public override List<string> GetCookExtraMaps()
@@ -392,11 +381,7 @@ public class HTML5Platform : Platform
 
     public override PakType RequiresPak(ProjectParams Params)
     {
-        var ConfigCache = new UnrealBuildTool.ConfigCacheIni(UnrealTargetPlatform.HTML5, "Engine", Path.GetDirectoryName(Params.RawProjectPath), CombinePaths(CmdEnv.LocalRoot, "Engine"));
-        bool UseAsyncLevelLoading = false; 
-        ConfigCache.GetBool("/Script/HTML5PlatformEditor.HTML5TargetSettings", "UseAsyncLevelLoading", out UseAsyncLevelLoading);
-
-        return UseAsyncLevelLoading ? PakType.Never : PakType.Always;
+        return HTMLPakAutomation.CanCreateMapPaks(Params) ? PakType.Never : PakType.Always;
     }
 
 	public override string GetPlatformPakCommandLine()
diff --git a/Engine/Source/Programs/AutomationTool/HTML5/HTML5Platform.PakFiles.Automation.cs b/Engine/Source/Programs/AutomationTool/HTML5/HTML5Platform.PakFiles.Automation.cs
--- a/Engine/Source/Programs/AutomationTool/HTML5/HTML5Platform.PakFiles.Automation.cs
+++ b/Engine/Source/Programs/AutomationTool/HTML5/HTML5Platform.PakFiles.Automation.cs
@@ -19,6 +19,18 @@ public class HTMLPakAutomation
     string PakOrderFileLocation;
 
     Dictionary<string, object> DependencyJson;
+    
+    public static bool CanCreateMapPaks(ProjectParams Param)
+    {
+        bool UseAsyncLevelLoading = false;
+        var ConfigCache = new UnrealBuildTool.ConfigCacheIni(UnrealTargetPlatform.HTML5, "Engine", Path.GetDirectoryName(Param.RawProjectPath), CommandUtils.CombinePaths(CommandUtils.CmdEnv.LocalRoot, "Engine"));
+        ConfigCache.GetBool("/Script/HTML5PlatformEditor.HTML5TargetSettings", "UseAsyncLevelLoading", out UseAsyncLevelLoading);
+
+        if (Param.Run)
+            return false; 
+
+        return UseAsyncLevelLoading;
+    }
 
     public HTMLPakAutomation(ProjectParams InParam, DeploymentContext InSC)
     {
