build script, compile tools for linux [UE-12934]

diff --git a/Engine/Source/Programs/AutomationTool/Scripts/GUBP.Automation.cs b/Engine/Source/Programs/AutomationTool/Scripts/GUBP.Automation.cs
--- a/Engine/Source/Programs/AutomationTool/Scripts/GUBP.Automation.cs
+++ b/Engine/Source/Programs/AutomationTool/Scripts/GUBP.Automation.cs
@@ -940,8 +940,85 @@ public class GUBP : BuildCommand
 			Agenda.AddTargets(
 				new string[] { bp.Branch.BaseEngineProject.Properties.Targets[TargetRules.TargetType.Editor].TargetName },
 				UnrealTargetPlatform.Linux, UnrealTargetConfiguration.Development, InAddArgs: AddArgs);
+			 foreach (var ProgramTarget in bp.Branch.BaseEngineProject.Properties.Programs)
+            {
+                if (ProgramTarget.Rules.GUBP_AlwaysBuildWithBaseEditor() && ProgramTarget.Rules.SupportsPlatform(UnrealTargetPlatform.Linux))
+                {
+                    Agenda.AddTargets(new string[] { ProgramTarget.TargetName }, UnrealTargetPlatform.Linux, UnrealTargetConfiguration.Development, InAddArgs: AddArgs);
+                }
+            }
 			return Agenda;
 		}
+        void DeleteStaleDLLs(GUBP bp)
+        {
+            if (GUBP.bForceIncrementalCompile)
+            {
+                return;
+            }
+            var Targets = new List<string> { bp.Branch.BaseEngineProject.Properties.Targets[TargetRules.TargetType.Editor].TargetName };
+            foreach (var ProgramTarget in bp.Branch.BaseEngineProject.Properties.Programs)
+            {
+                if (ProgramTarget.Rules.GUBP_AlwaysBuildWithBaseEditor() && ProgramTarget.Rules.SupportsPlatform(UnrealTargetPlatform.Linux))
+                {
+                    Targets.Add(ProgramTarget.TargetName);
+                }
+            }
+
+
+            foreach (var Target in Targets)
+            {
+                var EnginePlatformBinaries = CommandUtils.CombinePaths(CommandUtils.CmdEnv.LocalRoot, "Engine", "Binaries", UnrealTargetPlatform.Linux.ToString());
+                var Wildcard = Target + "-*";
+                Log("************Deleting stale editor DLLs, path {0} wildcard {1}", EnginePlatformBinaries, Wildcard);
+                foreach (var DiskFile in FindFiles(Wildcard, true, EnginePlatformBinaries))
+                {
+                    bool IsBuildProduct = false;
+                    foreach (var Product in BuildProducts)
+                    {
+                        if (Product.Equals(DiskFile, StringComparison.InvariantCultureIgnoreCase))
+                        {
+                            IsBuildProduct = true;
+                            break;
+                        }
+                    }
+                    if (!IsBuildProduct)
+                    {
+                        DeleteFile(DiskFile);
+                    }
+                }
+                var EnginePluginBinaries = CommandUtils.CombinePaths(CommandUtils.CmdEnv.LocalRoot, "Engine", "Plugins");
+                var HostSubstring = CommandUtils.CombinePaths("/", UnrealTargetPlatform.Linux.ToString(), "/");
+                Log("************Deleting stale editor DLLs, path {0} wildcard {1} host {2}", EnginePluginBinaries, Wildcard, HostSubstring);
+                foreach (var DiskFile in FindFiles(Wildcard, true, EnginePluginBinaries))
+                {
+                    if (DiskFile.IndexOf(HostSubstring, StringComparison.InvariantCultureIgnoreCase) < 0)
+                    {
+                        continue;
+                    }
+                    bool IsBuildProduct = false;
+                    foreach (var Product in BuildProducts)
+                    {
+                        if (Product.Equals(DiskFile, StringComparison.InvariantCultureIgnoreCase))
+                        {
+                            IsBuildProduct = true;
+                            break;
+                        }
+                    }
+                    if (!IsBuildProduct)
+                    {
+                        DeleteFile(DiskFile);
+                    }
+                }
+            }
+        }
+        public override void PostLoadFromSharedTempStorage(GUBP bp)
+        {
+            DeleteStaleDLLs(bp);
+        }
+        public override void PostBuildProducts(GUBP bp)
+        {
+            DeleteStaleDLLs(bp);        
+		}
 	}
     public class ToolsNode : CompileNode
     {
diff --git a/Engine/Source/Programs/UnrealPak/UnrealPak.Target.cs b/Engine/Source/Programs/UnrealPak/UnrealPak.Target.cs
--- a/Engine/Source/Programs/UnrealPak/UnrealPak.Target.cs
+++ b/Engine/Source/Programs/UnrealPak/UnrealPak.Target.cs
@@ -55,7 +55,7 @@ public class UnrealPakTarget : TargetRules
     {
         bInternalToolOnly = false;
         SeparateNode = false;
-		CrossCompile = false;
+		CrossCompile = true;
         return true;
     }
 }
