build script, temp storage retries on mac

diff --git a/Engine/Source/Programs/AutomationTool/CommandUtils.cs b/Engine/Source/Programs/AutomationTool/CommandUtils.cs
--- a/Engine/Source/Programs/AutomationTool/CommandUtils.cs
+++ b/Engine/Source/Programs/AutomationTool/CommandUtils.cs
@@ -1943,7 +1943,7 @@ namespace AutomationTool
             }
             string Root = RootSharedTempStorageDirectory();
             string Result = CombinePaths(Root, GameFolder);
-            if (String.IsNullOrEmpty(GameFolder) || !DirectoryExistsAndIsWritable_NoExceptions(Result))
+            if (String.IsNullOrEmpty(GameFolder) || !Robust_DirectoryExistsAndIsWritable_NoExceptions(Result))
             {
                 string GameStr = "Game";
                 bool HadGame = false;
@@ -1953,10 +1953,10 @@ namespace AutomationTool
                     Result = CombinePaths(Root, ShortFolder);
                     HadGame = true;
                 }
-                if (!HadGame || !DirectoryExistsAndIsWritable_NoExceptions(Result))
+                if (!HadGame || !Robust_DirectoryExistsAndIsWritable_NoExceptions(Result))
                 {
                     Result = CombinePaths(Root, "UE4");
-                    if (!DirectoryExistsAndIsWritable_NoExceptions(Result))
+                    if (!Robust_DirectoryExistsAndIsWritable_NoExceptions(Result))
                     {
                         throw new AutomationException("Could not find an appropriate shared temp folder {0}", Result);
                     }
diff --git a/Engine/Source/Programs/AutomationTool/TempStorage.cs b/Engine/Source/Programs/AutomationTool/TempStorage.cs
--- a/Engine/Source/Programs/AutomationTool/TempStorage.cs
+++ b/Engine/Source/Programs/AutomationTool/TempStorage.cs
@@ -145,6 +145,30 @@ namespace AutomationTool
 			}
 			return bFound;
 		}
+        public static bool Robust_DirectoryExistsAndIsWritable_NoExceptions(string Directoryname)
+        {
+            bool bFound = false;
+            if (!DirectoryExistsAndIsWritable_NoExceptions(Directoryname))
+            {
+                // mac is terrible on shares, this isn't a solution, but a stop gap
+                if (Directoryname.StartsWith("/Volumes/"))
+                {
+                    int Retry = 0;
+                    while (!bFound && Retry < 60)
+                    {
+                        CommandUtils.Log(System.Diagnostics.TraceEventType.Warning, "*** Mac temp storage retry {0}", Directoryname);
+                        System.Threading.Thread.Sleep(10000);
+                        bFound = DirectoryExistsAndIsWritable_NoExceptions(Directoryname);
+                        Retry++;
+                    }
+                }
+            }
+            else
+            {
+                bFound = true;
+            }
+            return bFound;
+        }
 
         public class TempStorageManifest
         {
@@ -592,7 +616,7 @@ namespace AutomationTool
                 int Retries = 0;
                 while (Retries < 24)
                 {
-                    if (DirectoryExistsAndIsWritable_NoExceptions(Dir))
+                    if (Robust_DirectoryExistsAndIsWritable_NoExceptions(Dir))
                     {
                         return true;
                     }
