Re-enable Rocket as a part of CIS. Partial installs are no longer passed around as build products; filtering generates a manifest which is only used when the copy happens, and copying writes directly to the network share. Building samples takes a smaller subset of the install.

diff --git a/Engine/Build/InstalledEngineFilters.ini b/Engine/Build/InstalledEngineFilters.ini
--- a/Engine/Build/InstalledEngineFilters.ini
+++ b/Engine/Build/InstalledEngineFilters.ini
@@ -31,6 +31,7 @@
 /Engine/Source/Runtime/...
 /Engine/Source/Developer/...
 /Engine/Source/Editor/...
+/Engine/Source/ThirdParty/.../*.Build.cs
 /Engine/Source/ThirdParty/Licenses/...
 
 ; Shaders
@@ -97,12 +98,6 @@
 
 {Mac} /Engine/Source/UE4Game*.cs
 
-[StripEditor]
-{Win64} /Engine/Binaries/.../*.pdb
-{Win64} /Engine/Plugins/.../*.pdb
-{Win64} -/Engine/Binaries/DotNET/...
-{Win64} -/Engine/.../UE4Editor*
-
 [CopyTargetPlatforms]
 ; API docs
 /Engine/Documentation/CHM/API.chm
@@ -122,12 +117,6 @@
 {Win64} /Engine/Extras/Android/*.exe
 {Mac} /Engine/Extras/Android/*.dmg
 
-[StripTargetPlatform.Android]
-/Engine/Binaries/Android/....a
-/Engine/Binaries/Android/....so
-/Engine/Plugins/.../Binaries/Android/....a
-/Engine/Plugins/.../Binaries/Android/....so
-
 [CopyTargetPlatform.IOS]
 ; Build files
 /Engine/Build/IOS/...
@@ -167,14 +156,6 @@
 /Engine/Binaries/ThirdParty/PhysX/PhysX-3.3/Win32/VS2013/nvToolsExt32_1.dll
 /Engine/Binaries/ThirdParty/Steamworks/Steamv132/Win32/*.dll
 
-[StripTargetPlatform.Win32]
-/Engine/Binaries/Win32/....pdb
-/Engine/Plugins/.../Binaries/Win32/....pdb
-
-[StripTargetPlatform.Win64]
-/Engine/Binaries/Win64/....pdb
-/Engine/Plugins/.../Binaries/Win64/....pdb
-
 [CopyTargetPlatform.Linux]
 ; Content-only: include pre-compiled executables manually
 /Engine/Binaries/Linux/CrashReportClient
@@ -185,5 +166,41 @@
 /Engine/Binaries/ThirdParty/OpenAL/Linux/*
 /Engine/Binaries/ThirdParty/Steamworks/Steamv132/Linux/*.so
 
-[StripTargetPlatform.Linux]
+[StripSymbols.Android]
+/Engine/Binaries/Android/....a
+/Engine/Binaries/Android/....so
+/Engine/Plugins/.../Binaries/Android/....a
+/Engine/Plugins/.../Binaries/Android/....so
+
+[StripSymbols.Win32]
+/Engine/Binaries/Win32/....pdb
+/Engine/Plugins/.../Binaries/Win32/....pdb
+
+[StripSymbols.Win64]
+/Engine/Binaries/Win64/....pdb
+/Engine/Plugins/.../Binaries/Win64/....pdb
+{Win64} -UE4Editor*.pdb
+
+[StripSymbols.Linux]
 /Engine/Binaries/Linux/*.
+
+[CopyEditorForSamples]
+; Include only binaries for the host platform
+-/Engine/Binaries/...
+/Engine/Binaries/DotNET/...
+/Engine/Binaries/ThirdParty/...
+{Win64} /Engine/Binaries/Win64/...
+{Mac} /Engine/Binaries/Mac/...
+
+; Exclude other files
+-/Engine/Documentation/...
+-/Engine/Extras/...
+-/Engine/Source/.../Private/...
+
+; Exclude other content
+-/FeaturePacks/...
+-/Samples/...
+-/Templates/...
+
+; No PDBs
+-*.pdb
diff --git a/Engine/Source/Programs/AutomationTool/CommandUtils.cs b/Engine/Source/Programs/AutomationTool/CommandUtils.cs
--- a/Engine/Source/Programs/AutomationTool/CommandUtils.cs
+++ b/Engine/Source/Programs/AutomationTool/CommandUtils.cs
@@ -1504,7 +1504,7 @@ namespace AutomationTool
 		}
 
 		/// <summary>
-		/// Copies files using miltiple threads
+		/// Copies files using multiple threads
 		/// </summary>
 		/// <param name="SourceDirectory"></param>
 		/// <param name="DestDirectory"></param>
@@ -1569,6 +1569,45 @@ namespace AutomationTool
 			}
 		}
 
+		/// <summary>
+		/// Copies a set of files from one folder to another
+		/// </summary>
+		/// <param name="SourceDir">Source directory</param>
+		/// <param name="TargetDir">Target directory</param>
+		/// <param name="RelativePaths">Paths relative to the source directory to copy</param>
+		/// <param name="bIgnoreSymlinks">Whether to ignore symlinks during the copy</param>
+		/// <param name="MaxThreads">Maximum number of threads to create</param>
+		/// <returns>Array of filenames copied to the target directory</returns>
+		public static string[] ThreadedCopyFiles(string SourceDir, string TargetDir, string[] RelativePaths, int MaxThreads = 64)
+		{
+			string[] SourceFileNames = new string[RelativePaths.Length];
+			string[] TargetFileNames = new string[RelativePaths.Length];
+			for(int Idx = 0; Idx < RelativePaths.Length; Idx++)
+			{
+				SourceFileNames[Idx] = CommandUtils.CombinePaths(SourceDir, RelativePaths[Idx]);
+				TargetFileNames[Idx] = CommandUtils.CombinePaths(TargetDir, RelativePaths[Idx]);
+			}
+			CommandUtils.ThreadedCopyFiles(SourceFileNames, TargetFileNames, MaxThreads);
+			return TargetFileNames;
+		}
+
+		/// <summary>
+		/// Copies a set of files from one folder to another
+		/// </summary>
+		/// <param name="SourceDir">Source directory</param>
+		/// <param name="TargetDir">Target directory</param>
+		/// <param name="Filter">Filter which selects files from the source directory to copy</param>
+		/// <param name="bIgnoreSymlinks">Whether to ignore symlinks during the copy</param>
+		/// <param name="MaxThreads">Maximum number of threads to create</param>
+		/// <returns>Array of filenames copied to the target directory</returns>
+		public static string[] ThreadedCopyFiles(string SourceDir, string TargetDir, FileFilter Filter, bool bIgnoreSymlinks, int MaxThreads = 64)
+		{
+			// Filter all the relative paths
+			CommandUtils.Log("Applying filter to {0}...", SourceDir);
+			string[] RelativePaths = Filter.ApplyToDirectory(SourceDir, bIgnoreSymlinks).ToArray();
+			return ThreadedCopyFiles(SourceDir, TargetDir, RelativePaths);
+		}
+
 		#endregion
 
 		#region Environment variables
diff --git a/Engine/Source/Programs/AutomationTool/FileFilter.cs b/Engine/Source/Programs/AutomationTool/FileFilter.cs
--- a/Engine/Source/Programs/AutomationTool/FileFilter.cs
+++ b/Engine/Source/Programs/AutomationTool/FileFilter.cs
@@ -1,5 +1,6 @@
 ﻿// Copyright 1998-2015 Epic Games, Inc. All Rights Reserved.
 
+using AutomationTool;
 using System;
 using System.Collections.Generic;
 using System.Diagnostics;
@@ -147,13 +148,39 @@ public class FileFilter
 		}
 	}
 
+	/// <summary>
+	/// Adds a rule which matches a filename relative to a given base directory.
+	/// </summary>
+	/// <param name="FileName">The filename to add a rule for</param>
+	/// <param name="BaseDirectoryName">Base directory for the rule</param>
+	/// <param name="Type">Whether to add an include or exclude rule</param>
+	public void AddRuleForFile(string FileName, string BaseDirectoryName, FileFilterType Type)
+	{
+		AddRuleForFiles(new string[]{ FileName }, BaseDirectoryName, Type);
+	}
+
+	/// <summary>
+	/// Adds rules for files which match the given names
+	/// </summary>
+	/// <param name="FileName">The filenames to rules for</param>
+	/// <param name="BaseDirectoryName">Base directory for the rules</param>
+	/// <param name="Type">Whether to add an include or exclude rule</param>
+	public void AddRuleForFiles(IEnumerable<string> FileNames, string BaseDirectoryName, FileFilterType Type)
+	{
+		string FullBaseDirectoryName = Path.GetFullPath(BaseDirectoryName);
+		foreach (string FileName in FileNames)
+		{
+			AddRule("/" + CommandUtils.StripBaseDirectory(Path.GetFullPath(FileName), FullBaseDirectoryName), Type);
+		}
+	}
+
 	/// <summary>
 	/// Reads a configuration file split into sections
 	/// </summary>
 	/// <param name="Filter"></param>
 	/// <param name="RulesFileName"></param>
 	/// <param name="Conditions"></param>
-	public void AddRulesFromFile(string FileName, string SectionName, params string[] AllowTags)
+	public void ReadRulesFromFile(string FileName, string SectionName, params string[] AllowTags)
 	{
 		bool bInSection = false;
 		foreach(string Line in File.ReadAllLines(FileName))
@@ -311,6 +338,26 @@ public class FileFilter
 		}
 	}
 
+	/// <summary>
+	/// Excludes all restricted platform folders from the filter
+	/// </summary>
+	public void ExcludeConfidentialPlatforms()
+	{
+		AddRule(".../PS4/...", FileFilterType.Exclude);
+		AddRule(".../XboxOne/...", FileFilterType.Exclude);
+	}
+
+	/// <summary>
+	/// Excludes all confidential folders from the filter
+	/// </summary>
+	public void ExcludeConfidentialFolders()
+	{
+		AddRule(".../EpicInternal/...", FileFilterType.Exclude);
+		AddRule(".../CarefullyRedist/...", FileFilterType.Exclude);
+		AddRule(".../NotForLicensees/...", FileFilterType.Exclude);
+		AddRule(".../NoRedist/...", FileFilterType.Exclude);
+	}
+
 	/// <summary>
 	/// Determines whether the given file matches the filter
 	/// </summary>
@@ -365,10 +412,10 @@ public class FileFilter
 	/// </summary>
 	/// <param name="FolderName">File to match</param>
 	/// <returns>True if the file passes the filter</returns>
-	public IEnumerable<string> ApplyToDirectory(string DirectoryName)
+	public IEnumerable<string> ApplyToDirectory(string DirectoryName, bool bIgnoreSymlinks)
 	{
 		List<string> MatchingFileNames = new List<string>();
-		FindMatchesFromDirectory(new DirectoryInfo(DirectoryName), "", MatchingFileNames);
+		FindMatchesFromDirectory(new DirectoryInfo(DirectoryName), "", bIgnoreSymlinks, MatchingFileNames);
 		return MatchingFileNames.ToArray();
 	}
 
@@ -377,12 +424,12 @@ public class FileFilter
 	/// </summary>
 	/// <param name="FolderName">File to match</param>
 	/// <returns>True if the file passes the filter</returns>
-	void FindMatchesFromDirectory(DirectoryInfo CurrentDirectory, string NamePrefix, List<string> MatchingFileNames)
+	void FindMatchesFromDirectory(DirectoryInfo CurrentDirectory, string NamePrefix, bool bIgnoreSymlinks, List<string> MatchingFileNames)
 	{
 		foreach (FileInfo NextFile in CurrentDirectory.EnumerateFiles())
 		{
 			string FileName = NamePrefix + NextFile.Name;
-			if (Matches(FileName))
+			if (Matches(FileName) && (!bIgnoreSymlinks || !NextFile.Attributes.HasFlag(FileAttributes.ReparsePoint)))
 			{
 				MatchingFileNames.Add(FileName);
 			}
@@ -392,7 +439,7 @@ public class FileFilter
 			string NextNamePrefix = NamePrefix + NextDirectory.Name;
 			if (PossiblyMatches(NextNamePrefix))
 			{
-				FindMatchesFromDirectory(NextDirectory, NextNamePrefix + "/", MatchingFileNames);
+				FindMatchesFromDirectory(NextDirectory, NextNamePrefix + "/", bIgnoreSymlinks, MatchingFileNames);
 			}
 		}
 	}
diff --git a/Engine/Source/Programs/AutomationTool/Scripts/GUBP.Automation.cs b/Engine/Source/Programs/AutomationTool/Scripts/GUBP.Automation.cs
--- a/Engine/Source/Programs/AutomationTool/Scripts/GUBP.Automation.cs
+++ b/Engine/Source/Programs/AutomationTool/Scripts/GUBP.Automation.cs
@@ -2464,6 +2464,7 @@ public class GUBP : BuildCommand
                             !Product.EndsWith("_Success.log", StringComparison.InvariantCultureIgnoreCase) &&
 							!Product.Replace('\\', '/').Contains("/Intermediate/") &&
 							!Product.Replace('\\', '/').Contains("/Engine/Saved/") &&
+							!Product.Replace('\\', '/').Contains("/DerivedDataCache/") &&
 							!Product.EndsWith(".lib") &&
 							!Product.EndsWith(".a") && 
 							!Product.EndsWith(".bc")
diff --git a/Engine/Source/Programs/AutomationTool/Scripts/RocketBuild.Automation.cs b/Engine/Source/Programs/AutomationTool/Scripts/RocketBuild.Automation.cs
--- a/Engine/Source/Programs/AutomationTool/Scripts/RocketBuild.Automation.cs
+++ b/Engine/Source/Programs/AutomationTool/Scripts/RocketBuild.Automation.cs
@@ -74,14 +74,7 @@ namespace Rocket
 		{
 			if(!bp.BranchOptions.bNoInstalledEngine)
 			{
-				// Get the promotable node, so we can add additional dependencies to it
-//				GUBP.GUBPNode PromotableNode = bp.FindNode(GUBP.SharedAggregatePromotableNode.StaticGetFullName());
-
-				// Get the output paths for the installed engine
-				string InstallDir = CommandUtils.CombinePaths(CommandUtils.CmdEnv.LocalRoot, "LocalBuilds", "Engine", CommandUtils.GetGenericPlatformName(HostPlatform));
-				string SymbolsDir = CommandUtils.CombinePaths(CommandUtils.CmdEnv.LocalRoot, "LocalBuilds", "Engine", CommandUtils.GetGenericPlatformName(HostPlatform) + "Symbols");
-
-				// Find all the target platforms for this host platform
+				// Find all the target platforms for this host platform.
 				List<UnrealTargetPlatform> TargetPlatforms = GetTargetPlatforms(bp, HostPlatform);
 
 				// Get a list of all the code target platforms
@@ -94,24 +87,32 @@ namespace Rocket
 					}
 				}
 
-				// Add the aggregate node for the whole install, if there isn't one already
-				if(!bp.HasNode(RocketAggregateNode.StaticGetFullName()))
-				{
-					bp.AddNode(new RocketAggregateNode());
-				}
+				// Build the DDC
+				bp.AddNode(new BuildDerivedDataCacheNode(HostPlatform, GetCookPlatforms(HostPlatform, TargetPlatforms), CurrentFeaturePacks));
+
+				// Strip symbols for all the targets built on the host platform
+				bp.AddNode(new StripRocketNode(bp, HostPlatform, CodeTargetPlatforms));
 
-				// Add a dependency on this platform finishing
-				RocketAggregateNode AggNode = (RocketAggregateNode)bp.FindNode(RocketAggregateNode.StaticGetFullName());
+				// Generate a list of files that needs to be copied for each target platform
+				bp.AddNode(new FilterRocketNode(bp, HostPlatform, TargetPlatforms, CodeTargetPlatforms, CurrentFeaturePacks, CurrentTemplates));
 
-				// Copy to output folders
-				AggNode.AddDependency(bp.AddNode(new CopyInstalledEditorNode(bp, HostPlatform, InstallDir)));
-				AggNode.AddDependency(bp.AddNode(new CopyInstalledPlatformsNode(bp, HostPlatform, TargetPlatforms, CodeTargetPlatforms, InstallDir, CurrentFeaturePacks, CurrentTemplates)));
-				AggNode.AddDependency(bp.AddNode(new CopyInstalledDerivedDataNode(HostPlatform, GetCookPlatforms(HostPlatform, TargetPlatforms), CurrentFeaturePacks, InstallDir)));
-				AggNode.AddDependency(bp.AddNode(new CopyInstalledSymbolsNode(bp, HostPlatform, CodeTargetPlatforms, SymbolsDir)));
+				// Copy the install to the output directory. This node adds a dependency on triggering a shared promotable if it's copying to the network.
+				if(ShouldDoSeriousThingsLikeP4CheckinAndPostToMCP())
+				{
+					string OutputDir = CommandUtils.CombinePaths(CommandUtils.RootSharedTempStorageDirectory(), "Rocket", "Automated", GetBuildLabel(), CommandUtils.GetGenericPlatformName(HostPlatform));
+					bp.AddNode(new CopyRocketNode(HostPlatform, OutputDir, true));
+					bp.AddNode(new CopyRocketSymbolsNode(bp, HostPlatform, CodeTargetPlatforms, OutputDir + "Symbols", true));
+				}
+				else
+				{
+					string OutputDir = CommandUtils.CombinePaths(CommandUtils.CmdEnv.LocalRoot, "LocalBuilds", "Rocket", CommandUtils.GetGenericPlatformName(HostPlatform));
+					bp.AddNode(new CopyRocketNode(HostPlatform, OutputDir, false));
+					bp.AddNode(new CopyRocketSymbolsNode(bp, HostPlatform, CodeTargetPlatforms, OutputDir + "Symbols", false));
+				}
 
 				// Add the aggregate node for the entire install
-				AggNode.AddDependency(bp.AddNode(new BuildInstalledEngineNode(HostPlatform, InstallDir, SymbolsDir)));
-//				PromotableNode.AddDependency(BuildInstalledEngineNode.StaticGetFullName(HostPlatform));
+				GUBP.GUBPNode PromotableNode = bp.FindNode(GUBP.SharedAggregatePromotableNode.StaticGetFullName());
+				PromotableNode.AddDependency(FilterRocketNode.StaticGetFullName(HostPlatform));
 
 				// Add a node for GitHub promotions
 				if(HostPlatform == UnrealTargetPlatform.Win64)
@@ -120,12 +121,17 @@ namespace Rocket
 					if(CommandUtils.FileExists(CommandUtils.CombinePaths(CommandUtils.CmdEnv.LocalRoot, GitConfigRelativePath)))
 					{
 						bp.AddNode(new BuildGitPromotable(bp, HostPlatform, GitConfigRelativePath));
-//						PromotableNode.AddDependency(BuildGitPromotable.StaticGetFullName(HostPlatform));
+						PromotableNode.AddDependency(BuildGitPromotable.StaticGetFullName(HostPlatform));
 					}
 				}
 			}
 		}
 
+		public static string GetBuildLabel()
+		{
+			return FEngineVersionSupport.FromVersionFile(CommandUtils.CombinePaths(CommandUtils.CmdEnv.LocalRoot, @"Engine\Source\Runtime\Launch\Resources\Version.h")).ToString();
+		}
+
 		public static List<UnrealTargetPlatform> GetTargetPlatforms(GUBP bp, UnrealTargetPlatform HostPlatform)
 		{
 			List<UnrealTargetPlatform> TargetPlatforms = new List<UnrealTargetPlatform>();
@@ -249,354 +255,171 @@ namespace Rocket
 		{
 			// Create a filter for all the promoted binaries
 			FileFilter PromotableFilter = new FileFilter();
-			CopyNode.AddFilesToFilter(PromotableFilter, AllDependencyBuildProducts, CommandUtils.CmdEnv.LocalRoot, FileFilterType.Include);
-			PromotableFilter.AddRulesFromFile(CommandUtils.CombinePaths(CommandUtils.CmdEnv.LocalRoot, ConfigRelativePath), "promotable");
-			CopyNode.ExcludeConfidentialFolders(PromotableFilter);
+			PromotableFilter.AddRuleForFiles(AllDependencyBuildProducts, CommandUtils.CmdEnv.LocalRoot, FileFilterType.Include);
+			PromotableFilter.ReadRulesFromFile(CommandUtils.CombinePaths(CommandUtils.CmdEnv.LocalRoot, ConfigRelativePath), "promotable");
+			PromotableFilter.ExcludeConfidentialFolders();
 			
 			// Copy everything that matches the filter to the promotion folder
 			string PromotableFolder = CommandUtils.CombinePaths(CommandUtils.CmdEnv.LocalRoot, "Engine", "Saved", "GitPromotable");
 			CommandUtils.DeleteDirectoryContents(PromotableFolder);
-			string[] PromotableFiles = CopyNode.CopyFolder(CommandUtils.CmdEnv.LocalRoot, PromotableFolder, PromotableFilter);
+			string[] PromotableFiles = CommandUtils.ThreadedCopyFiles(CommandUtils.CmdEnv.LocalRoot, PromotableFolder, PromotableFilter, bIgnoreSymlinks: true);
 			BuildProducts = new List<string>(PromotableFiles);
 		}
 	}
-
-	public class CopySymbolsNode : GUBP.HostPlatformNode
+	
+	public class StripRocketNode : GUBP.HostPlatformNode
 	{
-		public string OutputDirectory;
-		public bool bAddToBuildProducts;
+		List<UnrealTargetPlatform> Platforms = new List<UnrealTargetPlatform>();
+		public readonly string StrippedDir;
 
-		public CopySymbolsNode(UnrealTargetPlatform HostPlatform, string InOutputDirectory, bool bInAddToBuildProducts) : base(HostPlatform)
+		public StripRocketNode(GUBP bp, UnrealTargetPlatform InHostPlatform, IEnumerable<UnrealTargetPlatform> InTargetPlatforms) : base(InHostPlatform)
 		{
-			OutputDirectory = InOutputDirectory;
-			bAddToBuildProducts = bInAddToBuildProducts;
-			AgentSharingGroup = "BuildMinimalInstall" + StaticGetHostPlatformSuffix(HostPlatform);
-		}
-
-		public override void DoBuild(GUBP bp)
-		{
-			BuildProducts = new List<string>();
-
-			// Make a lookup for all the known debug extensions, and filter all the dependency build products against that
-			HashSet<string> DebugExtensions = new HashSet<string>(Platform.Platforms.Values.SelectMany(x => x.GetDebugFileExtentions()).Distinct().ToArray(), StringComparer.InvariantCultureIgnoreCase);
-			foreach(string InputFileName in AllDependencyBuildProducts)
-			{
-				string Extension = Path.GetExtension(InputFileName);
-				if(DebugExtensions.Contains(Extension))
-				{
-					string OutputFileName = CommandUtils.MakeRerootedFilePath(InputFileName, CommandUtils.CmdEnv.LocalRoot, OutputDirectory);
-					CommandUtils.CopyFile(InputFileName, OutputFileName);
-
-					if(bAddToBuildProducts)
-					{
-						BuildProducts.Add(OutputFileName);
-					}
-				}
-			}
-
-			// Add a dummy build product if we're not adding all the files we copied
-			if(!bAddToBuildProducts)
+			Platforms = new List<UnrealTargetPlatform>(InTargetPlatforms);
+			if(!Platforms.Contains(HostPlatform))
 			{
-				SaveRecordOfSuccessAndAddToBuildProducts();
+				Platforms.Add(HostPlatform);
 			}
-		}
-	}
+			StrippedDir = Path.GetFullPath(CommandUtils.CombinePaths(CommandUtils.CmdEnv.LocalRoot, "Engine", "Saved", "Rocket"));
 
-	public class CopyInstalledSymbolsNode : CopySymbolsNode
-	{
-		public CopyInstalledSymbolsNode(GUBP bp, UnrealTargetPlatform HostPlatform, IEnumerable<UnrealTargetPlatform> TargetPlatforms, string OutputDirectory) 
-			: base(HostPlatform, OutputDirectory, false)
-		{
+			// Add all the host nodes
 			AddDependency(GUBP.ToolsForCompileNode.StaticGetFullName(HostPlatform));
 			AddDependency(GUBP.RootEditorNode.StaticGetFullName(HostPlatform));
 			AddDependency(GUBP.ToolsNode.StaticGetFullName(HostPlatform));
 
-			foreach(UnrealTargetPlatform TargetPlatform in TargetPlatforms)
+			// Add all the target platform nodes
+			foreach(UnrealTargetPlatform TargetPlatform in InTargetPlatforms)
 			{
-				UnrealTargetPlatform SourceHostPlatform = CopyInstalledPlatformsNode.GetSourceHostPlatform(bp, HostPlatform, TargetPlatform);
-				AddDependency(GUBP.GamePlatformMonolithicsNode.StaticGetFullName(SourceHostPlatform, bp.Branch.BaseEngineProject, TargetPlatform, Precompiled: true));
+				UnrealTargetPlatform SourceHostPlatform = RocketBuild.GetSourceHostPlatform(bp, HostPlatform, TargetPlatform);
+				if(SourceHostPlatform == HostPlatform)
+				{
+					AddDependency(GUBP.GamePlatformMonolithicsNode.StaticGetFullName(HostPlatform, bp.Branch.BaseEngineProject, TargetPlatform, Precompiled: true));
+				}
 			}
 		}
 
 		public static string StaticGetFullName(UnrealTargetPlatform HostPlatform)
 		{
-			return "CopyEngineSymbols" + StaticGetHostPlatformSuffix(HostPlatform);
+			return "StripRocket" + StaticGetHostPlatformSuffix(HostPlatform);
 		}
 
 		public override string GetFullName()
 		{
 			return StaticGetFullName(HostPlatform);
 		}
-	}
-
-	public class CopyNode : GUBP.HostPlatformNode
-	{
-		public CopyNode(UnrealTargetPlatform InHostPlatform) : base(InHostPlatform)
-		{
-		}
 
-		static IEnumerable<string> FindBuildProducts(GUBP bp, string NodeName)
-		{
-			GUBP.GUBPNode Node = bp.FindNode(NodeName);
-			if (Node == null)
-			{
-				throw new AutomationException("Couldn't find node '{0}'", NodeName);
-			}
-			return Node.BuildProducts;
-		}
-
-		static IEnumerable<string> MakeRerootedFilePaths(IEnumerable<string> FileNames, string SourceDir, string TargetDir)
-		{
-			foreach(string FileName in FileNames)
-			{
-				yield return CommandUtils.MakeRerootedFilePath(Path.GetFullPath(FileName), SourceDir, TargetDir);
-			}
-		}
-
-		public static void AddFileToFilter(FileFilter Filter, string FileName, FileFilterType Type)
-		{
-			AddFileToFilter(Filter, FileName, CommandUtils.CmdEnv.LocalRoot, Type);
-		}
-
-		public static void AddFileToFilter(FileFilter Filter, string FileName, string BaseDirectoryName, FileFilterType Type)
-		{
-			Filter.AddRule("/" + CommandUtils.StripBaseDirectory(Path.GetFullPath(FileName), BaseDirectoryName), Type);
-		}
-
-		public static void AddFilesToFilter(FileFilter Filter, IEnumerable<string> FileNames, FileFilterType Type)
-		{
-			AddFilesToFilter(Filter, FileNames, CommandUtils.CmdEnv.LocalRoot, Type);
-		}
-
-		public static void AddFilesToFilter(FileFilter Filter, IEnumerable<string> FileNames, string BaseDirectoryName, FileFilterType Type)
+		public override void DoBuild(GUBP bp)
 		{
-			foreach (string FileName in FileNames)
-			{
-				AddFileToFilter(Filter, FileName, BaseDirectoryName, Type);
-			}
-		}
+			BuildProducts = new List<string>();
 
-		public static void AddNodeToFilter(FileFilter Filter, GUBP bp, string NodeName, FileFilterType Type)
-		{
-			AddNodeToFilter(Filter, bp, NodeName, CommandUtils.CmdEnv.LocalRoot, Type);
-		}
+			string InputDir = Path.GetFullPath(CommandUtils.CmdEnv.LocalRoot);
 
-		public static void AddNodeToFilter(FileFilter Filter, GUBP bp, string NodeName, string BaseDirectoryName, FileFilterType Type)
-		{
-			GUBP.GUBPNode Node = bp.FindNode(NodeName);
-			if(Node == null)
+			string RulesFileName = CommandUtils.CombinePaths(CommandUtils.CmdEnv.LocalRoot, "Engine", "Build", "InstalledEngineFilters.ini");
+			foreach(UnrealTargetPlatform Platform in Platforms)
 			{
-				throw new AutomationException("Couldn't find node '{0}'", NodeName);
-			}
-			AddFilesToFilter(Filter, Node.BuildProducts, BaseDirectoryName, Type);
-		}
-
-		public static void ExcludeConfidentialPlatforms(FileFilter Filter)
-		{
-			// Exclude platforms that can't be redistributed without an NDA
-			Filter.Exclude(".../PS4/...");
-			Filter.Exclude(".../XboxOne/...");
-		}
-
-		public static void ExcludeConfidentialFolders(FileFilter Filter)
-		{
-			// Exclude standard folders
-			Filter.Exclude(".../EpicInternal/...");
-			Filter.Exclude(".../CarefullyRedist/...");
-			Filter.Exclude(".../NotForLicensees/...");
-			Filter.Exclude(".../NoRedist/...");
-		}
-
-		public static string[] CopyFolder(string SourceDir, string TargetDir, FileFilter Filter)
-		{
-			// Filter all the relative paths
-			CommandUtils.Log("Applying filter to {0}...", SourceDir);
-			string[] RelativePaths = Filter.ApplyToDirectory(SourceDir).ToArray();
-			return CopyFolder(SourceDir, TargetDir, RelativePaths);
-		}
+				// Read the filter for files on this platform
+				FileFilter StripFilter = new FileFilter();
+				StripFilter.ReadRulesFromFile(RulesFileName, "StripSymbols." + Platform.ToString(), HostPlatform.ToString());
 
-		public static string[] CopyFolder(string SourceDir, string TargetDir, IEnumerable<string> RelativePaths)
-		{
-			// Find all the source and target filenames, removing any symlinks.
-			List<string> SourceFileNames = new List<string>();
-			List<string> TargetFileNames = new List<string>();
-			foreach(string RelativePath in RelativePaths)
-			{
-				string SourceFileName = CommandUtils.CombinePaths(SourceDir, RelativePath);
-				if(!File.GetAttributes(SourceFileName).HasFlag(FileAttributes.ReparsePoint))
+				// Apply the filter to the build products
+				List<string> SourcePaths = new List<string>();
+				List<string> TargetPaths = new List<string>();
+				foreach(string DependencyBuildProduct in AllDependencyBuildProducts)
 				{
-					SourceFileNames.Add(SourceFileName);
-					TargetFileNames.Add(CommandUtils.CombinePaths(TargetDir, RelativePath));
+					string RelativePath = CommandUtils.StripBaseDirectory(Path.GetFullPath(DependencyBuildProduct), InputDir);
+					if(StripFilter.Matches(RelativePath))
+					{
+						SourcePaths.Add(CommandUtils.CombinePaths(InputDir, RelativePath));
+						TargetPaths.Add(CommandUtils.CombinePaths(StrippedDir, RelativePath));
+					}
 				}
+
+				// Strip the files and add them to the build products
+				StripSymbols(Platform, SourcePaths.ToArray(), TargetPaths.ToArray());
+				BuildProducts.AddRange(TargetPaths);
 			}
 
-			// Copy them all over
-			string[] TargetFileNamesArray = TargetFileNames.ToArray();
-			CommandUtils.Log("Copying {0} files to {1}...", SourceFileNames.Count, TargetDir);
-			CommandUtils.ThreadedCopyFiles(SourceFileNames.ToArray(), TargetFileNamesArray);
-			return TargetFileNamesArray;
+			SaveRecordOfSuccessAndAddToBuildProducts();
 		}
 
-		public static void StripSymbols(UnrealTargetPlatform TargetPlatform, IEnumerable<string> FileNames)
+		public static void StripSymbols(UnrealTargetPlatform TargetPlatform, string[] SourceFileNames, string[] TargetFileNames)
 		{
 			if(TargetPlatform == UnrealTargetPlatform.Win32 || TargetPlatform == UnrealTargetPlatform.Win64)
 			{
 				string PDBCopyPath = Path.Combine(Environment.GetFolderPath(Environment.SpecialFolder.ProgramFilesX86), "MSBuild", "Microsoft", "VisualStudio", "v12.0", "AppxPackage", "PDBCopy.exe");
-				foreach(string FileName in FileNames)
-				{
-					string TempFileName = FileName + ".stripped";
-					CommandUtils.Run(PDBCopyPath, String.Format("{0} {1} -p", CommandUtils.MakePathSafeToUseWithCommandLine(FileName), CommandUtils.MakePathSafeToUseWithCommandLine(TempFileName)));
-					CommandUtils.DeleteFile(FileName);
-					CommandUtils.RenameFile(TempFileName, FileName);
-				}
-			}
-			else if(TargetPlatform == UnrealTargetPlatform.Android)
-			{
-				CommandUtils.LogWarning("Stripping has not yet been implemented for Android. Skipping files:");
-				foreach (string FileName in FileNames)
-				{
-					CommandUtils.LogWarning("  {0}", FileName);
-				}
-			}
-			else if(TargetPlatform == UnrealTargetPlatform.Linux)
-			{
-				CommandUtils.LogWarning("Stripping has not yet been implemented for Linux. Skipping files:");
-				foreach (string FileName in FileNames)
+				for(int Idx = 0; Idx < SourceFileNames.Length; Idx++)
 				{
-					CommandUtils.LogWarning("  {0}", FileName);
+					CommandUtils.CreateDirectory(Path.GetDirectoryName(TargetFileNames[Idx]));
+					CommandUtils.Run(PDBCopyPath, String.Format("{0} {1} -p", CommandUtils.MakePathSafeToUseWithCommandLine(SourceFileNames[Idx]), CommandUtils.MakePathSafeToUseWithCommandLine(TargetFileNames[Idx])));
 				}
 			}
 			else
 			{
-				if(FileNames.Any())
+				for(int Idx = 0; Idx < SourceFileNames.Length; Idx++)
 				{
-					throw new AutomationException("Symbol stripping has not been implemented for {0}", TargetPlatform.ToString());
+					CommandUtils.LogWarning("Stripping has not yet been implemented for {0}. Skipping strip for {1}.", TargetPlatform, SourceFileNames[Idx]);
+					CommandUtils.CopyFile(SourceFileNames[Idx], TargetFileNames[Idx]);
 				}
 			}
 		}
 	}
 
-	public class CopyInstalledEditorNode : CopyNode
-	{
-		string TargetDir;
-
-		public CopyInstalledEditorNode(GUBP bp, UnrealTargetPlatform InHostPlatform, string InTargetDir) 
-			: base(InHostPlatform)
-		{
-			TargetDir = InTargetDir;
-
-			AddDependency(GUBP.VersionFilesNode.StaticGetFullName());
-			AddDependency(GUBP.ToolsForCompileNode.StaticGetFullName(HostPlatform));
-			AddDependency(GUBP.RootEditorNode.StaticGetFullName(HostPlatform));
-			AddDependency(GUBP.ToolsNode.StaticGetFullName(HostPlatform));
-
-			// Add win64 tools on Mac, to get the win64 build of UBT and UAT
-			if (HostPlatform == UnrealTargetPlatform.Mac && bp.HostPlatforms.Contains(UnrealTargetPlatform.Win64))
-			{
-				AddDependency(GUBP.ToolsForCompileNode.StaticGetFullName(UnrealTargetPlatform.Win64));
-			}
-
-			AgentSharingGroup = "BuildMinimalInstall" + StaticGetHostPlatformSuffix(HostPlatform);
-		}
-
-		public override int CISFrequencyQuantumShift(GUBP bp)
-		{
-			return base.CISFrequencyQuantumShift(bp) + 2;
-		}
-
-		public static string StaticGetFullName(UnrealTargetPlatform InHostPlatform)
-		{
-			return "CopyEditor" + StaticGetHostPlatformSuffix(InHostPlatform);
-		}
-
-		public override string GetFullName()
-		{
-			return StaticGetFullName(HostPlatform);
-		}
-
-		public override void DoBuild(GUBP bp)
-		{
-			FileFilter Filter = new FileFilter();
-
-			// Include all the editor products
-			AddNodeToFilter(Filter, bp, GUBP.ToolsForCompileNode.StaticGetFullName(HostPlatform), FileFilterType.Include);
-			AddNodeToFilter(Filter, bp, GUBP.RootEditorNode.StaticGetFullName(HostPlatform), FileFilterType.Include);
-			AddNodeToFilter(Filter, bp, GUBP.ToolsNode.StaticGetFullName(HostPlatform), FileFilterType.Include);
-
-			// Include all the standard rules
-			string RulesFileName = CommandUtils.CombinePaths(CommandUtils.CmdEnv.LocalRoot, "Engine", "Build", "InstalledEngineFilters.ini");
-			Filter.AddRulesFromFile(RulesFileName, "CopyEditor", HostPlatform.ToString());
-
-			// Add the final exclusions for legal reasons.
-			ExcludeConfidentialPlatforms(Filter);
-			ExcludeConfidentialFolders(Filter);
-
-			// Wipe the target directory and copy all the files over
-			CommandUtils.DeleteDirectoryContents(TargetDir);
-
-			// Copy everything over and add the result as build products
-			string[] TargetFileNames = CopyFolder(CommandUtils.CmdEnv.LocalRoot, TargetDir, Filter);
-			BuildProducts = new List<string>(TargetFileNames);
-
-			// Strip symbols from all the debug files
-			FileFilter StripFilter = new FileFilter();
-			StripFilter.AddRulesFromFile(RulesFileName, "StripEditor", HostPlatform.ToString());
-			StripSymbols(HostPlatform, StripFilter.ApplyTo(TargetDir, TargetFileNames));
-
-			// Write the default command line to the root of the output directory
-			string CommandLinePath = Path.Combine(TargetDir, "UE4CommandLine.txt");
-			File.WriteAllText(CommandLinePath, "-installedengine -rocket");
-			AddBuildProduct(CommandLinePath);
-		}
-	}
-
-	public class CopyInstalledPlatformsNode : CopyNode
+	public class FilterRocketNode : GUBP.HostPlatformNode
 	{
+		List<UnrealTargetPlatform> SourceHostPlatforms;
 		List<UnrealTargetPlatform> TargetPlatforms;
 		List<UnrealTargetPlatform> CodeTargetPlatforms;
-		string InstallDir;
 		string[] CurrentFeaturePacks;
 		string[] CurrentTemplates;
+		public readonly string DepotManifestPath;
+		public readonly string StrippedManifestPath;
 
-		public CopyInstalledPlatformsNode(GUBP bp, UnrealTargetPlatform InHostPlatform, List<UnrealTargetPlatform> InTargetPlatforms, List<UnrealTargetPlatform> InCodeTargetPlatforms, string InInstallDir, string[] InCurrentFeaturePacks, string[] InCurrentTemplates)
+		public FilterRocketNode(GUBP bp, UnrealTargetPlatform InHostPlatform, List<UnrealTargetPlatform> InTargetPlatforms, List<UnrealTargetPlatform> InCodeTargetPlatforms, string[] InCurrentFeaturePacks, string[] InCurrentTemplates)
 			: base(InHostPlatform)
 		{
 			TargetPlatforms = new List<UnrealTargetPlatform>(InTargetPlatforms);
 			CodeTargetPlatforms = new List<UnrealTargetPlatform>(InCodeTargetPlatforms);
-			InstallDir = InInstallDir;
 			CurrentFeaturePacks = InCurrentFeaturePacks;
 			CurrentTemplates = InCurrentTemplates;
+			DepotManifestPath = CommandUtils.CombinePaths(CommandUtils.CmdEnv.LocalRoot, "Engine", "Saved", "Rocket", "DistillDepot" + StaticGetHostPlatformSuffix(HostPlatform) + ".txt");
+			StrippedManifestPath = CommandUtils.CombinePaths(CommandUtils.CmdEnv.LocalRoot, "Engine", "Saved", "Rocket", "DistillStripped" + StaticGetHostPlatformSuffix(HostPlatform) + ".txt");
 
-			// copy the editor
-			AddDependency(CopyInstalledEditorNode.StaticGetFullName(HostPlatform));
-
-			// add all the feature packs
-			AddDependency(GUBP.MakeFeaturePacksNode.StaticGetFullName(GUBP.MakeFeaturePacksNode.GetDefaultBuildPlatform(bp)));
+			// Add the editor
+			AddDependency(GUBP.VersionFilesNode.StaticGetFullName());
+			AddDependency(GUBP.ToolsForCompileNode.StaticGetFullName(HostPlatform));
+			AddDependency(GUBP.RootEditorNode.StaticGetFullName(HostPlatform));
+			AddDependency(GUBP.ToolsNode.StaticGetFullName(HostPlatform));
 
 			// Get all the external file lists. Android on Mac uses the same files as Win64.
 			foreach(UnrealTargetPlatform CodeTargetPlatform in CodeTargetPlatforms)
 			{
-				UnrealTargetPlatform SourceHostPlatform = GetSourceHostPlatform(bp, HostPlatform, CodeTargetPlatform);
+				UnrealTargetPlatform SourceHostPlatform = RocketBuild.GetSourceHostPlatform(bp, HostPlatform, CodeTargetPlatform);
 				AddDependency(GUBP.GamePlatformMonolithicsNode.StaticGetFullName(SourceHostPlatform, bp.Branch.BaseEngineProject, CodeTargetPlatform, Precompiled: true));
 			}
 
-			// Add win64 tools on Mac, to get the win64 build of IPP
+			// Add win64 tools on Mac, to get the win64 build of UBT, UAT and IPP
 			if (HostPlatform == UnrealTargetPlatform.Mac && bp.HostPlatforms.Contains(UnrealTargetPlatform.Win64))
 			{
 				AddDependency(GUBP.ToolsNode.StaticGetFullName(UnrealTargetPlatform.Win64));
 				AddDependency(GUBP.ToolsForCompileNode.StaticGetFullName(UnrealTargetPlatform.Win64));
 			}
-		}
 
-		public static UnrealTargetPlatform GetSourceHostPlatform(GUBP bp, UnrealTargetPlatform HostPlatform, UnrealTargetPlatform TargetPlatform)
-		{
-			if(HostPlatform == UnrealTargetPlatform.Mac && TargetPlatform == UnrealTargetPlatform.Android && bp.HostPlatforms.Contains(UnrealTargetPlatform.Win64))
+			// Add all the feature packs
+			AddDependency(GUBP.MakeFeaturePacksNode.StaticGetFullName(GUBP.MakeFeaturePacksNode.GetDefaultBuildPlatform(bp)));
+
+			// Add the DDC
+			AddDependency(BuildDerivedDataCacheNode.StaticGetFullName(HostPlatform));
+
+			// Find all the host platforms we need
+			SourceHostPlatforms = TargetPlatforms.Select(x => RocketBuild.GetSourceHostPlatform(bp, HostPlatform, x)).Distinct().ToList();
+			if(!SourceHostPlatforms.Contains(HostPlatform))
 			{
-				return UnrealTargetPlatform.Win64;
+				SourceHostPlatforms.Add(HostPlatform);
 			}
-			else
+
+			// Add the stripped target platforms
+			foreach(UnrealTargetPlatform SourceHostPlatform in SourceHostPlatforms)
 			{
-				return HostPlatform;
+				AddDependency(StripRocketNode.StaticGetFullName(SourceHostPlatform));
 			}
 		}
 
@@ -605,37 +428,50 @@ namespace Rocket
 			return base.CISFrequencyQuantumShift(bp) + 2;
 		}
 
-		public override float Priority()
+		public static string StaticGetFullName(UnrealTargetPlatform InHostPlatform)
 		{
-			return base.Priority() + 5.0f;
+			return "FilterRocket" + StaticGetHostPlatformSuffix(InHostPlatform);
+		}
+
+		public override string GetFullName()
+		{
+			return StaticGetFullName(HostPlatform);
 		}
 
 		public override void DoBuild(GUBP bp)
 		{
 			FileFilter Filter = new FileFilter();
 
+			// Include all the editor products
+			AddRuleForBuildProducts(Filter, bp, GUBP.ToolsForCompileNode.StaticGetFullName(HostPlatform), FileFilterType.Include);
+			AddRuleForBuildProducts(Filter, bp, GUBP.RootEditorNode.StaticGetFullName(HostPlatform), FileFilterType.Include);
+			AddRuleForBuildProducts(Filter, bp, GUBP.ToolsNode.StaticGetFullName(HostPlatform), FileFilterType.Include);
+
 			// Include the build dependencies for every code platform
 			foreach(UnrealTargetPlatform CodeTargetPlatform in CodeTargetPlatforms)
 			{
-				UnrealTargetPlatform SourceHostPlatform = GetSourceHostPlatform(bp, HostPlatform, CodeTargetPlatform);
+				UnrealTargetPlatform SourceHostPlatform = RocketBuild.GetSourceHostPlatform(bp, HostPlatform, CodeTargetPlatform);
 				string FileListPath = GUBP.GamePlatformMonolithicsNode.StaticGetBuildDependenciesPath(SourceHostPlatform, CodeTargetPlatform);
-				AddFilesToFilter(Filter, UnrealBuildTool.Utils.ReadClass<UnrealBuildTool.ExternalFileList>(FileListPath).FileNames, FileFilterType.Include);
+				Filter.AddRuleForFiles(UnrealBuildTool.Utils.ReadClass<UnrealBuildTool.ExternalFileList>(FileListPath).FileNames, CommandUtils.CmdEnv.LocalRoot, FileFilterType.Include);
 			}
 
 			// Add the monolithic binaries
 			foreach(UnrealTargetPlatform CodeTargetPlatform in CodeTargetPlatforms)
 			{
-				UnrealTargetPlatform SourceHostPlatform = GetSourceHostPlatform(bp, HostPlatform, CodeTargetPlatform);
-				AddNodeToFilter(Filter, bp, GUBP.GamePlatformMonolithicsNode.StaticGetFullName(SourceHostPlatform, bp.Branch.BaseEngineProject, CodeTargetPlatform, Precompiled: true), FileFilterType.Include);
+				UnrealTargetPlatform SourceHostPlatform = RocketBuild.GetSourceHostPlatform(bp, HostPlatform, CodeTargetPlatform);
+				AddRuleForBuildProducts(Filter, bp, GUBP.GamePlatformMonolithicsNode.StaticGetFullName(SourceHostPlatform, bp.Branch.BaseEngineProject, CodeTargetPlatform, Precompiled: true), FileFilterType.Include);
 			}
 
 			// Include the feature packs
 			foreach(string CurrentFeaturePack in CurrentFeaturePacks)
 			{
 				BranchInfo.BranchUProject Project = bp.Branch.FindGameChecked(CurrentFeaturePack);
-				AddFileToFilter(Filter, GUBP.MakeFeaturePacksNode.GetOutputFile(Project), FileFilterType.Include);
+				Filter.AddRuleForFile(GUBP.MakeFeaturePacksNode.GetOutputFile(Project), CommandUtils.CmdEnv.LocalRoot, FileFilterType.Include);
 			}
 
+			// Include the DDC
+			AddRuleForBuildProducts(Filter, bp, BuildDerivedDataCacheNode.StaticGetFullName(HostPlatform), FileFilterType.Include);
+
 			// Include all the templates
 			foreach (string Template in CurrentTemplates)
 			{
@@ -646,78 +482,110 @@ namespace Rocket
 				Filter.Exclude(TemplateRelativeDirectory + "/.../Intermediate/...");
 			}
 
-			// Custom rules for all platforms
+			// Include all the standard rules
 			string RulesFileName = CommandUtils.CombinePaths(CommandUtils.CmdEnv.LocalRoot, "Engine", "Build", "InstalledEngineFilters.ini");
-			Filter.AddRulesFromFile(RulesFileName, "CopyTargetPlatforms", HostPlatform.ToString());
+			Filter.ReadRulesFromFile(RulesFileName, "CopyEditor", HostPlatform.ToString());
+			Filter.ReadRulesFromFile(RulesFileName, "CopyTargetPlatforms", HostPlatform.ToString());
 
 			// Custom rules for each target platform
 			foreach(UnrealTargetPlatform TargetPlaform in TargetPlatforms)
 			{
 				string SectionName = String.Format("CopyTargetPlatform.{0}", TargetPlaform.ToString());
-				Filter.AddRulesFromFile(RulesFileName, SectionName, HostPlatform.ToString());
+				Filter.ReadRulesFromFile(RulesFileName, SectionName, HostPlatform.ToString());
 			}
 
-			// Exclude any the files which have already been added by the CopyInstalledEditor node
-			AddNodeToFilter(Filter, bp, CopyInstalledEditorNode.StaticGetFullName(HostPlatform), FileFilterType.Exclude);
-
 			// Add the final exclusions for legal reasons.
-			ExcludeConfidentialPlatforms(Filter);
-			ExcludeConfidentialFolders(Filter);
+			Filter.ExcludeConfidentialPlatforms();
+			Filter.ExcludeConfidentialFolders();
+
+			// Run the filter on the stripped symbols, and remove those files from the copy filter
+			List<string> StrippedFiles = new List<string>();
+			foreach(UnrealTargetPlatform SourceHostPlatform in SourceHostPlatforms)
+			{
+				StripRocketNode StripNode = (StripRocketNode)bp.FindNode(StripRocketNode.StaticGetFullName(SourceHostPlatform));
+				foreach(string BuildProduct in StripNode.BuildProducts)
+				{
+					if(Utils.IsFileUnderDirectory(BuildProduct, StripNode.StrippedDir))
+					{
+						string RelativePath = CommandUtils.StripBaseDirectory(Path.GetFullPath(BuildProduct), StripNode.StrippedDir);
+						if(Filter.Matches(RelativePath))
+						{
+							StrippedFiles.Add(RelativePath);
+							Filter.Exclude("/" + RelativePath);
+						}
+					}
+				}
+			}
+			WriteManifest(StrippedManifestPath, StrippedFiles);
 
-			// Copy all the files over
-			string[] TargetFileNames = CopyFolder(CommandUtils.CmdEnv.LocalRoot, InstallDir, Filter);
-			BuildProducts = new List<string>(TargetFileNames);
+			// Write the filtered list of depot files to disk, removing any symlinks
+			List<string> DepotFiles = Filter.ApplyToDirectory(CommandUtils.CmdEnv.LocalRoot, true).ToList();
+			WriteManifest(DepotManifestPath, DepotFiles);
 
-			// Strip symbols from any of the build products
-			foreach(UnrealTargetPlatform TargetPlatform in TargetPlatforms)
+			// Sort the list of output files
+			SortedDictionary<string, bool> SortedFiles = new SortedDictionary<string,bool>(StringComparer.InvariantCultureIgnoreCase);
+			foreach(string DepotFile in DepotFiles)
 			{
-				FileFilter StripFilter = new FileFilter();
-				StripFilter.AddRulesFromFile(RulesFileName, String.Format("StripTargetPlatform.{0}", TargetPlatform.ToString()), HostPlatform.ToString());
-				StripSymbols(TargetPlatform, StripFilter.ApplyTo(InstallDir, TargetFileNames));
+				SortedFiles.Add(DepotFile, false);
+			}
+			foreach(string StrippedFile in StrippedFiles)
+			{
+				SortedFiles.Add(StrippedFile, true);
 			}
-		}
 
-		public override string RootIfAnyForTempStorage() // this saves a bit of path space for temp storage
-		{
-			return InstallDir;
+			// Write the list to the log
+			CommandUtils.Log("Files to be included in Rocket build:");
+			foreach(KeyValuePair<string, bool> SortedFile in SortedFiles)
+			{
+				CommandUtils.Log("  {0}{1}", SortedFile.Key, SortedFile.Value? " (stripped)" : "");
+			}
+
+			// Add the build products
+			BuildProducts = new List<string>{ StrippedManifestPath, DepotManifestPath };
 		}
 
-		public static string StaticGetFullName(UnrealTargetPlatform InHostPlatform)
+		static void AddRuleForBuildProducts(FileFilter Filter, GUBP bp, string NodeName, FileFilterType Type)
 		{
-			return "CopyPlatforms" + StaticGetHostPlatformSuffix(InHostPlatform);
+			GUBP.GUBPNode Node = bp.FindNode(NodeName);
+			if(Node == null)
+			{
+				throw new AutomationException("Couldn't find node '{0}'", NodeName);
+			}
+			Filter.AddRuleForFiles(Node.BuildProducts, CommandUtils.CmdEnv.LocalRoot, Type);
 		}
 
-		public override string GetFullName()
+		public static void WriteManifest(string FileName, IEnumerable<string> Files)
 		{
-			return StaticGetFullName(HostPlatform);
+			CommandUtils.CreateDirectory(Path.GetDirectoryName(FileName));
+			CommandUtils.WriteAllLines(FileName, Files.ToArray());
 		}
 	}
 
-	public class CopyInstalledDerivedDataNode : GUBP.HostPlatformNode
+	public class CopyRocketNode : GUBP.HostPlatformNode
 	{
-		string TargetPlatforms;
-		string InstallDir;
-		string[] ProjectNames;
+		public readonly string OutputDir;
 
-		public CopyInstalledDerivedDataNode(UnrealTargetPlatform InHostPlatform, string InTargetPlatforms, string[] InProjectNames, string InInstallDir)
-			: base(InHostPlatform)
+		public CopyRocketNode(UnrealTargetPlatform HostPlatform, string InOutputDir, bool bWaitForPromotion) : base(HostPlatform)
 		{
-			TargetPlatforms = InTargetPlatforms;
-			ProjectNames = InProjectNames;
-			InstallDir = InInstallDir;
+			OutputDir = InOutputDir;
 
-			AddDependency(CopyInstalledEditorNode.StaticGetFullName(HostPlatform));
-			AgentSharingGroup = "BuildMinimalInstall" + StaticGetHostPlatformSuffix(InHostPlatform);
-		}
+			// If this is a CIS build, we need to wait for the promotable trigger before running. If not, make an agent sharing group so that the samples node can use our build products.
+			if(bWaitForPromotion)
+			{
+				AddDependency(GUBP.WaitForSharedPromotionUserInput.StaticGetFullName(false));
+			}
+			else
+			{
+				AgentSharingGroup = "CopyRocket" + StaticGetHostPlatformSuffix(HostPlatform);
+			}
 
-		public static string StaticGetFullName(UnrealTargetPlatform InHostPlatform)
-		{
-            return "CopyDerivedData" + StaticGetHostPlatformSuffix(InHostPlatform);
+			AddDependency(FilterRocketNode.StaticGetFullName(HostPlatform));
+			AddDependency(StripRocketNode.StaticGetFullName(HostPlatform));
 		}
 
-		public override int CISFrequencyQuantumShift(GUBP bp)
+		public static string StaticGetFullName(UnrealTargetPlatform HostPlatform)
 		{
-			return base.CISFrequencyQuantumShift(bp) + 2;
+			return "CopyRocket" + StaticGetHostPlatformSuffix(HostPlatform);
 		}
 
 		public override string GetFullName()
@@ -727,101 +595,148 @@ namespace Rocket
 
 		public override void DoBuild(GUBP bp)
 		{
-			// Remove any existing DDP file
-			string OutputFileName = CommandUtils.CombinePaths(InstallDir, "Engine/DerivedDataCache/Compressed.ddp");
-			if(CommandUtils.FileExists(OutputFileName))
-			{
-				CommandUtils.DeleteFile(OutputFileName);
-			}
+			CommandUtils.DeleteDirectoryContents(OutputDir);
 
-			// Enumerate all the files in the install directory. We'll remove everything in this list that's not a build product once we're done.
-			string[] InitialFiles = Directory.EnumerateFiles(InstallDir, "*", SearchOption.AllDirectories).ToArray();
+			// Copy the depot files to the output directory
+			FilterRocketNode FilterNode = (FilterRocketNode)bp.FindNode(FilterRocketNode.StaticGetFullName(HostPlatform));
+			CopyManifestFilesToOutput(FilterNode.DepotManifestPath, CommandUtils.CmdEnv.LocalRoot, OutputDir);
 
-			// Generate DDC for all the non-code projects. We don't necessarily have editor DLLs for the code projects, but they should be the same as their blueprint counterparts.
-			foreach(string ProjectName in ProjectNames)
+			// Copy the stripped symbols to the output directory
+			StripRocketNode StripNode = (StripRocketNode)bp.FindNode(StripRocketNode.StaticGetFullName(HostPlatform));
+			CopyManifestFilesToOutput(FilterNode.StrippedManifestPath, StripNode.StrippedDir, OutputDir);
+
+			// Create a dummy build product
+			BuildProducts = new List<string>();
+			SaveRecordOfSuccessAndAddToBuildProducts();
+		}
+
+		static void CopyManifestFilesToOutput(string ManifestPath, string InputDir, string OutputDir)
+		{
+			// Read the files from the manifest
+			CommandUtils.Log("Reading manifest: '{0}'", ManifestPath);
+			string[] Files = CommandUtils.ReadAllLines(ManifestPath).Select(x => x.Trim()).Where(x => x.Length > 0).ToArray();
+
+			// Create lists of source and target files
+			CommandUtils.Log("Preparing file lists...");
+			string[] SourceFiles = Files.Select(x => CommandUtils.CombinePaths(InputDir, x)).ToArray();
+			string[] TargetFiles = Files.Select(x => CommandUtils.CombinePaths(OutputDir, x)).ToArray();
+
+			// Copy everything
+			CommandUtils.ThreadedCopyFiles(SourceFiles, TargetFiles);
+		}
+	}
+
+	public class CopyRocketSymbolsNode : GUBP.HostPlatformNode
+	{
+		string SymbolsOutputDir;
+
+		public CopyRocketSymbolsNode(GUBP bp, UnrealTargetPlatform HostPlatform, IEnumerable<UnrealTargetPlatform> TargetPlatforms, string InSymbolsOutputDir, bool bWaitForPromotion) : base(HostPlatform)
+		{
+			SymbolsOutputDir = InSymbolsOutputDir;
+
+			if(bWaitForPromotion)
 			{
-                BranchInfo.BranchUProject Project = bp.Branch.FindGameChecked(ProjectName);
-				if(!Project.Properties.bIsCodeBasedProject)
-				{
-					string InstalledProjectPath = CommandUtils.MakeRerootedFilePath(Project.FilePath, CommandUtils.CmdEnv.LocalRoot, InstallDir);
-					if(!CommandUtils.FileExists(InstalledProjectPath))
-					{
-						CommandUtils.CopyFile(Project.FilePath, InstalledProjectPath);
-					}
-					CommandUtils.Log("Running RocketEditor to generate engine DDC data for {0} on {1}", InstalledProjectPath, TargetPlatforms);
-					CommandUtils.DDCCommandlet(InstalledProjectPath, CommandUtils.GetEditorCommandletExe(InstallDir, HostPlatform), null, TargetPlatforms, "-fill -DDC=CreateInstalledEnginePak -MergePaks");
-				}
+				AddDependency(GUBP.WaitForSharedPromotionUserInput.StaticGetFullName(false));
 			}
 
-			// Create the list of DDC build products 
-			BuildProducts = new List<string>();
-			BuildProducts.Add(OutputFileName);
+			AddDependency(GUBP.ToolsForCompileNode.StaticGetFullName(HostPlatform));
+			AddDependency(GUBP.RootEditorNode.StaticGetFullName(HostPlatform));
+			AddDependency(GUBP.ToolsNode.StaticGetFullName(HostPlatform));
 
-			// Remove all the files all the temporary files that can be removed
-			SortedSet<string> TemporaryFiles = new SortedSet<string>(Directory.EnumerateFiles(InstallDir, "*", SearchOption.AllDirectories));
-			TemporaryFiles.ExceptWith(InitialFiles);
-			TemporaryFiles.ExceptWith(BuildProducts.Select(x => Path.GetFullPath(x)));
-			foreach(string TemporaryFile in TemporaryFiles)
+			foreach(UnrealTargetPlatform TargetPlatform in TargetPlatforms)
 			{
-				CommandUtils.DeleteFile(TemporaryFile);
+				UnrealTargetPlatform SourceHostPlatform = RocketBuild.GetSourceHostPlatform(bp, HostPlatform, TargetPlatform);
+				AddDependency(GUBP.GamePlatformMonolithicsNode.StaticGetFullName(SourceHostPlatform, bp.Branch.BaseEngineProject, TargetPlatform, Precompiled: true));
 			}
+		}
+
+		public static string StaticGetFullName(UnrealTargetPlatform HostPlatform)
+		{
+			return "CopyRocketSymbols" + StaticGetHostPlatformSuffix(HostPlatform);
+		}
 
-			// Remove any empty directories left over
-			foreach(string DirectoryName in Directory.EnumerateDirectories(InstallDir, "*", SearchOption.AllDirectories).OrderByDescending(x => x.Length))
+		public override string GetFullName()
+		{
+			return StaticGetFullName(HostPlatform);
+		}
+
+		public override void DoBuild(GUBP bp)
+		{
+			// Make a lookup for all the known debug extensions, and filter all the dependency build products against that
+			HashSet<string> DebugExtensions = new HashSet<string>(Platform.Platforms.Values.SelectMany(x => x.GetDebugFileExtentions()).Distinct().ToArray(), StringComparer.InvariantCultureIgnoreCase);
+			foreach(string InputFileName in AllDependencyBuildProducts)
 			{
-				if(!Directory.EnumerateFileSystemEntries(DirectoryName).GetEnumerator().MoveNext())
+				string Extension = Path.GetExtension(InputFileName);
+				if(DebugExtensions.Contains(Extension))
 				{
-					try { Directory.Delete(DirectoryName); } catch(Exception){ }
+					string OutputFileName = CommandUtils.MakeRerootedFilePath(InputFileName, CommandUtils.CmdEnv.LocalRoot, SymbolsOutputDir);
+					CommandUtils.CopyFile(InputFileName, OutputFileName);
 				}
 			}
-		}
 
-		public override float Priority()
-		{
-			return base.Priority() + 55.0f;
+			// Add a dummy build product
+			BuildProducts = new List<string>();
+			SaveRecordOfSuccessAndAddToBuildProducts();
 		}
 	}
 
-	public class BuildInstalledEngineNode : GUBP.HostPlatformAggregateNode
+	public class BuildDerivedDataCacheNode : GUBP.HostPlatformNode
 	{
-		public readonly string InstallDir;
-		public readonly string SymbolsDir;
+		string TargetPlatforms;
+		string[] ProjectNames;
 
-		public BuildInstalledEngineNode(UnrealTargetPlatform InHostPlatform, string InInstallDir, string InSymbolsDir)
+		public BuildDerivedDataCacheNode(UnrealTargetPlatform InHostPlatform, string InTargetPlatforms, string[] InProjectNames)
 			: base(InHostPlatform)
 		{
-			InstallDir = InInstallDir;
-			SymbolsDir = InSymbolsDir;
+			TargetPlatforms = InTargetPlatforms;
+			ProjectNames = InProjectNames;
 
-			AddDependency(CopyInstalledEditorNode.StaticGetFullName(InHostPlatform));
-			AddDependency(CopyInstalledPlatformsNode.StaticGetFullName(InHostPlatform));
-			AddDependency(CopyInstalledDerivedDataNode.StaticGetFullName(InHostPlatform));
-			AddDependency(CopyInstalledSymbolsNode.StaticGetFullName(InHostPlatform));
+			AddDependency(GUBP.RootEditorNode.StaticGetFullName(HostPlatform));
+			AgentSharingGroup = "Editor" + StaticGetHostPlatformSuffix(InHostPlatform);
 		}
 
 		public static string StaticGetFullName(UnrealTargetPlatform InHostPlatform)
 		{
-			return "BuildInstalledEngine" + StaticGetHostPlatformSuffix(InHostPlatform);
+            return "BuildDerivedDataCache" + StaticGetHostPlatformSuffix(InHostPlatform);
 		}
 
-		public override string GetFullName()
+		public override int CISFrequencyQuantumShift(GUBP bp)
 		{
-			return StaticGetFullName(HostPlatform);
+			return base.CISFrequencyQuantumShift(bp) + 2;
 		}
-	}
 
-	public class RocketAggregateNode : GUBP.AggregateNode
-	{
-		public RocketAggregateNode()
+		public override string GetFullName()
 		{
+			return StaticGetFullName(HostPlatform);
 		}
-		public static string StaticGetFullName()
+
+		public override void DoBuild(GUBP bp)
 		{
-			return "Rocket_Aggregate";
+			BuildProducts = new List<string>();
+			if(!bp.ParseParam("NoDDC"))
+			{
+				// Delete the output file
+				string OutputFile = CommandUtils.CombinePaths(CommandUtils.CmdEnv.LocalRoot, "Engine", "DerivedDataCache", "Compressed.ddp");
+				CommandUtils.DeleteFile(OutputFile);
+
+				// Generate DDC for all the non-code projects. We don't necessarily have editor DLLs for the code projects, but they should be the same as their blueprint counterparts.
+				foreach(string ProjectName in ProjectNames)
+				{
+					BranchInfo.BranchUProject Project = bp.Branch.FindGameChecked(ProjectName);
+					if(!Project.Properties.bIsCodeBasedProject)
+					{
+						CommandUtils.Log("Running RocketEditor to generate engine DDC data for {0} on {1}", Project.GameName, TargetPlatforms);
+						CommandUtils.DDCCommandlet(Project.FilePath, CommandUtils.GetEditorCommandletExe(CommandUtils.CmdEnv.LocalRoot, HostPlatform), null, TargetPlatforms, "-fill -DDC=CreateInstalledEnginePak -MergePaks");
+					}
+				}
+				BuildProducts.Add(OutputFile);
+			}
+			SaveRecordOfSuccessAndAddToBuildProducts();
 		}
-		public override string GetFullName()
+
+		public override float Priority()
 		{
-			return StaticGetFullName();
+			return base.Priority() + 55.0f;
 		}
 	}
 }
