Fix trigger dependencies for building Rocket samples.

diff --git a/Engine/Source/Programs/AutomationTool/Scripts/GUBP.Automation.cs b/Engine/Source/Programs/AutomationTool/Scripts/GUBP.Automation.cs
--- a/Engine/Source/Programs/AutomationTool/Scripts/GUBP.Automation.cs
+++ b/Engine/Source/Programs/AutomationTool/Scripts/GUBP.Automation.cs
@@ -1466,7 +1466,6 @@ public class GUBP : BuildCommand
         {
 			Projects = new List<BranchInfo.BranchUProject>(InProjects);
 			AddDependency(ToolsNode.StaticGetFullName(InHostPlatform)); // for UnrealPak
-            AgentSharingGroup = "FeaturePacks"  + StaticGetHostPlatformSuffix(InHostPlatform);
         }
 
 		public static string GetOutputFile(BranchInfo.BranchUProject Project)
