PR #1014: Fix staging, at least in terminal, on linux.

- Contributed by [REDACTED_PERSON] (and fixed to work when AddArgs don't end with a space).

diff --git a/Engine/Source/Programs/AutomationTool/UE4Build.cs b/Engine/Source/Programs/AutomationTool/UE4Build.cs
--- a/Engine/Source/Programs/AutomationTool/UE4Build.cs
+++ b/Engine/Source/Programs/AutomationTool/UE4Build.cs
@@ -327,7 +327,7 @@ namespace AutomationTool
 
 				DeleteFile(UBTManifest);
 
-				RunUBT(CmdEnv, UBTExecutable: UBTExecutable, Project: ProjectName, Target: TargetName, Platform: TargetPlatform.ToString(), Config: Config, AdditionalArgs: "-generatemanifest" + AddArgs, EnvVars: EnvVars);
+				RunUBT(CmdEnv, UBTExecutable: UBTExecutable, Project: ProjectName, Target: TargetName, Platform: TargetPlatform.ToString(), Config: Config, AdditionalArgs: AddArgs +  " -generatemanifest" , EnvVars: EnvVars);
 
 				PrepareManifest(UBTManifest, false);
 			}
