Fix error message

diff --git a/Engine/Source/Programs/AutomationTool/Scripts/GUBP.Automation.cs b/Engine/Source/Programs/AutomationTool/Scripts/GUBP.Automation.cs
--- a/Engine/Source/Programs/AutomationTool/Scripts/GUBP.Automation.cs
+++ b/Engine/Source/Programs/AutomationTool/Scripts/GUBP.Automation.cs
@@ -5169,7 +5169,7 @@ public class GUBP : BuildCommand
                 PreflightShelveCL = int.Parse(PreflightShelveCLString);
                 if (PreflightShelveCL < 2000000)
                 {
-                    throw new AutomationException("{0} does not look like a CL");
+                    throw new AutomationException(String.Format( "{0} does not look like a CL", PreflightShelveCL));
                 }
             }
             bPreflightBuild = true;
