Linux: place UAT log files into Documents

diff --git a/Engine/Source/Programs/AutomationTool/LinuxHostPlatform.cs b/Engine/Source/Programs/AutomationTool/LinuxHostPlatform.cs
--- a/Engine/Source/Programs/AutomationTool/LinuxHostPlatform.cs
+++ b/Engine/Source/Programs/AutomationTool/LinuxHostPlatform.cs
@@ -52,7 +52,8 @@ namespace AutomationTool
 
 		public override string LocalBuildsLogFolder
 		{
-			get { return Path.Combine(Environment.GetFolderPath(Environment.SpecialFolder.Personal), "Library/Logs/Unreal Engine/LocalBuildLogs"); }
+			// @FIXME: should use xdg-user-dir DOCUMENTS
+			get { return Path.Combine(Environment.GetFolderPath(Environment.SpecialFolder.Personal), "Documents/Unreal Engine/LocalBuildLogs"); }
 		}
 
 		public override string P4Exe
