Remove DDC builds from the editor ASG to improve build times.

diff --git a/Engine/Source/Programs/AutomationTool/Scripts/RocketBuild.Automation.cs b/Engine/Source/Programs/AutomationTool/Scripts/RocketBuild.Automation.cs
--- a/Engine/Source/Programs/AutomationTool/Scripts/RocketBuild.Automation.cs
+++ b/Engine/Source/Programs/AutomationTool/Scripts/RocketBuild.Automation.cs
@@ -692,7 +692,6 @@ namespace Rocket
 			ProjectNames = InProjectNames;
 
 			AddDependency(GUBP.RootEditorNode.StaticGetFullName(HostPlatform));
-			AgentSharingGroup = "Editor" + StaticGetHostPlatformSuffix(InHostPlatform);
 		}
 
 		public static string StaticGetFullName(UnrealTargetPlatform InHostPlatform)
