build script, don't spin on dirs that don't exist on mac

diff --git a/Engine/Source/Programs/AutomationTool/TempStorage.cs b/Engine/Source/Programs/AutomationTool/TempStorage.cs
--- a/Engine/Source/Programs/AutomationTool/TempStorage.cs
+++ b/Engine/Source/Programs/AutomationTool/TempStorage.cs
@@ -154,10 +154,15 @@ namespace AutomationTool
                 if (Directoryname.StartsWith("/Volumes/"))
                 {
                     int Retry = 0;
-                    while (!bFound && Retry < 60)
+                    int NumRetries = 60;
+                    if(Directoryname.Contains("UE4"))
                     {
-                        CommandUtils.Log(System.Diagnostics.TraceEventType.Warning, "*** Mac temp storage retry {0}", Directoryname);
-                        System.Threading.Thread.Sleep(10000);
+                        NumRetries = 2;
+                    }
+                    while (!bFound && Retry < NumRetries)
+                    {
+                        CommandUtils.Log("*** Mac temp storage retry {0}", Directoryname);
+                        System.Threading.Thread.Sleep(1000);
                         bFound = DirectoryExistsAndIsWritable_NoExceptions(Directoryname);
                         Retry++;
                     }
@@ -636,6 +641,7 @@ namespace AutomationTool
         {
             string SharedSubdir = SharedTempStorageDirectory();
             string Result = CombinePaths(ResolveSharedBuildDirectory(GameFolder), SharedSubdir);
+
             if (!Robust_DirectoryExists_NoExceptions(Result, "Could not find {0}"))
             {
                 CreateDirectory_NoExceptions(Result);
