build script, fix retry logic

diff --git a/Engine/Source/Programs/AutomationTool/TempStorage.cs b/Engine/Source/Programs/AutomationTool/TempStorage.cs
--- a/Engine/Source/Programs/AutomationTool/TempStorage.cs
+++ b/Engine/Source/Programs/AutomationTool/TempStorage.cs
@@ -155,7 +155,7 @@ namespace AutomationTool
                 {
                     int Retry = 0;
                     int NumRetries = 60;
-                    if(Directoryname.Contains("UE4"))
+                    if(!Directoryname.Contains("UE4"))
                     {
                         NumRetries = 2;
                     }
