Add missing UE4CommandLine.txt back into the CopyRocket node.

diff --git a/Engine/Source/Programs/AutomationTool/Scripts/RocketBuild.Automation.cs b/Engine/Source/Programs/AutomationTool/Scripts/RocketBuild.Automation.cs
--- a/Engine/Source/Programs/AutomationTool/Scripts/RocketBuild.Automation.cs
+++ b/Engine/Source/Programs/AutomationTool/Scripts/RocketBuild.Automation.cs
@@ -605,6 +605,10 @@ namespace Rocket
 			StripRocketNode StripNode = (StripRocketNode)bp.FindNode(StripRocketNode.StaticGetFullName(HostPlatform));
 			CopyManifestFilesToOutput(FilterNode.StrippedManifestPath, StripNode.StrippedDir, OutputDir);
 
+			// Write the UE4CommandLine.txt file with the 
+			string CommandLineFile = CommandUtils.CombinePaths(OutputDir, "UE4CommandLine.txt");
+			CommandUtils.WriteAllText(CommandLineFile, "-installedengine -rocket");
+
 			// Create a dummy build product
 			BuildProducts = new List<string>();
 			SaveRecordOfSuccessAndAddToBuildProducts();
