Temporarily skip building DDC as part of Rocket.

diff --git a/Engine/Source/Programs/AutomationTool/Scripts/RocketBuild.Automation.cs b/Engine/Source/Programs/AutomationTool/Scripts/RocketBuild.Automation.cs
--- a/Engine/Source/Programs/AutomationTool/Scripts/RocketBuild.Automation.cs
+++ b/Engine/Source/Programs/AutomationTool/Scripts/RocketBuild.Automation.cs
@@ -716,7 +716,7 @@ namespace Rocket
 		public override void DoBuild(GUBP bp)
 		{
 			BuildProducts = new List<string>();
-			if(!bp.ParseParam("NoDDC"))
+			if(!bp.ParseParam("NoDDC") && bp.ParseParam("WithDDC"))
 			{
 				// Delete the output file
 				string OutputFile = CommandUtils.CombinePaths(CommandUtils.CmdEnv.LocalRoot, "Engine", "DerivedDataCache", "Compressed.ddp");
