Add command line parameters for explicitly selecting target platforms for rocket builds.

diff --git a/Engine/Source/Programs/AutomationTool/Scripts/RocketBuild.Automation.cs b/Engine/Source/Programs/AutomationTool/Scripts/RocketBuild.Automation.cs
--- a/Engine/Source/Programs/AutomationTool/Scripts/RocketBuild.Automation.cs
+++ b/Engine/Source/Programs/AutomationTool/Scripts/RocketBuild.Automation.cs
@@ -91,7 +91,7 @@ namespace Rocket
 				bp.AddNode(new BuildDerivedDataCacheNode(HostPlatform, GetCookPlatforms(HostPlatform, TargetPlatforms), CurrentFeaturePacks));
 
 				// Strip symbols for all the targets built on the host platform
-				bp.AddNode(new StripRocketNode(bp, HostPlatform, CodeTargetPlatforms));
+				bp.AddNode(new StripRocketNode(bp, HostPlatform, TargetPlatforms, CodeTargetPlatforms));
 
 				// Generate a list of files that needs to be copied for each target platform
 				bp.AddNode(new FilterRocketNode(bp, HostPlatform, TargetPlatforms, CodeTargetPlatforms, CurrentFeaturePacks, CurrentTemplates));
@@ -135,33 +135,55 @@ namespace Rocket
 		public static List<UnrealTargetPlatform> GetTargetPlatforms(GUBP bp, UnrealTargetPlatform HostPlatform)
 		{
 			List<UnrealTargetPlatform> TargetPlatforms = new List<UnrealTargetPlatform>();
-			bool bSupportsMobilePlatforms = (HostPlatform == UnrealTargetPlatform.Win64 || HostPlatform == UnrealTargetPlatform.Mac);
-			if(!bp.ParseParam("NoPlatforms"))
+			if(!bp.ParseParam("NoTargetPlatforms"))
 			{
+				// Always support the host platform
 				TargetPlatforms.Add(HostPlatform);
 
+				// Add other target platforms for each host platform
 				if(HostPlatform == UnrealTargetPlatform.Win64)
 				{
 					TargetPlatforms.Add(UnrealTargetPlatform.Win32);
 				}
-				if(!bp.ParseParam("NoAndroid") && bSupportsMobilePlatforms)
+				if(HostPlatform == UnrealTargetPlatform.Win64 || HostPlatform == UnrealTargetPlatform.Mac)
 				{
 					TargetPlatforms.Add(UnrealTargetPlatform.Android);
 				}
-				if(!bp.ParseParam("NoIOS") && bSupportsMobilePlatforms)
+				if(HostPlatform == UnrealTargetPlatform.Win64 || HostPlatform == UnrealTargetPlatform.Mac)
 				{
 					TargetPlatforms.Add(UnrealTargetPlatform.IOS);
 				}
-				if(!bp.ParseParam("NoLinux") && HostPlatform == UnrealTargetPlatform.Win64)
+				if(HostPlatform == UnrealTargetPlatform.Win64)
 				{
 					TargetPlatforms.Add(UnrealTargetPlatform.Linux);
 				}
-				if(!bp.ParseParam("NoHTML5") && HostPlatform == UnrealTargetPlatform.Win64)
+				if(HostPlatform == UnrealTargetPlatform.Win64)
 				{
 					TargetPlatforms.Add(UnrealTargetPlatform.HTML5);
 				}
 
+				// Remove any platforms that aren't available on this machine
 				TargetPlatforms.RemoveAll(x => !bp.ActivePlatforms.Contains(x));
+
+				// Remove any platforms that aren't enabled on the command line
+				string TargetPlatformFilter = bp.ParseParamValue("TargetPlatforms", null);
+				if(TargetPlatformFilter != null)
+				{
+					List<UnrealTargetPlatform> NewTargetPlatforms = new List<UnrealTargetPlatform>();
+					foreach (string TargetPlatformName in TargetPlatformFilter.Split(new char[]{ '+' }, StringSplitOptions.RemoveEmptyEntries))
+					{
+						UnrealTargetPlatform TargetPlatform;
+						if(!Enum.TryParse(TargetPlatformName, out TargetPlatform))
+						{
+							throw new AutomationException("Unknown target platform '{0}' specified on command line");
+						}
+						else if(TargetPlatforms.Contains(TargetPlatform))
+						{
+							NewTargetPlatforms.Add(TargetPlatform);
+						}
+					}
+					TargetPlatforms = NewTargetPlatforms;
+				}
 			}
 			return TargetPlatforms;
 		}
@@ -272,7 +294,7 @@ namespace Rocket
 		List<UnrealTargetPlatform> Platforms = new List<UnrealTargetPlatform>();
 		public readonly string StrippedDir;
 
-		public StripRocketNode(GUBP bp, UnrealTargetPlatform InHostPlatform, IEnumerable<UnrealTargetPlatform> InTargetPlatforms) : base(InHostPlatform)
+		public StripRocketNode(GUBP bp, UnrealTargetPlatform InHostPlatform, IEnumerable<UnrealTargetPlatform> InTargetPlatforms, IEnumerable<UnrealTargetPlatform> InCodeTargetPlatforms) : base(InHostPlatform)
 		{
 			Platforms = new List<UnrealTargetPlatform>(InTargetPlatforms);
 			if(!Platforms.Contains(HostPlatform))
@@ -292,7 +314,7 @@ namespace Rocket
 				UnrealTargetPlatform SourceHostPlatform = RocketBuild.GetSourceHostPlatform(bp, HostPlatform, TargetPlatform);
 				if(SourceHostPlatform == HostPlatform)
 				{
-					AddDependency(GUBP.GamePlatformMonolithicsNode.StaticGetFullName(HostPlatform, bp.Branch.BaseEngineProject, TargetPlatform, Precompiled: true));
+					AddDependency(GUBP.GamePlatformMonolithicsNode.StaticGetFullName(HostPlatform, bp.Branch.BaseEngineProject, TargetPlatform, Precompiled: InCodeTargetPlatforms.Contains(TargetPlatform)));
 				}
 			}
 		}
