Improve performance for building Rocket DDC: DDPs for individual projects are generated separately, skipping engine content, and separate DDP files are combined without decompressing/recompressing.

diff --git a/Engine/Build/InstalledEngineFilters.ini b/Engine/Build/InstalledEngineFilters.ini
--- a/Engine/Build/InstalledEngineFilters.ini
+++ b/Engine/Build/InstalledEngineFilters.ini
@@ -3,6 +3,10 @@
 /Engine/Binaries/DotNET/Ionic.Zip.Reduced.dll
 /Engine/Binaries/DotNET/OneSky.dll
 
+; DDC
+-/Engine/DerivedDataCache/...
+/Engine/DerivedDataCache/Compressed.ddp
+
 ; In-editor documentation
 /Engine/Documentation/Source/Shared/...
 
@@ -65,7 +69,7 @@
 {Win64} /Engine/Build/BatchFiles/RunUAT.bat
 {Win64} /Engine/Build/BatchFiles/GetVSComnToolsPath.bat
 
-{Win64} /Engine/Extras/3dsMaxScripts/*.ms
+{Win64} /Engine/Extras/3dsMaxScripts/...
 {Win64} /Engine/Extras/VisualStudioDebugging/UE4.natvis
 {Win64} /Engine/Extras/Maya_AnimationRiggingTools/...
 {Win64} /Engine/Extras/MayaVelocityGridExporter/...
diff --git a/Engine/Config/BaseEngine.ini b/Engine/Config/BaseEngine.ini
--- a/Engine/Config/BaseEngine.ini
+++ b/Engine/Config/BaseEngine.ini
@@ -1893,12 +1893,11 @@ PakWrite=(Type=WritePak, Filename="%GAMEDIR%DerivedDataCache/Compressed.ddp", Co
 MinimumDaysToKeepFile=7
 Root=(Type=KeyLength, Length=120, Inner=AsyncPut)
 AsyncPut=(Type=AsyncPut, Inner=Hierarchy)
-Hierarchy=(Type=Hierarchical, Inner=Boot, Inner=PakWrite, Inner=PakRead, Inner=Local, Inner=Shared)
+Hierarchy=(Type=Hierarchical, Inner=Boot, Inner=PakWrite, Inner=Local, Inner=Shared)
 Boot=(Type=Boot, Filename="%GAMEDIR%DerivedDataCache/Boot.ddc", MaxCacheSize=512)
 Local=(Type=FileSystem, ReadOnly=false, Clean=false, Flush=false, PurgeTransient=true, DeleteUnused=true, UnusedFileAge=34, FoldersToClean=-1, Path=../../../Engine/DerivedDataCache, EnvPathOverride=UE-LocalDataCachePath)
 Shared=(Type=FileSystem, ReadOnly=false, Clean=false, Flush=false, DeleteUnused=false, UnusedFileAge=23, FoldersToClean=-1, Path=?EpicDDC, EnvPathOverride=UE-SharedDataCachePath)
 AltShared=(Type=FileSystem, ReadOnly=true, Clean=false, Flush=false, DeleteUnused=false, UnusedFileAge=23, FoldersToClean=-1, Path=?EpicDDC2, EnvPathOverride=UE-SharedDataCachePath2)
-PakRead=(Type=ReadPak, Filename=../../../Engine/DerivedDataCache/Compressed.ddp, Compressed=true)
 PakWrite=(Type=WritePak, Filename=../../../Engine/DerivedDataCache/Compressed.ddp, Compressed=true)
 
 [CreateProjectCache]
diff --git a/Engine/Source/Developer/DerivedDataCache/Private/DerivedDataBackends.cpp b/Engine/Source/Developer/DerivedDataCache/Private/DerivedDataBackends.cpp
--- a/Engine/Source/Developer/DerivedDataCache/Private/DerivedDataBackends.cpp
+++ b/Engine/Source/Developer/DerivedDataCache/Private/DerivedDataBackends.cpp
@@ -675,11 +675,16 @@ public:
 		}
 		if (bShutdown)
 		{
-			if(FParse::Param(FCommandLine::Get(), TEXT("MergePaks")) && WritePakCache && WritePakCache->IsWritable())
+			FString MergePaks;
+			if(WritePakCache && WritePakCache->IsWritable() && FParse::Value(FCommandLine::Get(), TEXT("MergePaks="), MergePaks))
 			{
-				for (int32 ReadPakIndex = 0; ReadPakIndex < ReadPakCache.Num(); ReadPakIndex++)
+				TArray<FString> MergePakList;
+				MergePaks.FString::ParseIntoArray(MergePakList, TEXT("+"));
+
+				for(const FString& MergePakName : MergePakList)
 				{
-					WritePakCache->MergeCache(ReadPakCache[ReadPakIndex]);
+					FPakFileDerivedDataBackend ReadPak(*FPaths::Combine(*FPaths::GetPath(WritePakFilename), *MergePakName), false);
+					WritePakCache->MergeCache(&ReadPak);
 				}
 			}
 			for (int32 ReadPakIndex = 0; ReadPakIndex < ReadPakCache.Num(); ReadPakIndex++)
diff --git a/Engine/Source/Developer/DerivedDataCache/Private/PakFileDerivedDataBackend.h b/Engine/Source/Developer/DerivedDataCache/Private/PakFileDerivedDataBackend.h
--- a/Engine/Source/Developer/DerivedDataCache/Private/PakFileDerivedDataBackend.h
+++ b/Engine/Source/Developer/DerivedDataCache/Private/PakFileDerivedDataBackend.h
@@ -359,20 +359,26 @@ public:
 		// Get all the existing keys
 		TArray<FString> KeyNames;
 		OtherPak->CacheItems.GenerateKeyArray(KeyNames);
-		UE_LOG(LogDerivedDataCache, Display, TEXT("Found %d entries."), KeyNames.Num());
 
-		// Copy them all to the new cache
+		// Find all the keys to copy
+		TArray<FString> CopyKeyNames;
+		for(const FString& KeyName : KeyNames)
+		{
+			if(!CachedDataProbablyExists(*KeyName))
+			{
+				CopyKeyNames.Add(KeyName);
+			}
+		}
+		UE_LOG(LogDerivedDataCache, Display, TEXT("Merging %d entries (%d skipped)."), CopyKeyNames.Num(), KeyNames.Num() - CopyKeyNames.Num());
+
+		// Copy them all to the new cache. Don't use the overloaded get/put methods (which may compress/decompress); copy the raw data directly.
 		TArray<uint8> Buffer;
-		for(int32 Idx = 0; Idx < KeyNames.Num(); Idx++)
+		for(const FString& CopyKeyName : CopyKeyNames)
 		{
-			if(!CachedDataProbablyExists(*KeyNames[Idx]))
+			Buffer.Reset();
+			if(OtherPak->FPakFileDerivedDataBackend::GetCachedData(*CopyKeyName, Buffer))
 			{
-				Buffer.Reset();
-				if(OtherPak->GetCachedData(*KeyNames[Idx], Buffer))
-				{
-					UE_LOG(LogDerivedDataCache, Display, TEXT("[%d/%d] Copying %s (%d bytes)..."), Idx + 1, KeyNames.Num(), *KeyNames[Idx], Buffer.Num());
-					PutCachedData(*KeyNames[Idx], Buffer, false);
-				}
+				FPakFileDerivedDataBackend::PutCachedData(*CopyKeyName, Buffer, false);
 			}
 		}
 	}
@@ -404,6 +410,14 @@ public:
 			InputPak.GetCachedData(*KeyNames[KeyIndex], Buffer);
 			OutputPak.PutCachedData(*KeyNames[KeyIndex], Buffer, false);
 		}
+
+		// Write out a TOC listing for debugging
+		FStringOutputDevice Output;
+		for(int KeyIndex = 0; KeyIndex < KeyNames.Num(); KeyIndex++)
+		{
+			Output.Logf(TEXT("%s") LINE_TERMINATOR, *KeyNames[KeyIndex]);
+		}
+		FFileHelper::SaveStringToFile(Output, *FPaths::Combine(*FPaths::GetPath(OutputFilename), *(FPaths::GetBaseFilename(OutputFilename) + TEXT(".txt"))));
 		return true;
 	}
 
diff --git a/Engine/Source/Programs/AutomationTool/CommandletUtils.cs b/Engine/Source/Programs/AutomationTool/CommandletUtils.cs
--- a/Engine/Source/Programs/AutomationTool/CommandletUtils.cs
+++ b/Engine/Source/Programs/AutomationTool/CommandletUtils.cs
@@ -166,7 +166,7 @@ namespace AutomationTool
 			string LocalLogFile = LogUtils.GetUniqueLogName(CombinePaths(CmdEnv.EngineSavedFolder, Commandlet));
 			Log("Commandlet log file is {0}", LocalLogFile);
 			var RunResult = Run(EditorExe, String.Format("{0} -run={1} {2} -abslog={3} -stdout -FORCELOGFLUSH -CrashForUAT -unattended -AllowStdOutLogVerbosity {4}", 
-                CommandUtils.MakePathSafeToUseWithCommandLine(ProjectName), 
+                (ProjectName == null)? "" : CommandUtils.MakePathSafeToUseWithCommandLine(ProjectName), 
                 Commandlet, 
 				String.IsNullOrEmpty(Parameters) ? "" : Parameters, 
 				CommandUtils.MakePathSafeToUseWithCommandLine(LocalLogFile),
diff --git a/Engine/Source/Programs/AutomationTool/Scripts/RocketBuild.Automation.cs b/Engine/Source/Programs/AutomationTool/Scripts/RocketBuild.Automation.cs
--- a/Engine/Source/Programs/AutomationTool/Scripts/RocketBuild.Automation.cs
+++ b/Engine/Source/Programs/AutomationTool/Scripts/RocketBuild.Automation.cs
@@ -740,21 +740,43 @@ namespace Rocket
 			BuildProducts = new List<string>();
 			if(!bp.ParseParam("NoDDC") && bp.ParseParam("WithDDC"))
 			{
+				string EditorExe = CommandUtils.GetEditorCommandletExe(CommandUtils.CmdEnv.LocalRoot, HostPlatform);
+
 				// Delete the output file
-				string OutputFile = CommandUtils.CombinePaths(CommandUtils.CmdEnv.LocalRoot, "Engine", "DerivedDataCache", "Compressed.ddp");
-				CommandUtils.DeleteFile(OutputFile);
+				string OutputPakFile = CommandUtils.CombinePaths(CommandUtils.CmdEnv.LocalRoot, "Engine", "DerivedDataCache", "Compressed.ddp");
+				string OutputTxtFile = Path.ChangeExtension(OutputPakFile, ".txt");
 
 				// Generate DDC for all the non-code projects. We don't necessarily have editor DLLs for the code projects, but they should be the same as their blueprint counterparts.
+				List<string> ProjectPakFiles = new List<string>();
 				foreach(string ProjectName in ProjectNames)
 				{
 					BranchInfo.BranchUProject Project = bp.Branch.FindGameChecked(ProjectName);
 					if(!Project.Properties.bIsCodeBasedProject)
 					{
-						CommandUtils.Log("Running RocketEditor to generate engine DDC data for {0} on {1}", Project.GameName, TargetPlatforms);
-						CommandUtils.DDCCommandlet(Project.FilePath, CommandUtils.GetEditorCommandletExe(CommandUtils.CmdEnv.LocalRoot, HostPlatform), null, TargetPlatforms, "-fill -DDC=CreateInstalledEnginePak -MergePaks");
+						CommandUtils.Log("Generating DDC data for {0} on {1}", Project.GameName, TargetPlatforms);
+						CommandUtils.DDCCommandlet(Project.FilePath, EditorExe, null, TargetPlatforms, "-fill -DDC=CreateInstalledEnginePak -ProjectOnly");
+
+						string ProjectPakFile = CommandUtils.CombinePaths(Path.GetDirectoryName(OutputPakFile), String.Format("Compressed-{0}.ddp", ProjectName));
+						CommandUtils.DeleteFile(ProjectPakFile);
+						CommandUtils.RenameFile(OutputPakFile, ProjectPakFile);
+						BuildProducts.Add(ProjectPakFile);
+
+						string ProjectTxtFile = Path.ChangeExtension(ProjectPakFile, ".txt");
+						CommandUtils.DeleteFile(ProjectTxtFile);
+						CommandUtils.RenameFile(OutputTxtFile, ProjectTxtFile);
+						BuildProducts.Add(ProjectTxtFile);
+
+						ProjectPakFiles.Add(Path.GetFileName(ProjectPakFile));
 					}
 				}
-				BuildProducts.Add(OutputFile);
+
+				// Generate DDC for the editor, and merge all the other PAK files in
+				CommandUtils.Log("Generating DDC data for engine content on {0}", TargetPlatforms);
+				CommandUtils.DDCCommandlet(null, EditorExe, null, TargetPlatforms, "-fill -DDC=CreateInstalledEnginePak " + CommandUtils.MakePathSafeToUseWithCommandLine("-MergePaks=" + String.Join("+", ProjectPakFiles)));
+
+				// Add the final PAK file as output
+				BuildProducts.Add(OutputPakFile);
+				BuildProducts.Add(OutputTxtFile);
 			}
 			SaveRecordOfSuccessAndAddToBuildProducts();
 		}
