Remove redundant configs from monolithic nodes; GUBP is filtering these out of the build products because they're already exist, but they don't need to be built again.

diff --git a/Engine/Source/Programs/AutomationTool/Scripts/GUBP.Automation.cs b/Engine/Source/Programs/AutomationTool/Scripts/GUBP.Automation.cs
--- a/Engine/Source/Programs/AutomationTool/Scripts/GUBP.Automation.cs
+++ b/Engine/Source/Programs/AutomationTool/Scripts/GUBP.Automation.cs
@@ -1765,7 +1765,7 @@ public class GUBP : BuildCommand
 							}
 							else
 							{
-								Configs = Target.Rules.GUBP_GetConfigs_MonolithicOnly(HostPlatform, TargetPlatform);
+								Configs = Target.Rules.GUBP_GetConfigs_MonolithicOnly(HostPlatform, TargetPlatform).Except(Target.Rules.GUBP_GetConfigsForPrecompiledBuilds_MonolithicOnly(HostPlatform, TargetPlatform)).ToList();
 							}
 							
 							foreach (var Config in Configs)
