Re-enable DDC build as part of CIS.

diff --git a/Engine/Source/Programs/AutomationTool/Scripts/RocketBuild.Automation.cs b/Engine/Source/Programs/AutomationTool/Scripts/RocketBuild.Automation.cs
--- a/Engine/Source/Programs/AutomationTool/Scripts/RocketBuild.Automation.cs
+++ b/Engine/Source/Programs/AutomationTool/Scripts/RocketBuild.Automation.cs
@@ -738,7 +738,7 @@ namespace Rocket
 		public override void DoBuild(GUBP bp)
 		{
 			BuildProducts = new List<string>();
-			if(!bp.ParseParam("NoDDC") && bp.ParseParam("WithDDC"))
+			if(!bp.ParseParam("NoDDC"))
 			{
 				string EditorExe = CommandUtils.GetEditorCommandletExe(CommandUtils.CmdEnv.LocalRoot, HostPlatform);
 
