Add code to fixup the promoted label view if QA use the recovery script that adds paths which are only visible to the buildmachine

diff --git a/Engine/Source/Programs/AutomationTool/P4Utils.cs b/Engine/Source/Programs/AutomationTool/P4Utils.cs
--- a/Engine/Source/Programs/AutomationTool/P4Utils.cs
+++ b/Engine/Source/Programs/AutomationTool/P4Utils.cs
@@ -258,6 +258,10 @@ namespace AutomationTool
 			string[] Lines = Text.Split('\n');
 			for(int LineIdx = 0; LineIdx < Lines.Length; LineIdx++)
 			{
+				if(Lines[LineIdx].EndsWith("\r"))
+				{
+					Lines[LineIdx] = Lines[LineIdx].Substring(0, Lines[LineIdx].Length - 1);
+				}
 				if(!String.IsNullOrWhiteSpace(Lines[LineIdx]) && !Lines[LineIdx].StartsWith("#"))
 				{
 					// Read the section name
@@ -1883,6 +1887,31 @@ namespace AutomationTool
 			return false;
 		}
 
+		/// <summary>
+		/// Reads a label spec
+		/// </summary>
+		/// <param name="Name">Label name</param>
+		/// <param name="AllowSpew">Whether to allow log spew</param>
+		public P4Spec ReadLabelSpec(string Name, bool AllowSpew = true)
+		{
+			string LabelSpec;
+			if(!LogP4Output(out LabelSpec, "label -o " + Name, AllowSpew: AllowSpew))
+			{
+				throw new P4Exception("Couldn't describe existing label '{0}', output was:\n", Name, LabelSpec);
+			}
+			return P4Spec.FromString(LabelSpec);
+		}
+
+		/// <summary>
+		/// Updates a label with a new spec
+		/// </summary>
+		/// <param name="Spec">Label specification</param>
+		/// <param name="AllowSpew">Whether to allow log spew</param>
+		public void UpdateLabelSpec(P4Spec Spec, bool AllowSpew = true)
+		{
+			LogP4("label -i", Input: Spec.ToString(), AllowSpew: AllowSpew);
+		}
+
 		/// <summary>
 		/// Updates a label description.
 		/// </summary>
