Changes to GUBP for native Linux host.

diff --git a/Engine/Source/Programs/AutomationTool/MCPPublic.cs b/Engine/Source/Programs/AutomationTool/MCPPublic.cs
--- a/Engine/Source/Programs/AutomationTool/MCPPublic.cs
+++ b/Engine/Source/Programs/AutomationTool/MCPPublic.cs
@@ -72,9 +72,14 @@ namespace EpicGames.MCP.Automation
         Windows,
 
         /// <summary>
-        /// Only other platform MCP understands is Mac.
+        /// Mac platform.
         /// </summary>
         Mac,
+
+		/// <summary>
+		/// Linux platform.
+		/// </summary>
+		Linux
     }
 
     /// <summary>
@@ -155,22 +160,42 @@ namespace EpicGames.MCP.Automation
         /// </summary>
         static public MCPPlatform ToMCPPlatform(UnrealTargetPlatform TargetPlatform)
         {
-            if (TargetPlatform != UnrealTargetPlatform.Win64 && TargetPlatform != UnrealTargetPlatform.Win32 && TargetPlatform != UnrealTargetPlatform.Mac)
+            if (TargetPlatform != UnrealTargetPlatform.Win64 && TargetPlatform != UnrealTargetPlatform.Win32 && TargetPlatform != UnrealTargetPlatform.Mac && TargetPlatform != UnrealTargetPlatform.Linux)
             {
                 throw new AutomationException("Platform {0} is not properly supported by the MCP backend yet", TargetPlatform);
             }
-            return (TargetPlatform == UnrealTargetPlatform.Win64 || TargetPlatform == UnrealTargetPlatform.Win32) ? MCPPlatform.Windows : MCPPlatform.Mac;
+
+			if (TargetPlatform == UnrealTargetPlatform.Win64 || TargetPlatform == UnrealTargetPlatform.Win32)
+			{
+				return MCPPlatform.Windows;
+			}
+			else if (TargetPlatform == UnrealTargetPlatform.Mac)
+			{
+				return MCPPlatform.Mac;
+			}
+
+			return MCPPlatform.Linux;
         }
         /// <summary>
         /// Determine the platform name (Win32/64 becomes Windows, Mac is Mac, the rest we don't currently understand)
         /// </summary>
         static public UnrealTargetPlatform FromMCPPlatform(MCPPlatform TargetPlatform)
         {
-            if (TargetPlatform != MCPPlatform.Windows && TargetPlatform != MCPPlatform.Mac)
+            if (TargetPlatform != MCPPlatform.Windows && TargetPlatform != MCPPlatform.Mac && TargetPlatform != MCPPlatform.Linux)
             {
                 throw new AutomationException("Platform {0} is not properly supported by the MCP backend yet", TargetPlatform);
             }
-            return (TargetPlatform == MCPPlatform.Windows) ? UnrealTargetPlatform.Win64 : UnrealTargetPlatform.Mac;
+
+			if (TargetPlatform == MCPPlatform.Windows)
+			{
+				return UnrealTargetPlatform.Win64;
+			}
+			else if (TargetPlatform == MCPPlatform.Mac)
+			{
+				return UnrealTargetPlatform.Mac;
+			}
+
+			return UnrealTargetPlatform.Linux;
         }
         /// <summary>
         /// Returns the build root path (P:\Builds on build machines usually)
diff --git a/Engine/Source/Programs/AutomationTool/Scripts/GUBP.Automation.cs b/Engine/Source/Programs/AutomationTool/Scripts/GUBP.Automation.cs
--- a/Engine/Source/Programs/AutomationTool/Scripts/GUBP.Automation.cs
+++ b/Engine/Source/Programs/AutomationTool/Scripts/GUBP.Automation.cs
@@ -1499,6 +1499,11 @@ public class GUBP : BuildCommand
 			}
 			else
 			{
+				if (UnrealBuildTool.BuildHostPlatform.Current.Platform == UnrealTargetPlatform.Linux && bp.HostPlatforms[0] != UnrealTargetPlatform.Linux)
+				{
+					throw new AutomationException("Linux is not (yet?) able to cross-compile nodes for platform {0}, did you forget -NoPC / -NoMac?", bp.HostPlatforms[0]);
+				}
+
 				return bp.HostPlatforms[0];
 			}
 		}
