Build feature packs in the tools agent sharing group; it only takes a minute or two to run, and doesn't have any other dependencies.

diff --git a/Engine/Source/Programs/AutomationTool/Scripts/GUBP.Automation.cs b/Engine/Source/Programs/AutomationTool/Scripts/GUBP.Automation.cs
--- a/Engine/Source/Programs/AutomationTool/Scripts/GUBP.Automation.cs
+++ b/Engine/Source/Programs/AutomationTool/Scripts/GUBP.Automation.cs
@@ -1466,6 +1466,7 @@ public class GUBP : BuildCommand
         {
 			Projects = new List<BranchInfo.BranchUProject>(InProjects);
 			AddDependency(ToolsNode.StaticGetFullName(InHostPlatform)); // for UnrealPak
+			AgentSharingGroup = "ToolsGroup" + StaticGetHostPlatformSuffix(HostPlatform);
         }
 
 		public static string GetOutputFile(BranchInfo.BranchUProject Project)
