build script, disassociate possible causers from node owners

diff --git a/Engine/Source/Programs/AutomationTool/Scripts/GUBP.Automation.cs b/Engine/Source/Programs/AutomationTool/Scripts/GUBP.Automation.cs
--- a/Engine/Source/Programs/AutomationTool/Scripts/GUBP.Automation.cs
+++ b/Engine/Source/Programs/AutomationTool/Scripts/GUBP.Automation.cs
@@ -4870,7 +4870,6 @@ public class GUBP : BuildCommand
 	void GetFailureEmails(string NodeToDo, string CLString, bool OnlyLateUpdates = false)
 	{
 		var StartTime = DateTime.UtcNow;        
-        string EMails = "";
         string FailCauserEMails = "";
         string EMailNote = "";
         bool SendSuccessForGreenAfterRed = false;
@@ -4948,21 +4947,12 @@ public class GUBP : BuildCommand
 		RunECTool(String.Format("setProperty \"/myWorkflow/FailCausers/{0}\" \"{1}\"", NodeToDo, FailCauserEMails));
 		RunECTool(String.Format("setProperty \"/myWorkflow/EmailNotes/{0}\" \"{1}\"", NodeToDo, EMailNote));
         {
-            var AdditonalEmails = "";
-
 			string Causers = "";
             if (ParseParam("CIS") && !GUBPNodes[NodeToDo].SendSuccessEmail() && !GUBPNodes[NodeToDo].TriggerNode())
             {
 				Causers = FailCauserEMails;
-           }
-            string AddEmails = ParseParamValue("AddEmails");
-            if (!String.IsNullOrEmpty(AddEmails))
-            {
-                AdditonalEmails = GUBPNode.MergeSpaceStrings(AddEmails, AdditonalEmails);
             }
-
-            EMails = GetEMailListForNode(this, NodeToDo, AdditonalEmails, Causers);
-			RunECTool(String.Format("setProperty \"/myWorkflow/FailEmails/{0}\" \"{1}\"", NodeToDo, EMails));            
+			RunECTool(String.Format("setProperty \"/myWorkflow/FailEmails/{0}\" \"{1}\"", NodeToDo, Causers));            
         }
 		if (GUBPNodes[NodeToDo].SendSuccessEmail() || SendSuccessForGreenAfterRed)
 		{
@@ -5071,7 +5061,7 @@ public class GUBP : BuildCommand
 			AdditonalEmails = GUBPNode.MergeSpaceStrings(AddEmails, AdditonalEmails);
 		}
 		EMails = GetEMailListForNode(this, NodeToDo, AdditonalEmails, Causers);
-		ECProps.Add("FailEmails/" + NodeToDo + "=" + EMails);
+		ECProps.Add("EmailOwners/" + NodeToDo + "=" + EMails);
 	
 		if (!OnlyLateUpdates)
 		{
