Put Rocket build-finalizing nodes into the same ASG.

diff --git a/Engine/Source/Programs/AutomationTool/Scripts/RocketBuild.Automation.cs b/Engine/Source/Programs/AutomationTool/Scripts/RocketBuild.Automation.cs
--- a/Engine/Source/Programs/AutomationTool/Scripts/RocketBuild.Automation.cs
+++ b/Engine/Source/Programs/AutomationTool/Scripts/RocketBuild.Automation.cs
@@ -261,6 +261,8 @@ namespace Rocket
 			{ 
 				AddDependency(GUBP.SingleInternalToolsNode.StaticGetFullName(UnrealTargetPlatform.Win64, VersionSelector));
 			}
+
+			AgentSharingGroup = "BuildRocket" + StaticGetHostPlatformSuffix(HostPlatform);
 		}
 
 		public static string StaticGetFullName(UnrealTargetPlatform HostPlatform)
@@ -317,6 +319,8 @@ namespace Rocket
 					AddDependency(GUBP.GamePlatformMonolithicsNode.StaticGetFullName(HostPlatform, bp.Branch.BaseEngineProject, TargetPlatform, Precompiled: InCodeTargetPlatforms.Contains(TargetPlatform)));
 				}
 			}
+
+			AgentSharingGroup = "BuildRocket" + StaticGetHostPlatformSuffix(HostPlatform);
 		}
 
 		public static string StaticGetFullName(UnrealTargetPlatform HostPlatform)
@@ -443,6 +447,8 @@ namespace Rocket
 			{
 				AddDependency(StripRocketNode.StaticGetFullName(SourceHostPlatform));
 			}
+
+			AgentSharingGroup = "BuildRocket" + StaticGetHostPlatformSuffix(HostPlatform);
 		}
 
 		public override int CISFrequencyQuantumShift(GUBP bp)
