build script, rollback 2514786 while I fix the javascript

diff --git a/Engine/Source/Programs/AutomationTool/Scripts/GUBP.Automation.cs b/Engine/Source/Programs/AutomationTool/Scripts/GUBP.Automation.cs
--- a/Engine/Source/Programs/AutomationTool/Scripts/GUBP.Automation.cs
+++ b/Engine/Source/Programs/AutomationTool/Scripts/GUBP.Automation.cs
@@ -4869,7 +4869,8 @@ public class GUBP : BuildCommand
     }
 	void GetFailureEmails(string NodeToDo, string CLString, bool OnlyLateUpdates = false)
 	{
-		var StartTime = DateTime.UtcNow;        
+		var StartTime = DateTime.UtcNow;
+        string EMails;
         string FailCauserEMails = "";
         string EMailNote = "";
         bool SendSuccessForGreenAfterRed = false;
@@ -4947,12 +4948,19 @@ public class GUBP : BuildCommand
 		RunECTool(String.Format("setProperty \"/myWorkflow/FailCausers/{0}\" \"{1}\"", NodeToDo, FailCauserEMails));
 		RunECTool(String.Format("setProperty \"/myWorkflow/EmailNotes/{0}\" \"{1}\"", NodeToDo, EMailNote));
         {
+            var AdditionalEmails = "";
 			string Causers = "";
             if (ParseParam("CIS") && !GUBPNodes[NodeToDo].SendSuccessEmail() && !GUBPNodes[NodeToDo].TriggerNode())
             {
 				Causers = FailCauserEMails;
             }
-			RunECTool(String.Format("setProperty \"/myWorkflow/FailEmails/{0}\" \"{1}\"", NodeToDo, Causers));            
+            string AddEmails = ParseParamValue("AddEmails");
+            if (!String.IsNullOrEmpty(AddEmails))
+            {
+                AdditionalEmails = GUBPNode.MergeSpaceStrings(AddEmails, AdditionalEmails);
+            }
+			EMails = GetEMailListForNode(this, NodeToDo, AdditionalEmails, Causers);
+            RunECTool(String.Format("setProperty \"/myWorkflow/FailEmails/{0}\" \"{1}\"", NodeToDo, EMails));            
         }
 		if (GUBPNodes[NodeToDo].SendSuccessEmail() || SendSuccessForGreenAfterRed)
 		{
@@ -5061,7 +5069,7 @@ public class GUBP : BuildCommand
 			AdditonalEmails = GUBPNode.MergeSpaceStrings(AddEmails, AdditonalEmails);
 		}
 		EMails = GetEMailListForNode(this, NodeToDo, AdditonalEmails, Causers);
-		ECProps.Add("EmailOwners/" + NodeToDo + "=" + EMails);
+		ECProps.Add("FailEmails/" + NodeToDo + "=" + EMails);
 	
 		if (!OnlyLateUpdates)
 		{
