UAT: Add support to cloud storage for:
  * Setting and getting metadata on files
  * Touching a file, to update its modified date to now
  * Listing a directory non-recursively

diff --git a/Engine/Source/Programs/AutomationTool/MCPPublic.cs b/Engine/Source/Programs/AutomationTool/MCPPublic.cs
--- a/Engine/Source/Programs/AutomationTool/MCPPublic.cs
+++ b/Engine/Source/Programs/AutomationTool/MCPPublic.cs
@@ -740,8 +740,40 @@ namespace EpicGames.MCP.Automation
 		/// </summary>
 		/// <param name="Container">The name of the folder or container from which to list files.</param>
 		/// <param name="Prefix">A string with which the identifier or filename should start. Typically used to specify a relative directory within the container to list all of its files recursively. Specify null to return all files.</param>
+		/// <param name="Recursive">Indicates whether the list of files returned should traverse subdirectories</param>
 		/// <returns>An array of paths to the files in the specified location and matching the prefix constraint.</returns>
-		abstract public string[] ListFiles(string Container, string Prefix = null);
+		abstract public string[] ListFiles(string Container, string Prefix = null, bool bRecursive = true);
+
+		/// <summary>
+		/// Sets one or more items of metadata on an object in cloud storage
+		/// </summary>
+		/// <param name="Container">The name of the folder or container in which the file is stored.</param>
+		/// <param name="Identifier">The identifier of filename of the file to set metadata on.</param>
+		/// <param name="Metadata">A dictionary containing the metadata keys and their values</param>
+		/// <param name="bMerge">If true, then existing metadata will be replaced (or overwritten if the keys match). If false, no existing metadata is retained.</param>
+		abstract public void SetMetadata(string Container, string Identifier, IDictionary<string, object> Metadata, bool bMerge = true);
+
+		/// <summary>
+		/// Gets all items of metadata on an object in cloud storage. Metadata values are all returned as strings.
+		/// </summary>
+		/// <param name="Container">The name of the folder or container in which the file is stored.</param>
+		/// <param name="Identifier">The identifier of filename of the file to get metadata.</param>
+		abstract public Dictionary<string, string> GetMetadata(string Container, string Identifier);
+
+		/// <summary>
+		/// Gets an item of metadata from an object in cloud storage. The object is casted to the specified type.
+		/// </summary>
+		/// <param name="Container">The name of the folder or container in which the file is stored.</param>
+		/// <param name="Identifier">The identifier of filename of the file to get metadata.</param>
+		/// <param name="MetadataKey">The key of the item of metadata to retrieve.</param>
+		abstract public T GetMetadata<T>(string Container, string Identifier, string MetadataKey);
+
+		/// <summary>
+		/// Updates the timestamp on a particular file in cloud storage to the current time.
+		/// </summary>
+		/// <param name="Container">The identifier of filename of the file to touch.</param>
+		/// <param name="Identifier">The identifier of filename of the file to touch.</param>
+		abstract public void TouchFile(string Container, string Identifier);
 
 		/// <summary>
 		/// Copies chunks from a staged location to cloud storage.
