Allow bunching game editors together in a single build node, and use it to group all the templates together.

diff --git a/Engine/Source/Programs/AutomationTool/Scripts/GUBP.Automation.cs b/Engine/Source/Programs/AutomationTool/Scripts/GUBP.Automation.cs
--- a/Engine/Source/Programs/AutomationTool/Scripts/GUBP.Automation.cs
+++ b/Engine/Source/Programs/AutomationTool/Scripts/GUBP.Automation.cs
@@ -1403,55 +1403,47 @@ public class GUBP : BuildCommand
 
     public class EditorGameNode : CompileNode
     {
-        BranchInfo.BranchUProject GameProj;
+		List<BranchInfo.BranchUProject> GameProjects = new List<BranchInfo.BranchUProject>();
 
         public EditorGameNode(GUBP bp, UnrealTargetPlatform InHostPlatform, BranchInfo.BranchUProject InGameProj)
             : base(InHostPlatform)
         {
             AgentSharingGroup = "Editor"  + StaticGetHostPlatformSuffix(InHostPlatform);
-            GameProj = InGameProj;
+            GameProjects.Add(InGameProj);
 			AddDependency(RootEditorNode.StaticGetFullName(InHostPlatform));						
 		}
-		
+		public void AddProject(BranchInfo.BranchUProject InGameProj)
+		{
+			if(InGameProj.Options(HostPlatform).GroupName != GameProjects[0].Options(HostPlatform).GroupName)
+			{
+				throw new AutomationException("Attempt to merge projects with different group names");
+			}
+			GameProjects.Add(InGameProj);
+		}
         public static string StaticGetFullName(UnrealTargetPlatform InHostPlatform, BranchInfo.BranchUProject InGameProj)
         {
-            return InGameProj.GameName + "_EditorGame" + StaticGetHostPlatformSuffix(InHostPlatform);
+            return (InGameProj.Options(InHostPlatform).GroupName ?? InGameProj.GameName) + "_EditorGame" + StaticGetHostPlatformSuffix(InHostPlatform);
         }
         public override string GetFullName()
         {
-            return StaticGetFullName(HostPlatform, GameProj);
+            return StaticGetFullName(HostPlatform, GameProjects[0]);
         }
         public override string GameNameIfAnyForTempStorage()
         {
-            return GameProj.GameName;
+            return GameProjects[0].Options(HostPlatform).GroupName ?? GameProjects[0].GameName;
         }		
-        public override float Priority()
-        {
-            float Result = base.Priority();
-            if (GameProj.Options(HostPlatform).bTestWithShared)
-            {
-                Result -= 1;
-            }
-            return Result;
-        }
-        public override int CISFrequencyQuantumShift(GUBP bp)
-        {
-            int Result = base.CISFrequencyQuantumShift(bp);
-            if(GameProj.Options(HostPlatform).bTestWithShared)
-            {
-                Result += 3;
-            }                    
-            return Result;
-        }
         public override UE4Build.BuildAgenda GetAgenda(GUBP bp)
         {
             var Agenda = new UE4Build.BuildAgenda();
 
             string Args = "-nobuilduht -skipactionhistory -skipnonhostplatforms -CopyAppBundleBackToDevice -forceheadergeneration";
 
-            Agenda.AddTargets(
-                new string[] { GameProj.Properties.Targets[TargetRules.TargetType.Editor].TargetName },
-                HostPlatform, UnrealTargetConfiguration.Development, GameProj.FilePath, InAddArgs: Args);
+			foreach(BranchInfo.BranchUProject GameProj in GameProjects)
+			{
+				Agenda.AddTargets(
+					new string[] { GameProj.Properties.Targets[TargetRules.TargetType.Editor].TargetName },
+					HostPlatform, UnrealTargetConfiguration.Development, GameProj.FilePath, InAddArgs: Args);
+			}
 
             return Agenda;
         }
@@ -3827,6 +3819,12 @@ public class GUBP : BuildCommand
         return GUBPNodes[Node];
     }
 
+	public GUBPNode TryFindNode(string Node)
+	{
+		GUBPNode Result;
+		GUBPNodes.TryGetValue(Node, out Result);
+		return Result;
+	}
 
     public void RemovePseudodependencyFromNode(string Node, string Dep)
     {
@@ -5964,7 +5962,15 @@ public class GUBP : BuildCommand
 
 				if (!BranchOptions.ExcludePlatformsForEditor.Contains(HostPlatform) && !Options.bIsNonCode)
 				{
-					AddNode(new EditorGameNode(this, HostPlatform, CodeProj));
+					EditorGameNode Node = (EditorGameNode)TryFindNode(EditorGameNode.StaticGetFullName(HostPlatform, CodeProj));
+					if(Node == null)
+					{
+						AddNode(new EditorGameNode(this, HostPlatform, CodeProj));
+					}
+					else
+					{
+						Node.AddProject(CodeProj);
+					}
 				}
 
 				if (!bNoAutomatedTesting && HostPlatform == UnrealTargetPlatform.Win64) //temp hack till automated testing works on other platforms than Win64
diff --git a/Engine/Source/Programs/UnrealBuildTool/System/RulesCompiler.cs b/Engine/Source/Programs/UnrealBuildTool/System/RulesCompiler.cs
--- a/Engine/Source/Programs/UnrealBuildTool/System/RulesCompiler.cs
+++ b/Engine/Source/Programs/UnrealBuildTool/System/RulesCompiler.cs
@@ -869,6 +869,7 @@ namespace UnrealBuildTool
             public bool bCustomWorkflowForPromotion = false;
 			public bool bIsNonCode = false;
             public bool bPromoteEditorOnly = true;
+			public string GroupName = null;
 		}
         public virtual GUBPProjectOptions GUBP_IncludeProjectInPromotedBuild_EditorTypeOnly(UnrealTargetPlatform HostPlatform)
         {
@@ -878,6 +879,7 @@ namespace UnrealBuildTool
             if (!String.IsNullOrEmpty(TargetName) && ( TargetName.StartsWith("TP_") || TargetName.StartsWith("FP_")) )
             {
                 Result.bTestWithShared = true;
+				Result.GroupName = "Templates";
             }
             return Result;
         }
