Increase frequency for documentation tools; we don't need them in the promotable.

diff --git a/Engine/Source/Programs/AutomationTool/Scripts/BuildDocumentation.Automation.cs b/Engine/Source/Programs/AutomationTool/Scripts/BuildDocumentation.Automation.cs
--- a/Engine/Source/Programs/AutomationTool/Scripts/BuildDocumentation.Automation.cs
+++ b/Engine/Source/Programs/AutomationTool/Scripts/BuildDocumentation.Automation.cs
@@ -26,6 +26,11 @@ class ToolsForDocumentationNode : GUBP.CompileNode
         return StaticGetFullName(HostPlatform);
     }
 
+	public override int CISFrequencyQuantumShift(GUBP bp)
+	{
+		return base.CISFrequencyQuantumShift(bp) + 3;
+	}
+
 	public override UE4Build.BuildAgenda GetAgenda(GUBP bp)
     {
         var Agenda = new UE4Build.BuildAgenda();
