Explicitly add non-unity compiles for launcher programs. No branch that inherits launcher code as part of the engine wants that stuff, but it's triggered hourly.

diff --git a/Engine/Source/Programs/AutomationTool/Scripts/GUBP.Automation.cs b/Engine/Source/Programs/AutomationTool/Scripts/GUBP.Automation.cs
--- a/Engine/Source/Programs/AutomationTool/Scripts/GUBP.Automation.cs
+++ b/Engine/Source/Programs/AutomationTool/Scripts/GUBP.Automation.cs
@@ -193,6 +193,9 @@ public class GUBP : BuildCommand
         public virtual void ModifyOptions(GUBP bp, ref BranchOptions Options, string Branch)
         {
         }
+		public virtual void ModifyNodes(GUBP bp, string Branch)
+		{
+		}
     }
 
     private static List<GUBPBranchHacker> BranchHackers;
@@ -226,6 +229,14 @@ public class GUBP : BuildCommand
         return Result;
     }
 
+	private void ModifyBranchNodes(string Branch)
+	{
+		foreach(GUBPBranchHacker Hacker in BranchHackers)
+		{
+			Hacker.ModifyNodes(this, Branch);
+		}
+	}
+
     public abstract class GUBPEmailHacker
     {        
         public virtual List<string> AddEmails(GUBP bp, string Branch, string NodeName)
@@ -5190,11 +5201,15 @@ public class GUBP : BuildCommand
         {
             HostPlatforms.Add(UnrealTargetPlatform.Win64);
         }
-        var BranchForOptions = "";
+        string BranchForOptions;
         if (P4Enabled)
         {
             BranchForOptions = P4Env.BuildRootP4;
         }
+		else
+		{ 
+			BranchForOptions = ParseParamValue("Branch", "");
+		}
         BranchOptions = GetBranchOptions(BranchForOptions);
         bool WithMac = !BranchOptions.PlatformsToRemove.Contains(UnrealTargetPlatform.Mac);
         if (ParseParam("NoMac"))
@@ -6251,6 +6266,7 @@ public class GUBP : BuildCommand
                 AddNode(new CleanSharedTempStorageNode(this));
             }
         }
+		ModifyBranchNodes(BranchForOptions);
 #if false
         // this doesn't work for lots of reasons...we can't figure out what the dependencies are until far later
         if (bPreflightBuild)
