Increase frequency for launcher nodes, and add launcher non-unity builds in a different place.

diff --git a/Engine/Source/Programs/AutomationTool/Scripts/GUBP.Automation.cs b/Engine/Source/Programs/AutomationTool/Scripts/GUBP.Automation.cs
--- a/Engine/Source/Programs/AutomationTool/Scripts/GUBP.Automation.cs
+++ b/Engine/Source/Programs/AutomationTool/Scripts/GUBP.Automation.cs
@@ -103,6 +103,7 @@ public class TestECJobErrorParse : BuildCommand
 public class GUBP : BuildCommand
 {
     public string StoreName = null;
+	public string BranchName;
     public int CL = 0;
     public int TimeIndex = 0;
     public bool bSignBuildProducts = false;
@@ -193,9 +194,6 @@ public class GUBP : BuildCommand
         public virtual void ModifyOptions(GUBP bp, ref BranchOptions Options, string Branch)
         {
         }
-		public virtual void ModifyNodes(GUBP bp, string Branch)
-		{
-		}
     }
 
     private static List<GUBPBranchHacker> BranchHackers;
@@ -229,14 +227,6 @@ public class GUBP : BuildCommand
         return Result;
     }
 
-	private void ModifyBranchNodes(string Branch)
-	{
-		foreach(GUBPBranchHacker Hacker in BranchHackers)
-		{
-			Hacker.ModifyNodes(this, Branch);
-		}
-	}
-
     public abstract class GUBPEmailHacker
     {        
         public virtual List<string> AddEmails(GUBP bp, string Branch, string NodeName)
@@ -4140,12 +4130,7 @@ public class GUBP : BuildCommand
     }
     int GetFrequencyForNode(GUBP bp, string NodeToDo, int BaseFrequency)
     {
-        var BranchForFrequency = "";
-        if (P4Enabled)
-        {
-            BranchForFrequency = P4Env.BuildRootP4;
-        }
-        return HackFrequency(bp, BranchForFrequency, NodeToDo, BaseFrequency);
+        return HackFrequency(bp, BranchName, NodeToDo, BaseFrequency);
     }
 
     List<P4Connection.ChangeRecord> GetChanges(int LastOutputForChanges, int TopCL, int LastGreen)
@@ -5201,16 +5186,15 @@ public class GUBP : BuildCommand
         {
             HostPlatforms.Add(UnrealTargetPlatform.Win64);
         }
-        string BranchForOptions;
         if (P4Enabled)
         {
-            BranchForOptions = P4Env.BuildRootP4;
+            BranchName = P4Env.BuildRootP4;
         }
 		else
 		{ 
-			BranchForOptions = ParseParamValue("Branch", "");
+			BranchName = ParseParamValue("BranchName", "");
 		}
-        BranchOptions = GetBranchOptions(BranchForOptions);
+        BranchOptions = GetBranchOptions(BranchName);
         bool WithMac = !BranchOptions.PlatformsToRemove.Contains(UnrealTargetPlatform.Mac);
         if (ParseParam("NoMac"))
         {
@@ -6266,7 +6250,6 @@ public class GUBP : BuildCommand
                 AddNode(new CleanSharedTempStorageNode(this));
             }
         }
-		ModifyBranchNodes(BranchForOptions);
 #if false
         // this doesn't work for lots of reasons...we can't figure out what the dependencies are until far later
         if (bPreflightBuild)
