Add barriers to force frequency for all dependencies of a certain node to run with a given interval. Force promotable builds every 1h20m.

diff --git a/Engine/Source/Programs/AutomationTool/Scripts/GUBP.Automation.cs b/Engine/Source/Programs/AutomationTool/Scripts/GUBP.Automation.cs
--- a/Engine/Source/Programs/AutomationTool/Scripts/GUBP.Automation.cs
+++ b/Engine/Source/Programs/AutomationTool/Scripts/GUBP.Automation.cs
@@ -189,6 +189,7 @@ public class GUBP : BuildCommand
 			public bool bNoDocumentation = false;
 			public bool bNoInstalledEngine = false;
 			public bool bMakeFormalBuildWithoutLabelPromotable = false;
+			public Dictionary<string, sbyte> FrequencyBarriers = new Dictionary<string,sbyte>();
 			public int QuantumOverride = 0;
         }
         public virtual void ModifyOptions(GUBP bp, ref BranchOptions Options, string Branch)
@@ -5193,7 +5194,7 @@ public class GUBP : BuildCommand
 		else
 		{ 
 			BranchName = ParseParamValue("BranchName", "");
-		}
+        }
         BranchOptions = GetBranchOptions(BranchName);
         bool WithMac = !BranchOptions.PlatformsToRemove.Contains(UnrealTargetPlatform.Mac);
         if (ParseParam("NoMac"))
@@ -6295,6 +6296,37 @@ public class GUBP : BuildCommand
             ComputeDependentCISFrequencyQuantumShift(NodeToDo.Key);
         }
 
+		// Make sure that everything that's listed as a frequency barrier is completed with the given interval
+		foreach(KeyValuePair<string, sbyte> Barrier in BranchOptions.FrequencyBarriers)
+		{
+			// All the nodes which are dependencies of the barrier node
+			HashSet<string> IncludedNodes = new HashSet<string>{ Barrier.Key };
+
+			// Find all the nodes which are indirect dependencies of this node
+			List<string> SearchNodes = new List<string>{ Barrier.Key };
+			for(int Idx = 0; Idx < SearchNodes.Count; Idx++)
+			{
+				GUBPNode Node = GUBPNodes[SearchNodes[Idx]];
+				foreach(string DependencyName in Node.FullNamesOfDependencies)
+				{
+					if(!IncludedNodes.Contains(DependencyName))
+					{
+						IncludedNodes.Add(DependencyName);
+						SearchNodes.Add(DependencyName);
+					}
+				}
+			}
+
+			// Make sure that everything included in this list is before the cap, and everything not in the list is after it
+			foreach(KeyValuePair<string, GUBPNode> NodePair in GUBPNodes)
+			{
+				if(IncludedNodes.Contains(NodePair.Key))
+				{
+					NodePair.Value.ComputedDependentCISFrequencyQuantumShift = Math.Min(NodePair.Value.ComputedDependentCISFrequencyQuantumShift, Barrier.Value);
+				}
+			}
+		}
+
         if (bCleanLocalTempStorage)  // shared temp storage can never be wiped
         {
             DeleteLocalTempStorageManifests(CmdEnv);
