Build PS4MapFileUtil and XboxOnePDBFileUtil as part of the regular tools node, so we can take advantage of building in parallel (and reduce costs due to startup times).

diff --git a/Engine/Source/Programs/AutomationTool/Scripts/GUBP.Automation.cs b/Engine/Source/Programs/AutomationTool/Scripts/GUBP.Automation.cs
--- a/Engine/Source/Programs/AutomationTool/Scripts/GUBP.Automation.cs
+++ b/Engine/Source/Programs/AutomationTool/Scripts/GUBP.Automation.cs
@@ -1573,24 +1573,12 @@ public class GUBP : BuildCommand
 			WithXp = InWithXp;
 			Precompiled = InPrecompiled;
 
-            if (TargetPlatform == UnrealTargetPlatform.PS4)
+            if (TargetPlatform == UnrealTargetPlatform.PS4 || TargetPlatform == UnrealTargetPlatform.XboxOne)
             {
-                var PS4MapFileUtil = bp.Branch.FindProgram("PS4MapFileUtil");
-                if(PS4MapFileUtil.Rules == null)
-                {
-                    throw new AutomationException("PS4MapFileUtil is not is this branch, but is required to build PS4 monolithics");
-                }
-                AddDependency(SingleToolsNode.StaticGetFullName(HostPlatform, PS4MapFileUtil));
+				// Required for PS4MapFileUtil/XboxOnePDBFileUtil
+				AddDependency(ToolsNode.StaticGetFullName(InHostPlatform));
             }
-			if (TargetPlatform == UnrealTargetPlatform.XboxOne)
-			{
-				var XboxOnePDBFileUtil = bp.Branch.FindProgram("XboxOnePDBFileUtil");
-				if (XboxOnePDBFileUtil.Rules == null)
-				{
-					throw new AutomationException("XboxOnePDBFileUtil is not is this branch, but is required to build Xbox One monolithics");
-				}
-				AddDependency(SingleToolsNode.StaticGetFullName(HostPlatform, XboxOnePDBFileUtil));
-			}
+
             if (InGameProj.GameName != bp.Branch.BaseEngineProject.GameName && GameProj.Properties.Targets.ContainsKey(TargetRules.TargetType.Editor))
             {
 				if (!bp.BranchOptions.ExcludePlatformsForEditor.Contains(InHostPlatform))
