Put the CopyRocketSymbols node into the same agent sharing group as the CopyRocket node.

diff --git a/Engine/Source/Programs/AutomationTool/Scripts/RocketBuild.Automation.cs b/Engine/Source/Programs/AutomationTool/Scripts/RocketBuild.Automation.cs
--- a/Engine/Source/Programs/AutomationTool/Scripts/RocketBuild.Automation.cs
+++ b/Engine/Source/Programs/AutomationTool/Scripts/RocketBuild.Automation.cs
@@ -658,6 +658,10 @@ namespace Rocket
 			{
 				AddDependency(GUBP.WaitForSharedPromotionUserInput.StaticGetFullName(false));
 			}
+			else
+			{
+				AgentSharingGroup = "CopyRocket" + StaticGetHostPlatformSuffix(HostPlatform);
+			}
 
 			AddDependency(GUBP.ToolsForCompileNode.StaticGetFullName(HostPlatform));
 			AddDependency(GUBP.RootEditorNode.StaticGetFullName(HostPlatform));
