Rename ASG for copying Rocket files. Seems to cause everything in that group to be always built, even if selecting a particular node.

diff --git a/Engine/Source/Programs/AutomationTool/Scripts/RocketBuild.Automation.cs b/Engine/Source/Programs/AutomationTool/Scripts/RocketBuild.Automation.cs
--- a/Engine/Source/Programs/AutomationTool/Scripts/RocketBuild.Automation.cs
+++ b/Engine/Source/Programs/AutomationTool/Scripts/RocketBuild.Automation.cs
@@ -592,7 +592,7 @@ namespace Rocket
 			}
 			else
 			{
-				AgentSharingGroup = "CopyRocket" + StaticGetHostPlatformSuffix(HostPlatform);
+				AgentSharingGroup = "RocketGroup" + StaticGetHostPlatformSuffix(HostPlatform);
 			}
 
 			AddDependency(FilterRocketNode.StaticGetFullName(HostPlatform));
@@ -660,7 +660,7 @@ namespace Rocket
 			}
 			else
 			{
-				AgentSharingGroup = "CopyRocket" + StaticGetHostPlatformSuffix(HostPlatform);
+				AgentSharingGroup = "RocketGroup" + StaticGetHostPlatformSuffix(HostPlatform);
 			}
 
 			AddDependency(GUBP.ToolsForCompileNode.StaticGetFullName(HostPlatform));
