Build fix

diff --git a/Engine/Source/Programs/AutomationTool/ProjectUtils.cs b/Engine/Source/Programs/AutomationTool/ProjectUtils.cs
--- a/Engine/Source/Programs/AutomationTool/ProjectUtils.cs
+++ b/Engine/Source/Programs/AutomationTool/ProjectUtils.cs
@@ -211,7 +211,7 @@ namespace AutomationTool
 						// check the plugin info for this plugin itself
 						foreach (var Plugin in Plugins.AllPlugins)
 						{
-							if (Plugin.Name == PluginName && !Plugin.bEnabledByDefault && !RetVal)
+							if (Plugin.Name == PluginName && !Plugin.Descriptor.bEnabledByDefault && !RetVal)
 							{
 								foreach (var Module in Plugin.Descriptor.Modules)
 								{
