Prevent overlapping build products between Mac and Windows DDC nodes; copy output files to the Saved directory, and bring them back in when assembling the final build.

diff --git a/Engine/Source/Programs/AutomationTool/Scripts/RocketBuild.Automation.cs b/Engine/Source/Programs/AutomationTool/Scripts/RocketBuild.Automation.cs
--- a/Engine/Source/Programs/AutomationTool/Scripts/RocketBuild.Automation.cs
+++ b/Engine/Source/Programs/AutomationTool/Scripts/RocketBuild.Automation.cs
@@ -113,6 +113,7 @@ namespace Rocket
 				// Add the aggregate node for the entire install
 				GUBP.GUBPNode PromotableNode = bp.FindNode(GUBP.SharedAggregatePromotableNode.StaticGetFullName());
 				PromotableNode.AddDependency(FilterRocketNode.StaticGetFullName(HostPlatform));
+				PromotableNode.AddDependency(BuildDerivedDataCacheNode.StaticGetFullName(HostPlatform));
 
 				// Add a node for GitHub promotions
 				if(HostPlatform == UnrealTargetPlatform.Win64)
@@ -295,7 +296,7 @@ namespace Rocket
 			{
 				Platforms.Add(HostPlatform);
 			}
-			StrippedDir = Path.GetFullPath(CommandUtils.CombinePaths(CommandUtils.CmdEnv.LocalRoot, "Engine", "Saved", "Rocket"));
+			StrippedDir = Path.GetFullPath(CommandUtils.CombinePaths(CommandUtils.CmdEnv.LocalRoot, "Engine", "Saved", "Rocket", HostPlatform.ToString()));
 
 			// Add all the host nodes
 			AddDependency(GUBP.ToolsForCompileNode.StaticGetFullName(HostPlatform));
@@ -396,8 +397,8 @@ namespace Rocket
 			CodeTargetPlatforms = new List<UnrealTargetPlatform>(InCodeTargetPlatforms);
 			CurrentFeaturePacks = InCurrentFeaturePacks;
 			CurrentTemplates = InCurrentTemplates;
-			DepotManifestPath = CommandUtils.CombinePaths(CommandUtils.CmdEnv.LocalRoot, "Engine", "Saved", "Rocket", "DistillDepot" + StaticGetHostPlatformSuffix(HostPlatform) + ".txt");
-			StrippedManifestPath = CommandUtils.CombinePaths(CommandUtils.CmdEnv.LocalRoot, "Engine", "Saved", "Rocket", "DistillStripped" + StaticGetHostPlatformSuffix(HostPlatform) + ".txt");
+			DepotManifestPath = CommandUtils.CombinePaths(CommandUtils.CmdEnv.LocalRoot, "Engine", "Saved", "Rocket", HostPlatform.ToString(), "DistillDepot.txt");
+			StrippedManifestPath = CommandUtils.CombinePaths(CommandUtils.CmdEnv.LocalRoot, "Engine", "Saved", "Rocket", HostPlatform.ToString(), "DistillStripped.txt");
 
 			// Add the editor
 			AddDependency(GUBP.VersionFilesNode.StaticGetFullName());
@@ -422,9 +423,6 @@ namespace Rocket
 			// Add all the feature packs
 			AddDependency(GUBP.MakeFeaturePacksNode.StaticGetFullName(GUBP.MakeFeaturePacksNode.GetDefaultBuildPlatform(bp)));
 
-			// Add the DDC
-			AddDependency(BuildDerivedDataCacheNode.StaticGetFullName(HostPlatform));
-
 			// Find all the host platforms we need
 			SourceHostPlatforms = TargetPlatforms.Select(x => RocketBuild.GetSourceHostPlatform(bp, HostPlatform, x)).Distinct().ToList();
 			if(!SourceHostPlatforms.Contains(HostPlatform))
@@ -597,6 +595,7 @@ namespace Rocket
 
 			AddDependency(FilterRocketNode.StaticGetFullName(HostPlatform));
 			AddDependency(StripRocketNode.StaticGetFullName(HostPlatform));
+			AddDependency(BuildDerivedDataCacheNode.StaticGetFullName(HostPlatform));
 		}
 
 		public static string StaticGetFullName(UnrealTargetPlatform HostPlatform)
@@ -621,6 +620,10 @@ namespace Rocket
 			StripRocketNode StripNode = (StripRocketNode)bp.FindNode(StripRocketNode.StaticGetFullName(HostPlatform));
 			CopyManifestFilesToOutput(FilterNode.StrippedManifestPath, StripNode.StrippedDir, OutputDir);
 
+			// Copy the DDC to the output directory
+			BuildDerivedDataCacheNode DerivedDataCacheNode = (BuildDerivedDataCacheNode)bp.FindNode(BuildDerivedDataCacheNode.StaticGetFullName(HostPlatform));
+			CopyManifestFilesToOutput(DerivedDataCacheNode.SavedManifestPath, DerivedDataCacheNode.SavedDir, OutputDir);
+
 			// Write the UE4CommandLine.txt file with the 
 			string CommandLineFile = CommandUtils.CombinePaths(OutputDir, "UE4CommandLine.txt");
 			CommandUtils.WriteAllText(CommandLineFile, "-installedengine -rocket");
@@ -708,12 +711,16 @@ namespace Rocket
 	{
 		string TargetPlatforms;
 		string[] ProjectNames;
+		public readonly string SavedDir;
+		public readonly string SavedManifestPath;
 
 		public BuildDerivedDataCacheNode(UnrealTargetPlatform InHostPlatform, string InTargetPlatforms, string[] InProjectNames)
 			: base(InHostPlatform)
 		{
 			TargetPlatforms = InTargetPlatforms;
 			ProjectNames = InProjectNames;
+			SavedDir = CommandUtils.CombinePaths(CommandUtils.CmdEnv.LocalRoot, "Engine", "Saved", "Rocket", HostPlatform.ToString());
+			SavedManifestPath = CommandUtils.CombinePaths(SavedDir, "DerivedDataCacheManifest.txt");
 
 			AddDependency(GUBP.RootEditorNode.StaticGetFullName(HostPlatform));
 		}
@@ -735,13 +742,18 @@ namespace Rocket
 
 		public override void DoBuild(GUBP bp)
 		{
+			CommandUtils.CreateDirectory(SavedDir);
+
 			BuildProducts = new List<string>();
+
+			List<string> ManifestFiles = new List<string>();
 			if(!bp.ParseParam("NoDDC"))
 			{
 				string EditorExe = CommandUtils.GetEditorCommandletExe(CommandUtils.CmdEnv.LocalRoot, HostPlatform);
+				string RelativePakPath = "Engine/DerivedDataCache/Compressed.ddp";
 
 				// Delete the output file
-				string OutputPakFile = CommandUtils.CombinePaths(CommandUtils.CmdEnv.LocalRoot, "Engine", "DerivedDataCache", "Compressed.ddp");
+				string OutputPakFile = CommandUtils.CombinePaths(CommandUtils.CmdEnv.LocalRoot, RelativePakPath);
 				string OutputTxtFile = Path.ChangeExtension(OutputPakFile, ".txt");
 
 				// Generate DDC for all the non-code projects. We don't necessarily have editor DLLs for the code projects, but they should be the same as their blueprint counterparts.
@@ -757,12 +769,10 @@ namespace Rocket
 						string ProjectPakFile = CommandUtils.CombinePaths(Path.GetDirectoryName(OutputPakFile), String.Format("Compressed-{0}.ddp", ProjectName));
 						CommandUtils.DeleteFile(ProjectPakFile);
 						CommandUtils.RenameFile(OutputPakFile, ProjectPakFile);
-						BuildProducts.Add(ProjectPakFile);
 
 						string ProjectTxtFile = Path.ChangeExtension(ProjectPakFile, ".txt");
 						CommandUtils.DeleteFile(ProjectTxtFile);
 						CommandUtils.RenameFile(OutputTxtFile, ProjectTxtFile);
-						BuildProducts.Add(ProjectTxtFile);
 
 						ProjectPakFiles.Add(Path.GetFileName(ProjectPakFile));
 					}
@@ -772,10 +782,17 @@ namespace Rocket
 				CommandUtils.Log("Generating DDC data for engine content on {0}", TargetPlatforms);
 				CommandUtils.DDCCommandlet(null, EditorExe, null, TargetPlatforms, "-fill -DDC=CreateInstalledEnginePak " + CommandUtils.MakePathSafeToUseWithCommandLine("-MergePaks=" + String.Join("+", ProjectPakFiles)));
 
-				// Add the final PAK file as output
-				BuildProducts.Add(OutputPakFile);
-				BuildProducts.Add(OutputTxtFile);
+				// Copy the DDP file to the output path
+				string SavedPakFile = CommandUtils.CombinePaths(SavedDir, RelativePakPath);
+				CommandUtils.CopyFile(OutputPakFile, SavedPakFile);
+				BuildProducts.Add(SavedPakFile);
+
+				// Add the pak file to the list of files to copy
+				ManifestFiles.Add(RelativePakPath);
 			}
+			CommandUtils.WriteAllLines(SavedManifestPath, ManifestFiles.ToArray());
+			BuildProducts.Add(SavedManifestPath);
+
 			SaveRecordOfSuccessAndAddToBuildProducts();
 		}
 
