Added error for when engine content is access by dlc at cook time.

diff --git a/Engine/Source/Developer/LauncherServices/Private/Profiles/LauncherProfile.h b/Engine/Source/Developer/LauncherServices/Private/Profiles/LauncherProfile.h
--- a/Engine/Source/Developer/LauncherServices/Private/Profiles/LauncherProfile.h
+++ b/Engine/Source/Developer/LauncherServices/Private/Profiles/LauncherProfile.h
@@ -6,6 +6,7 @@
 #define LAUNCHERSERVICES_ADDEDPATCHSOURCECONTENTPATH 12
 #define LAUNCHERSERVICES_ADDEDRELEASEVERSION 13
 #define LAUNCHERSERVICES_REMOVEDPATCHSOURCECONTENTPATH 14
+#define LAUNCHERSERVICES_ADDEDDLCINCLUDEENGINECONTENT 15
 
 /**
 * Implements a simple profile which controls the desired output of the Launcher for simple
@@ -286,21 +287,31 @@ public:
 	{
 		return CreateDLC;
 	}
+	virtual void SetCreateDLC(bool InBuildDLC) override
+	{
+		CreateDLC = InBuildDLC;
+	}
 
 	virtual FString GetDLCName() const override
 	{
 		return DLCName;
 	}
+	virtual void SetDLCName(const FString& InDLCName) override
+	{
+		DLCName = InDLCName;
+	}
 
-	virtual void SetCreateDLC(bool InBuildDLC) override
+	virtual bool IsDLCIncludingEngineContent() const
 	{
-		CreateDLC = InBuildDLC;
+		return DLCIncludeEngineContent;
 	}
-	virtual void SetDLCName(const FString& InDLCName) override
+	virtual void SetDLCIncludeEngineContent(bool InDLCIncludeEngineContent)
 	{
-		DLCName = InDLCName;
+		DLCIncludeEngineContent = InDLCIncludeEngineContent;
 	}
 
+
+
 	virtual bool IsCreatingReleaseVersion() const override
 	{
 		return CreateReleaseVersion;
@@ -618,6 +629,16 @@ public:
 			GeneratePatch = false;
 		}
 
+		if (Version >= LAUNCHERSERVICES_ADDEDDLCINCLUDEENGINECONTENT)
+		{
+			Archive << DLCIncludeEngineContent;
+		}
+		else if (Archive.IsLoading())
+		{
+			DLCIncludeEngineContent = false;
+		}
+		
+
 		if ( Version >= LAUNCHERSERVICES_ADDEDRELEASEVERSION )
 		{
 			Archive << CreateReleaseVersion;
@@ -725,6 +746,7 @@ public:
 		CreateReleaseVersion = false;
 		GeneratePatch = false;
 		CreateDLC = false;
+		DLCIncludeEngineContent = false;
 
 		// default launch settings
 		LaunchMode = ELauncherProfileLaunchModes::DefaultRole;
@@ -1292,9 +1314,10 @@ private:
 	// name of the dlc we are going to build (the name of the dlc plugin)
 	FString DLCName;
 
-	
-
-
+	// should the dlc include engine content in the current dlc 
+	//  engine content which was not referenced by original release
+	//  otherwise error on any access of engine content during dlc cook
+	bool DLCIncludeEngineContent;
 
 	// Holds a flag indicating whether to use incremental deployment
 	bool DeployIncremental;
diff --git a/Engine/Source/Developer/LauncherServices/Public/Interfaces/ILauncherProfile.h b/Engine/Source/Developer/LauncherServices/Public/Interfaces/ILauncherProfile.h
--- a/Engine/Source/Developer/LauncherServices/Public/Interfaces/ILauncherProfile.h
+++ b/Engine/Source/Developer/LauncherServices/Public/Interfaces/ILauncherProfile.h
@@ -842,14 +842,16 @@ public:
 	 */
 	virtual void SetGeneratePatch( bool InShouldGeneratePatch ) = 0;
 
-	virtual bool IsCreatingDLC() const = 0;
 
+	virtual bool IsCreatingDLC() const = 0;
 	virtual void SetCreateDLC(bool InBuildDLC) = 0;
 
 	virtual FString GetDLCName() const = 0;
-
 	virtual void SetDLCName(const FString& InDLCName) = 0;
 
+	virtual bool IsDLCIncludingEngineContent() const = 0;
+	virtual void SetDLCIncludeEngineContent(bool InDLCIncludeEngineContent) = 0;
+
 
     /**
      * Sets the cook on the fly close mode
diff --git a/Engine/Source/Developer/ProjectLauncher/Private/Widgets/Cook/SProjectLauncherCookByTheBookSettings.cpp b/Engine/Source/Developer/ProjectLauncher/Private/Widgets/Cook/SProjectLauncherCookByTheBookSettings.cpp
--- a/Engine/Source/Developer/ProjectLauncher/Private/Widgets/Cook/SProjectLauncherCookByTheBookSettings.cpp
+++ b/Engine/Source/Developer/ProjectLauncher/Private/Widgets/Cook/SProjectLauncherCookByTheBookSettings.cpp
@@ -476,6 +476,23 @@ TSharedRef<SWidget> SProjectLauncherCookByTheBookSettings::MakeComplexWidget()
 								.OnTextCommitted(this, &SProjectLauncherCookByTheBookSettings::HandleDLCNameCommitted)
 							]
 						]
+
+					+ SVerticalBox::Slot()
+						.AutoHeight()
+						.Padding(0.0f, 4.0f, 0.0f, 0.0f)
+						[
+							// unreal pak check box
+							SNew(SCheckBox)
+							.IsChecked(this, &SProjectLauncherCookByTheBookSettings::HandleDLCIncludeEngineContentCheckBoxIsChecked)
+							.OnCheckStateChanged(this, &SProjectLauncherCookByTheBookSettings::HandleDLCIncludeEngineContentCheckBoxCheckStateChanged)
+							.Padding(FMargin(4.0f, 0.0f))
+							.ToolTipText(LOCTEXT("HandleDLCIncludeEngineContentCheckBoxTooltip", "If checked, dlc will include engine content which was not included in original release, if not checked will error when accessing content from engine directory."))
+							.Content()
+							[
+								SNew(STextBlock)
+								.Text(LOCTEXT("HandleDLCIncludeEngineContentCheckBoxText", "Include engine content"))
+							]
+						]
 					// end generate dlc 
 					//////////////////////////////////////////////////////////////////////////
 			]
@@ -1388,8 +1405,6 @@ void SProjectLauncherCookByTheBookSettings::HandleBasedOnReleaseVersionNameCommi
 	}
 
 }
-
-
 FText SProjectLauncherCookByTheBookSettings::HandleDLCNameTextBlockText() const
 {
 	ILauncherProfilePtr SelectedProfile = Model->GetSelectedProfile();
@@ -1442,4 +1457,40 @@ ECheckBoxState SProjectLauncherCookByTheBookSettings::HandleBuildDLCCheckBoxIsCh
 }
 
 
+
+// Callback for check state changes of the 'UnrealPak' check box.
+void SProjectLauncherCookByTheBookSettings::HandleDLCIncludeEngineContentCheckBoxCheckStateChanged(ECheckBoxState NewState)
+{
+	ILauncherProfilePtr SelectedProfile = Model->GetSelectedProfile();
+
+	if (SelectedProfile.IsValid())
+	{
+		SelectedProfile->SetDLCIncludeEngineContent(NewState == ECheckBoxState::Checked);
+	}
+}
+
+// Callback for determining the checked state of the 'UnrealPak' check box.
+ECheckBoxState SProjectLauncherCookByTheBookSettings::HandleDLCIncludeEngineContentCheckBoxIsChecked() const
+{
+	ILauncherProfilePtr SelectedProfile = Model->GetSelectedProfile();
+
+	if (SelectedProfile.IsValid())
+	{
+		if (SelectedProfile->IsDLCIncludingEngineContent())
+		{
+			return ECheckBoxState::Checked;
+		}
+	}
+
+	return ECheckBoxState::Unchecked;
+}
+
+
+
+
+
+
+
+
+
 #undef LOCTEXT_NAMESPACE
diff --git a/Engine/Source/Developer/ProjectLauncher/Private/Widgets/Cook/SProjectLauncherCookByTheBookSettings.h b/Engine/Source/Developer/ProjectLauncher/Private/Widgets/Cook/SProjectLauncherCookByTheBookSettings.h
--- a/Engine/Source/Developer/ProjectLauncher/Private/Widgets/Cook/SProjectLauncherCookByTheBookSettings.h
+++ b/Engine/Source/Developer/ProjectLauncher/Private/Widgets/Cook/SProjectLauncherCookByTheBookSettings.h
@@ -157,6 +157,8 @@ private:
 	ECheckBoxState HandleBuildDLCCheckBoxIsChecked() const;
 	void HandleDLCNameCommitted(const FText& NewText, ETextCommit::Type CommitType);
 	FText HandleDLCNameTextBlockText() const;
+	void HandleDLCIncludeEngineContentCheckBoxCheckStateChanged(ECheckBoxState NewState);
+	ECheckBoxState SProjectLauncherCookByTheBookSettings::HandleDLCIncludeEngineContentCheckBoxIsChecked() const;
 	
 
 
diff --git a/Engine/Source/Editor/UnrealEd/Classes/Commandlets/CookCommandlet.h b/Engine/Source/Editor/UnrealEd/Classes/Commandlets/CookCommandlet.h
--- a/Engine/Source/Editor/UnrealEd/Classes/Commandlets/CookCommandlet.h
+++ b/Engine/Source/Editor/UnrealEd/Classes/Commandlets/CookCommandlet.h
@@ -34,6 +34,8 @@ class UCookCommandlet
 	bool bUnversioned;
 	/** Generate manifests for building streaming install packages */
 	bool bGenerateStreamingInstallManifests;
+	/** Error if we access engine content (useful for dlc) */
+	bool bErrorOnEngineContentUse;
 	/** All commandline tokens */
 	TArray<FString> Tokens;
 	/** All commandline switches */
diff --git a/Engine/Source/Programs/AutomationTool/ProjectParams.cs b/Engine/Source/Programs/AutomationTool/ProjectParams.cs
--- a/Engine/Source/Programs/AutomationTool/ProjectParams.cs
+++ b/Engine/Source/Programs/AutomationTool/ProjectParams.cs
@@ -200,6 +200,7 @@ namespace AutomationTool
             this.CreateReleaseVersion = InParams.CreateReleaseVersion;
             this.GeneratePatch = InParams.GeneratePatch;
             this.DLCName = InParams.DLCName;
+            this.DLCIncludeEngineContent = InParams.DLCIncludeEngineContent;
             this.NewCook = InParams.NewCook;
             this.OldCook = InParams.OldCook;
             this.AdditionalCookerOptions = InParams.AdditionalCookerOptions;
@@ -345,6 +346,7 @@ namespace AutomationTool
             string CreateReleaseVersion = null,
             bool? GeneratePatch = null,
             string DLCName = null,
+            bool? DLCIncludeEngineContent = null,
             bool? NewCook = null,
             bool? OldCook = null,
 			bool? CrashReporter = null,
@@ -448,6 +450,7 @@ namespace AutomationTool
             this.GeneratePatch = GetParamValueIfNotSpecified(Command, GeneratePatch, this.GeneratePatch, "GeneratePatch");
             this.AdditionalCookerOptions = ParseParamValueIfNotSpecified(Command, AdditionalCookerOptions, "AdditionalCookerOptions", String.Empty);
             this.DLCName = ParseParamValueIfNotSpecified(Command, DLCName, "DLCName", String.Empty);
+            this.DLCIncludeEngineContent = GetParamValueIfNotSpecified(Command, DLCIncludeEngineContent, this.DLCIncludeEngineContent, "DLCIncludeEngineContent");
 			this.SkipCook = GetParamValueIfNotSpecified(Command, SkipCook, this.SkipCook, "skipcook");
 			if (this.SkipCook)
 			{
@@ -1037,6 +1040,12 @@ namespace AutomationTool
         /// </summary>
         public string DLCName;
 
+        /// <summary>
+        /// Enable cooking of engine content when cooking dlc 
+        ///  not included in original release but is referenced by current cook
+        /// </summary>
+        public bool DLCIncludeEngineContent;
+
         /// <summary>
         /// Cook: Additional cooker options to include on the cooker commandline
         /// </summary>
@@ -1928,6 +1937,11 @@ namespace AutomationTool
 				throw new AutomationException("Can't use both -cook and -cookonthefly.");
 			}
 
+            if (!HasDLCName && DLCIncludeEngineContent)
+            {
+                throw new AutomationException("DLCIncludeEngineContent flag is only valid when cooking dlc.");
+            }
+
             if ((IsGeneratingPatch || HasDLCName) && !HasBasedOnReleaseVersion)
             {
                 throw new AutomationException("Require based on release version to build patches or dlc");
@@ -1995,6 +2009,7 @@ namespace AutomationTool
                 CommandUtils.Log("CreateReleaseVersion={0}", CreateReleaseVersion);
                 CommandUtils.Log("BasedOnReleaseVersion={0}", BasedOnReleaseVersion);
                 CommandUtils.Log("DLCName={0}", DLCName);
+                CommandUtils.Log("DLCIncludeEngineContent={0}", DLCIncludeEngineContent);
                 CommandUtils.Log("AdditionalCookerOptions={0}", AdditionalCookerOptions);
 				CommandUtils.Log("DedicatedServer={0}", DedicatedServer);
 				CommandUtils.Log("DirectoriesToCook={0}", DirectoriesToCook.ToString());
diff --git a/Engine/Source/Programs/AutomationTool/Scripts/CookCommand.Automation.cs b/Engine/Source/Programs/AutomationTool/Scripts/CookCommand.Automation.cs
--- a/Engine/Source/Programs/AutomationTool/Scripts/CookCommand.Automation.cs
+++ b/Engine/Source/Programs/AutomationTool/Scripts/CookCommand.Automation.cs
@@ -174,6 +174,10 @@ public partial class Project : CommandUtils
                 if (Params.HasDLCName)
                 {
                     CommandletParams += " -dlcname=" + Params.DLCName;
+                    if ( !Params.DLCIncludeEngineContent )
+                    {
+                        CommandletParams += " -errorOnEngineContentUse";
+                    }
                 }
                 // don't include the based on release version unless we are cooking dlc or creating a new release version
                 // in this case the based on release version is used in packaging
diff --git a/Engine/Source/Programs/AutomationTool/Scripts/CopyBuildToStagingDirectory.Automation.cs b/Engine/Source/Programs/AutomationTool/Scripts/CopyBuildToStagingDirectory.Automation.cs
--- a/Engine/Source/Programs/AutomationTool/Scripts/CopyBuildToStagingDirectory.Automation.cs
+++ b/Engine/Source/Programs/AutomationTool/Scripts/CopyBuildToStagingDirectory.Automation.cs
@@ -247,7 +247,7 @@ public partial class Project : CommandUtils
             SC.StageFiles(StagedFileType.NonUFS, CombinePaths(SC.ProjectRoot, "Plugins", DLCName, "Binaries"), "UE4Server-*.dll", true, null, null, true);
             
             // Put all of the cooked dir into the staged dir
-            if (SC.DedicatedServer)
+            if (SC.DedicatedServer || Params.DLCIncludeEngineContent)
             {
                 // Dedicated server cook doesn't save shaders so no Engine dir is created
                 SC.StageFiles(StagedFileType.UFS, CombinePaths(SC.ProjectRoot, "Plugins", DLCName, "Saved", "Cooked", SC.CookPlatform), "*", true, new[] { "AssetRegistry.bin" }, "", true, !Params.UsePak(SC.StageTargetPlatform));
