Archive headers as part of building the root editor and precompiled monolithics, and extract them into the Rocket build.

diff --git a/Engine/Source/Programs/AutomationTool/CommandUtils.cs b/Engine/Source/Programs/AutomationTool/CommandUtils.cs
--- a/Engine/Source/Programs/AutomationTool/CommandUtils.cs
+++ b/Engine/Source/Programs/AutomationTool/CommandUtils.cs
@@ -2097,6 +2097,36 @@ namespace AutomationTool
 				return Enum.GetName(typeof(UnrealBuildTool.UnrealTargetPlatform), Platform);
 			}
 		}
+
+		/// <summary>
+		/// Creates a zip file containing the given input files
+		/// </summary>
+		/// <param name="ZipFileName">Filename for the zip</param>
+		/// <param name="Filter">Filter which selects files to be included in the zip</param>
+		/// <param name="BaseDirectory">Base directory to store relative paths in the zip file to</param>
+		public static void ZipFiles(string ZipFileName, string BaseDirectory, FileFilter Filter)
+		{
+			Ionic.Zip.ZipFile Zip = new Ionic.Zip.ZipFile();
+			foreach(string FilteredFile in Filter.ApplyToDirectory(BaseDirectory, true))
+			{
+				Zip.AddFile(Path.Combine(BaseDirectory, FilteredFile), Path.GetDirectoryName(FilteredFile));
+			}
+			CommandUtils.CreateDirectory(Path.GetDirectoryName(ZipFileName));
+			Zip.Save(ZipFileName);
+		}
+
+		/// <summary>
+		/// Extracts the contents of a zip file
+		/// </summary>
+		/// <param name="ZipFileName">Name of the zip file</param>
+		/// <param name="BaseDirectory">Output directory</param>
+		/// <returns>List of files written</returns>
+		public static IEnumerable<string> UnzipFiles(string ZipFileName, string BaseDirectory)
+		{
+			Ionic.Zip.ZipFile Zip = new Ionic.Zip.ZipFile(ZipFileName);
+			Zip.ExtractAll(BaseDirectory, Ionic.Zip.ExtractExistingFileAction.OverwriteSilently);
+			return Zip.EntryFileNames.Select(x => Path.Combine(BaseDirectory, x));
+		}
 	}
 
 
diff --git a/Engine/Source/Programs/AutomationTool/Scripts/GUBP.Automation.cs b/Engine/Source/Programs/AutomationTool/Scripts/GUBP.Automation.cs
--- a/Engine/Source/Programs/AutomationTool/Scripts/GUBP.Automation.cs
+++ b/Engine/Source/Programs/AutomationTool/Scripts/GUBP.Automation.cs
@@ -822,6 +822,25 @@ public class GUBP : BuildCommand
         {
             return true;
         }
+		public override void DoBuild(GUBP bp)
+		{
+			base.DoBuild(bp);
+
+			if(!bp.BranchOptions.bNoInstalledEngine)
+			{
+				FileFilter Filter = new FileFilter();
+				Filter.Include("/Engine/Intermediate/Build/" + HostPlatform.ToString() + "/UE4Editor/Inc/...");
+				Filter.Include("/Engine/Plugins/.../Intermediate/Build/" + HostPlatform.ToString() + "/UE4Editor/Inc/...");
+
+				string ZipFileName = StaticGetArchivedHeadersPath(HostPlatform);
+				CommandUtils.ZipFiles(ZipFileName, CommandUtils.CmdEnv.LocalRoot, Filter);
+				BuildProducts.Add(ZipFileName);
+			}
+		}
+		public static string StaticGetArchivedHeadersPath(UnrealTargetPlatform HostPlatform)
+		{
+			return CommandUtils.CombinePaths(CommandUtils.CmdEnv.LocalRoot, "Engine", "Saved", "Precompiled", "Headers-RootEditor" + StaticGetHostPlatformSuffix(HostPlatform) + ".zip");
+		}
         public override UE4Build.BuildAgenda GetAgenda(GUBP bp)
         {
             var Agenda = new UE4Build.BuildAgenda();
@@ -1700,9 +1719,18 @@ public class GUBP : BuildCommand
 				}
 
 				// Write the resulting file list out to disk
-				string OutputFileListPath = StaticGetBuildDependenciesPath(HostPlatform, TargetPlatform);
+				string OutputFileListPath = StaticGetBuildDependenciesPath(HostPlatform, GameProj, TargetPlatform);
 				UnrealBuildTool.Utils.WriteClass<UnrealBuildTool.ExternalFileList>(FileList, OutputFileListPath, "");
 				AddBuildProduct(OutputFileListPath);
+
+				// Archive all the headers
+				FileFilter Filter = new FileFilter();
+				Filter.Include("/Engine/Intermediate/Build/" + HostPlatform.ToString() + "/UE4/Inc/...");
+				Filter.Include("/Engine/Plugins/.../Intermediate/Build/" + HostPlatform.ToString() + "/UE4/Inc/...");
+
+				string ZipFileName = StaticGetArchivedHeadersPath(HostPlatform, GameProj, TargetPlatform);
+				CommandUtils.ZipFiles(ZipFileName, CommandUtils.CmdEnv.LocalRoot, Filter);
+				BuildProducts.Add(ZipFileName);
 			}
 		}
 
@@ -1775,34 +1803,16 @@ public class GUBP : BuildCommand
             return Agenda;
         }
 
-		public static string StaticGetBuildDependenciesPath(UnrealTargetPlatform HostPlatform, UnrealTargetPlatform TargetPlatform)
+		public static string StaticGetArchivedHeadersPath(UnrealTargetPlatform HostPlatform, BranchInfo.BranchUProject GameProj, UnrealTargetPlatform TargetPlatform)
 		{
-			return CommandUtils.CombinePaths(CommandUtils.CmdEnv.LocalRoot, "Engine/Saved/BuildDependencies/" + TargetPlatform.ToString() + StaticGetHostPlatformSuffix(HostPlatform) + ".xml");
+			return CommandUtils.CombinePaths(CommandUtils.CmdEnv.LocalRoot, "Engine", "Saved", "Precompiled", "Headers-" + StaticGetFullName(HostPlatform, GameProj, TargetPlatform) + ".zip");
 		}
-    }
 
-	public static class HeadersNode
-	{
-		public static void ZipHeaders(IEnumerable<string> HeaderFiles, string ZipFileName)
+		public static string StaticGetBuildDependenciesPath(UnrealTargetPlatform HostPlatform, BranchInfo.BranchUProject GameProj, UnrealTargetPlatform TargetPlatform)
 		{
-			string NormalizedPrefix = CommandUtils.ConvertSeparators(PathSeparator.Slash, CmdEnv.LocalRoot).TrimEnd('/') + "/";
-
-			Ionic.Zip.ZipFile Zip = new Ionic.Zip.ZipFile();
-			foreach(string HeaderFile in HeaderFiles)
-			{
-				string NormalizedDirectoryName = CommandUtils.ConvertSeparators(PathSeparator.Slash, Path.GetDirectoryName(HeaderFile));
-				if(NormalizedDirectoryName.StartsWith(NormalizedPrefix))
-				{
-					Zip.AddFile(HeaderFile, NormalizedDirectoryName.Substring(NormalizedPrefix.Length));
-				}
-				else
-				{
-					throw new AutomationException("Header file '{0}' was not under root directory ('{1}')", NormalizedDirectoryName, NormalizedPrefix);
-				}
-			}
-			Zip.Save(ZipFileName);
+			return CommandUtils.CombinePaths(CommandUtils.CmdEnv.LocalRoot, "Engine", "Saved", "Precompiled", "BuildDependencies-" + StaticGetFullName(HostPlatform, GameProj, TargetPlatform) + ".xml");
 		}
-	}
+    }
 
     public class SuccessNode : GUBPNode
     {
diff --git a/Engine/Source/Programs/AutomationTool/Scripts/RocketBuild.Automation.cs b/Engine/Source/Programs/AutomationTool/Scripts/RocketBuild.Automation.cs
--- a/Engine/Source/Programs/AutomationTool/Scripts/RocketBuild.Automation.cs
+++ b/Engine/Source/Programs/AutomationTool/Scripts/RocketBuild.Automation.cs
@@ -100,13 +100,13 @@ namespace Rocket
 				if(ShouldDoSeriousThingsLikeP4CheckinAndPostToMCP())
 				{
 					string OutputDir = CommandUtils.CombinePaths(CommandUtils.RootSharedTempStorageDirectory(), "Rocket", "Automated", GetBuildLabel(), CommandUtils.GetGenericPlatformName(HostPlatform));
-					bp.AddNode(new CopyRocketNode(HostPlatform, OutputDir, true));
+					bp.AddNode(new CopyRocketNode(HostPlatform, CodeTargetPlatforms, OutputDir, true));
 					bp.AddNode(new CopyRocketSymbolsNode(bp, HostPlatform, CodeTargetPlatforms, OutputDir + "Symbols", true));
 				}
 				else
 				{
 					string OutputDir = CommandUtils.CombinePaths(CommandUtils.CmdEnv.LocalRoot, "LocalBuilds", "Rocket", CommandUtils.GetGenericPlatformName(HostPlatform));
-					bp.AddNode(new CopyRocketNode(HostPlatform, OutputDir, false));
+					bp.AddNode(new CopyRocketNode(HostPlatform, CodeTargetPlatforms, OutputDir, false));
 					bp.AddNode(new CopyRocketSymbolsNode(bp, HostPlatform, CodeTargetPlatforms, OutputDir + "Symbols", false));
 				}
 
@@ -461,12 +461,16 @@ namespace Rocket
 			AddRuleForBuildProducts(Filter, bp, GUBP.RootEditorNode.StaticGetFullName(HostPlatform), FileFilterType.Include);
 			AddRuleForBuildProducts(Filter, bp, GUBP.ToolsNode.StaticGetFullName(HostPlatform), FileFilterType.Include);
 
+			// Include the editor headers
+			UnzipAndAddRuleForHeaders(GUBP.RootEditorNode.StaticGetFullName(HostPlatform), Filter, FileFilterType.Include);
+
 			// Include the build dependencies for every code platform
 			foreach(UnrealTargetPlatform CodeTargetPlatform in CodeTargetPlatforms)
 			{
 				UnrealTargetPlatform SourceHostPlatform = RocketBuild.GetSourceHostPlatform(bp, HostPlatform, CodeTargetPlatform);
-				string FileListPath = GUBP.GamePlatformMonolithicsNode.StaticGetBuildDependenciesPath(SourceHostPlatform, CodeTargetPlatform);
+				string FileListPath = GUBP.GamePlatformMonolithicsNode.StaticGetBuildDependenciesPath(SourceHostPlatform, bp.Branch.BaseEngineProject, CodeTargetPlatform);
 				Filter.AddRuleForFiles(UnrealBuildTool.Utils.ReadClass<UnrealBuildTool.ExternalFileList>(FileListPath).FileNames, CommandUtils.CmdEnv.LocalRoot, FileFilterType.Include);
+				UnzipAndAddRuleForHeaders(GUBP.GamePlatformMonolithicsNode.StaticGetArchivedHeadersPath(HostPlatform, bp.Branch.BaseEngineProject, CodeTargetPlatform), Filter, FileFilterType.Include);
 			}
 
 			// Add the monolithic binaries
@@ -568,6 +572,12 @@ namespace Rocket
 			Filter.AddRuleForFiles(Node.BuildProducts, CommandUtils.CmdEnv.LocalRoot, Type);
 		}
 
+		static void UnzipAndAddRuleForHeaders(string ZipFileName, FileFilter Filter, FileFilterType Type)
+		{
+			IEnumerable<string> FileNames = CommandUtils.UnzipFiles(ZipFileName, CommandUtils.CmdEnv.LocalRoot);
+			Filter.AddRuleForFiles(FileNames, CommandUtils.CmdEnv.LocalRoot, FileFilterType.Include);
+		}
+
 		public static void WriteManifest(string FileName, IEnumerable<string> Files)
 		{
 			CommandUtils.CreateDirectory(Path.GetDirectoryName(FileName));
@@ -578,10 +588,12 @@ namespace Rocket
 	public class CopyRocketNode : GUBP.HostPlatformNode
 	{
 		public readonly string OutputDir;
+		public List<UnrealTargetPlatform> CodeTargetPlatforms;
 
-		public CopyRocketNode(UnrealTargetPlatform HostPlatform, string InOutputDir, bool bWaitForPromotion) : base(HostPlatform)
+		public CopyRocketNode(UnrealTargetPlatform HostPlatform, List<UnrealTargetPlatform> InCodeTargetPlatforms, string InOutputDir, bool bWaitForPromotion) : base(HostPlatform)
 		{
 			OutputDir = InOutputDir;
+			CodeTargetPlatforms = new List<UnrealTargetPlatform>(InCodeTargetPlatforms);
 
 			// If this is a CIS build, we need to wait for the promotable trigger before running. If not, make an agent sharing group so that the samples node can use our build products.
 			if(bWaitForPromotion)
@@ -612,6 +624,17 @@ namespace Rocket
 		{
 			CommandUtils.DeleteDirectoryContents(OutputDir);
 
+			// Extract the editor headers
+			CommandUtils.UnzipFiles(GUBP.RootEditorNode.StaticGetFullName(HostPlatform), CommandUtils.CmdEnv.LocalRoot);
+
+			// Extract all the headers for code target platforms
+			foreach(UnrealTargetPlatform CodeTargetPlatform in CodeTargetPlatforms)
+			{
+				UnrealTargetPlatform SourceHostPlatform = RocketBuild.GetSourceHostPlatform(bp, HostPlatform, CodeTargetPlatform);
+				string ZipFileName = GUBP.GamePlatformMonolithicsNode.StaticGetArchivedHeadersPath(SourceHostPlatform, bp.Branch.BaseEngineProject, CodeTargetPlatform);
+				CommandUtils.UnzipFiles(ZipFileName, CommandUtils.CmdEnv.LocalRoot);
+			}
+
 			// Copy the depot files to the output directory
 			FilterRocketNode FilterNode = (FilterRocketNode)bp.FindNode(FilterRocketNode.StaticGetFullName(HostPlatform));
 			CopyManifestFilesToOutput(FilterNode.DepotManifestPath, CommandUtils.CmdEnv.LocalRoot, OutputDir);
