Fix build failure on monolithic build nodes; VCCompiler attribute may not exist.

diff --git a/Engine/Source/Programs/AutomationTool/UE4Build.cs b/Engine/Source/Programs/AutomationTool/UE4Build.cs
--- a/Engine/Source/Programs/AutomationTool/UE4Build.cs
+++ b/Engine/Source/Programs/AutomationTool/UE4Build.cs
@@ -915,7 +915,10 @@ namespace AutomationTool
 
 								ToolElement.SetAttribute("Params", Element.Attributes["Params"].Value);
 								ToolElement.SetAttribute("Path", Element.Attributes["Path"].Value);
-								ToolElement.SetAttribute("VCCompiler", Element.Attributes["VCCompiler"].Value);
+								if(Element.HasAttribute("VCCompiler"))
+								{
+									ToolElement.SetAttribute("VCCompiler", Element.Attributes["VCCompiler"].Value);
+								}
 								ToolElement.SetAttribute("SkipIfProjectFailed", Element.Attributes["SkipIfProjectFailed"].Value);
 								if (Element.HasAttribute("AutoReserveMemory"))
 								{
