Linux: correct editor commandlet.

(Submitting a missed a local change).

diff --git a/Engine/Source/Programs/AutomationTool/CommandletUtils.cs b/Engine/Source/Programs/AutomationTool/CommandletUtils.cs
--- a/Engine/Source/Programs/AutomationTool/CommandletUtils.cs
+++ b/Engine/Source/Programs/AutomationTool/CommandletUtils.cs
@@ -204,6 +204,8 @@ namespace AutomationTool
 					return CommandUtils.CombinePaths(BuildRoot, "Engine/Binaries/Mac/UE4Editor.app/Contents/MacOS/UE4Editor");
 				case UnrealBuildTool.UnrealTargetPlatform.Win64:
 					return CommandUtils.CombinePaths(BuildRoot, "Engine/Binaries/Win64/UE4Editor-Cmd.exe");
+				case UnrealBuildTool.UnrealTargetPlatform.Linux:
+					return CommandUtils.CombinePaths(BuildRoot, "Engine/Binaries/Linux/UE4Editor");
 				default:
 					throw new NotImplementedException();
 			}
