Don't try to filter DDC as part of the FilterRocketNode; it no longer has a dependency on it (and CopyRocketNode brings it in manually).

diff --git a/Engine/Source/Programs/AutomationTool/Scripts/RocketBuild.Automation.cs b/Engine/Source/Programs/AutomationTool/Scripts/RocketBuild.Automation.cs
--- a/Engine/Source/Programs/AutomationTool/Scripts/RocketBuild.Automation.cs
+++ b/Engine/Source/Programs/AutomationTool/Scripts/RocketBuild.Automation.cs
@@ -487,9 +487,6 @@ namespace Rocket
 				Filter.AddRuleForFile(GUBP.MakeFeaturePacksNode.GetOutputFile(Project), CommandUtils.CmdEnv.LocalRoot, FileFilterType.Include);
 			}
 
-			// Include the DDC
-			AddRuleForBuildProducts(Filter, bp, BuildDerivedDataCacheNode.StaticGetFullName(HostPlatform), FileFilterType.Include);
-
 			// Include all the templates
 			foreach (string Template in CurrentTemplates)
 			{
