Implement CommandUtils.UnzipFiles() on Mac without using Ionic.Zip.ExtractAll(). Seems to throw an exception under Mono, and there's obvious reason why.

diff --git a/Engine/Source/Programs/AutomationTool/CommandUtils.cs b/Engine/Source/Programs/AutomationTool/CommandUtils.cs
--- a/Engine/Source/Programs/AutomationTool/CommandUtils.cs
+++ b/Engine/Source/Programs/AutomationTool/CommandUtils.cs
@@ -2124,8 +2124,20 @@ namespace AutomationTool
 		public static IEnumerable<string> UnzipFiles(string ZipFileName, string BaseDirectory)
 		{
 			Ionic.Zip.ZipFile Zip = new Ionic.Zip.ZipFile(ZipFileName);
-			Zip.ExtractAll(BaseDirectory, Ionic.Zip.ExtractExistingFileAction.OverwriteSilently);
-			return Zip.EntryFileNames.Select(x => Path.Combine(BaseDirectory, x));
+
+			// For some reason, calling ExtractAll() under Mono throws an exception from a call to Directory.CreateDirectory(). Manually extract all the files instead.
+			List<string> OutputFileNames = new List<string>();
+			foreach(Ionic.Zip.ZipEntry Entry in Zip.Entries)
+			{
+				string OutputFileName = Path.Combine(BaseDirectory, Entry.FileName);
+				Directory.CreateDirectory(Path.GetDirectoryName(OutputFileName));
+				using(FileStream OutputStream = new FileStream(OutputFileName, FileMode.Create, FileAccess.Write))
+				{
+					Entry.Extract(OutputStream);
+				}
+				OutputFileNames.Add(OutputFileName);
+			}
+			return OutputFileNames;
 		}
 	}
 
