- Made some changes to allow receipt based packaging to work with the Desktop platform. One hacky bit remains that I have entered a JIRA for [UE-14120]

diff --git a/Engine/Source/Developer/DesktopTargetPlatform/Private/DesktopTargetPlatform.cpp b/Engine/Source/Developer/DesktopTargetPlatform/Private/DesktopTargetPlatform.cpp
--- a/Engine/Source/Developer/DesktopTargetPlatform/Private/DesktopTargetPlatform.cpp
+++ b/Engine/Source/Developer/DesktopTargetPlatform/Private/DesktopTargetPlatform.cpp
@@ -41,8 +41,12 @@ void FDesktopTargetPlatform::GetAllPossibleShaderFormats( TArray<FName>& OutForm
 	static FName NAME_GLSL_430(TEXT("GLSL_430"));
 	static FName NAME_GLSL_150_MAC(TEXT("GLSL_150_MAC"));
 	
-//	OutFormats.AddUnique(NAME_PCD3D_SM5);
-//	OutFormats.AddUnique(NAME_PCD3D_SM4);
+#if PLATFORM_WINDOWS
+	// right now, only windows can properly compile D3D shaders (this won't corrupt the DDC, but it will 
+	// make it so that packages cooked on Mac/Linux will only run on Windows with -opengl)
+	OutFormats.AddUnique(NAME_PCD3D_SM5);
+	OutFormats.AddUnique(NAME_PCD3D_SM4);
+#endif
 	OutFormats.AddUnique(NAME_GLSL_150);
 	OutFormats.AddUnique(NAME_GLSL_430);
 	OutFormats.AddUnique(NAME_GLSL_150_MAC);
diff --git a/Engine/Source/Programs/AutomationTool/DeploymentContext.cs b/Engine/Source/Programs/AutomationTool/DeploymentContext.cs
--- a/Engine/Source/Programs/AutomationTool/DeploymentContext.cs
+++ b/Engine/Source/Programs/AutomationTool/DeploymentContext.cs
@@ -372,6 +372,12 @@ public class DeploymentContext //: ProjectParams
 		// Stage all the build products needed at runtime
 		foreach(BuildProduct BuildProduct in Receipt.BuildProducts)
 		{
+			// allow missing files if needed
+			if (Receipt.bRequireDependenciesToExist == false && File.Exists(BuildProduct.Path) == false)
+			{
+				continue;
+			}
+
 			if(BuildProduct.Type == BuildProductType.Executable || BuildProduct.Type == BuildProductType.DynamicLibrary || BuildProduct.Type == BuildProductType.RequiredResource)
 			{
 				StageFile(StagedFileType.NonUFS, BuildProduct.Path);
@@ -385,6 +391,12 @@ public class DeploymentContext //: ProjectParams
 		// Also stage any additional runtime dependencies, like ThirdParty DLLs
 		foreach(RuntimeDependency RuntimeDependency in Receipt.RuntimeDependencies)
 		{
+			// allow missing files if needed
+			if (Receipt.bRequireDependenciesToExist == false && File.Exists(RuntimeDependency.Path) == false)
+			{
+				continue;
+			}
+
 			StageFile(StagedFileType.NonUFS, RuntimeDependency.Path, RuntimeDependency.StagePath);
 		}
 	}
diff --git a/Engine/Source/Programs/AutomationTool/Mac/MacPlatform.Automation.cs b/Engine/Source/Programs/AutomationTool/Mac/MacPlatform.Automation.cs
--- a/Engine/Source/Programs/AutomationTool/Mac/MacPlatform.Automation.cs
+++ b/Engine/Source/Programs/AutomationTool/Mac/MacPlatform.Automation.cs
@@ -239,7 +239,12 @@ public class MacPlatform : Platform
 			File.WriteAllText(InfoPlistPath, InfoPlistContents);
 		}
 
-		Directory.CreateDirectory(CombinePaths(TargetPath, ExeName.StartsWith("UE4Game") ? "Engine" : SC.ShortProjectName, "Binaries", "Mac"));
+		if (!SC.bIsCombiningMultiplePlatforms)
+		{
+			// creating this directory when the content isn't moved into the application causes it 
+			// to fail to load, and isn't needed
+			Directory.CreateDirectory(CombinePaths(TargetPath, ExeName.StartsWith("UE4Game") ? "Engine" : SC.ShortProjectName, "Binaries", "Mac"));
+		}
 	}
 
 	public override ProcessResult RunClient(ERunOptions ClientRunFlags, string ClientApp, string ClientCmdLine, ProjectParams Params)
diff --git a/Engine/Source/Programs/AutomationTool/Scripts/CopyBuildToStagingDirectory.Automation.cs b/Engine/Source/Programs/AutomationTool/Scripts/CopyBuildToStagingDirectory.Automation.cs
--- a/Engine/Source/Programs/AutomationTool/Scripts/CopyBuildToStagingDirectory.Automation.cs
+++ b/Engine/Source/Programs/AutomationTool/Scripts/CopyBuildToStagingDirectory.Automation.cs
@@ -1369,17 +1369,40 @@ public partial class Project : CommandUtils
 				{
 					string ReceiptBaseDir = Params.IsCodeBasedProject? Path.GetDirectoryName(Params.RawProjectPath) : EngineDir;
 
-					// Read the receipt
-					string ReceiptFileName = BuildReceipt.GetDefaultPath(ReceiptBaseDir, Target, StagePlatform, Config, "");
-					if(!File.Exists(ReceiptFileName))
+					bool bRequireStagedFilesToExist = true;
+					UnrealTargetPlatform[] ReceiptPlatforms = { StagePlatform };
+					// @todo josha: This should ask the platform somehow instead of being hardcoded
+					if (StagePlatform == UnrealTargetPlatform.Desktop)
 					{
-						throw new AutomationException("Missing receipt '{0}'. Check that this target has been built.", Path.GetFileName(ReceiptFileName));
+						ReceiptPlatforms = new UnrealTargetPlatform[] { UnrealTargetPlatform.Win32, UnrealTargetPlatform.Win64, UnrealTargetPlatform.Mac, UnrealTargetPlatform.Linux };
+						bRequireStagedFilesToExist = false;
 					}
 
-					// Convert the paths to absolute
-					BuildReceipt Receipt = BuildReceipt.Read(ReceiptFileName);
-					Receipt.ExpandPathVariables(EngineDir, Path.GetDirectoryName(Params.RawProjectPath));
-					TargetsToStage.Add(Receipt);
+					foreach (UnrealTargetPlatform ReceiptPlatform in ReceiptPlatforms)
+					{
+						// Read the receipt
+						string ReceiptFileName = BuildReceipt.GetDefaultPath(ReceiptBaseDir, Target, ReceiptPlatform, Config, "");
+						if(!File.Exists(ReceiptFileName))
+						{
+							if (bRequireStagedFilesToExist)
+							{
+								// if we aren't collecting multiple platforms, then it is expected to exist
+								throw new AutomationException("Missing receipt '{0}'. Check that this target has been built.", Path.GetFileName(ReceiptFileName));
+							}
+							else
+							{
+								// if it's multiple platforms, then allow missing receipts
+								continue;
+							}
+
+						}
+
+						// Convert the paths to absolute
+						BuildReceipt Receipt = BuildReceipt.Read(ReceiptFileName);
+						Receipt.ExpandPathVariables(EngineDir, Path.GetDirectoryName(Params.RawProjectPath));
+						Receipt.SetDependenciesToBeRequired(bRequireStagedFilesToExist);
+						TargetsToStage.Add(Receipt);
+					}
 				}
 			}
 
diff --git a/Engine/Source/Programs/AutomationTool/Win/WinPlatform.Automation.cs b/Engine/Source/Programs/AutomationTool/Win/WinPlatform.Automation.cs
--- a/Engine/Source/Programs/AutomationTool/Win/WinPlatform.Automation.cs
+++ b/Engine/Source/Programs/AutomationTool/Win/WinPlatform.Automation.cs
@@ -52,14 +52,18 @@ public abstract class BaseWinPlatform : Platform
 				BuildProduct Executable = Receipt.BuildProducts.FirstOrDefault(x => x.Type == BuildProductType.Executable);
 				if(Executable != null)
 				{
-					string BootstrapArguments = "";
-					if(!SC.IsCodeBasedProject && !ShouldStageCommandLine(Params, SC))
+					// only create bootstraps for executables
+					if (SC.NonUFSStagingFiles.ContainsKey(Executable.Path) && Path.GetExtension(Executable.Path) == ".exe")
 					{
-						BootstrapArguments = String.Format("..\\..\\..\\{0}\\{0}.uproject", SC.ShortProjectName);
-					}
+						string BootstrapArguments = "";
+						if (!SC.IsCodeBasedProject && !ShouldStageCommandLine(Params, SC))
+						{
+							BootstrapArguments = String.Format("..\\..\\..\\{0}\\{0}.uproject", SC.ShortProjectName);
+						}
 
-					string BootstrapExeName = (SC.StageTargetConfigurations.Count == 1)? (Receipt.GetProperty("TargetName", SC.ShortProjectName) + ".exe") : Path.GetFileName(Executable.Path);
-					StageBootstrapExecutable(SC, BootstrapExeName, Executable.Path, SC.NonUFSStagingFiles[Executable.Path], BootstrapArguments);
+						string BootstrapExeName = (SC.StageTargetConfigurations.Count == 1) ? (Receipt.GetProperty("TargetName", SC.ShortProjectName) + ".exe") : Path.GetFileName(Executable.Path);
+						StageBootstrapExecutable(SC, BootstrapExeName, Executable.Path, SC.NonUFSStagingFiles[Executable.Path], BootstrapArguments);
+					}
 				}
 			}
 		}
diff --git a/Engine/Source/Programs/UnrealBuildTool/System/BuildReceipt.cs b/Engine/Source/Programs/UnrealBuildTool/System/BuildReceipt.cs
--- a/Engine/Source/Programs/UnrealBuildTool/System/BuildReceipt.cs
+++ b/Engine/Source/Programs/UnrealBuildTool/System/BuildReceipt.cs
@@ -126,6 +126,9 @@ namespace UnrealBuildTool
 		[XmlArrayItem("RuntimeDependency")]
 		public List<RuntimeDependency> RuntimeDependencies = new List<RuntimeDependency>();
 
+		// if packaging in a mode where some files aren't required, set this to false
+		public bool bRequireDependenciesToExist = true;
+
 		/// <summary>
 		/// Default constructor
 		/// </summary>
@@ -240,6 +243,15 @@ namespace UnrealBuildTool
 			ExpandPathVariables(EngineDir, ProjectDir, new Dictionary<string, string>());
 		}
 
+		/// <summary>
+		/// Control whether the dependencies are required while staging them
+		/// </summary>
+		/// <param name="InDependenciesAreRequired"></param>
+		public void SetDependenciesToBeRequired(bool InDependenciesAreRequired)
+		{
+			bRequireDependenciesToExist = InDependenciesAreRequired;
+		}
+
 		/// <summary>
 		/// Expand all the path variables in the manifest, including a list of supplied variable values.
 		/// </summary>
