Rename AllPlugins -> AvailablePlugins.

diff --git a/Engine/Source/Programs/AutomationTool/ProjectUtils.cs b/Engine/Source/Programs/AutomationTool/ProjectUtils.cs
--- a/Engine/Source/Programs/AutomationTool/ProjectUtils.cs
+++ b/Engine/Source/Programs/AutomationTool/ProjectUtils.cs
@@ -190,7 +190,7 @@ namespace AutomationTool
 
 			// Read the project descriptor, and find all the plugins available to this project
 			ProjectDescriptor Project = ProjectDescriptor.FromFile(RawProjectPath);
-			List<PluginInfo> AllPlugins = Plugins.ReadAvailablePlugins(RawProjectPath);
+			List<PluginInfo> AvailablePlugins = Plugins.ReadAvailablePlugins(RawProjectPath);
 
 			// check the target platforms for any differences in build settings or additional plugins
 			bool RetVal = false;
@@ -204,7 +204,7 @@ namespace AutomationTool
 				}
 
 				// find if there are any plugins enabled or disabled which differ from the default
-				foreach(PluginInfo Plugin in AllPlugins)
+				foreach(PluginInfo Plugin in AvailablePlugins)
 				{
 					if(UProjectInfo.IsPluginEnabledForProject(Plugin, Project, TargetPlatformType) != Plugin.Descriptor.bEnabledByDefault)
 					{
