build script, remove pseudodeps from fort promotable in fort branch

diff --git a/Engine/Source/Programs/AutomationTool/Scripts/GUBP.Automation.cs b/Engine/Source/Programs/AutomationTool/Scripts/GUBP.Automation.cs
--- a/Engine/Source/Programs/AutomationTool/Scripts/GUBP.Automation.cs
+++ b/Engine/Source/Programs/AutomationTool/Scripts/GUBP.Automation.cs
@@ -3827,6 +3827,15 @@ public class GUBP : BuildCommand
         GUBPNodes[Node].RemovePseudodependency(Dep);
     }
 
+    public void RemoveAllPseudodependenciesFromNode(string Node)
+    {
+        if (!GUBPNodes.ContainsKey(Node))
+        {
+            throw new AutomationException("Node {0} not found", Node);
+        }
+        GUBPNodes[Node].FullNamesOfPseudosependencies.Clear();
+    }
+
     List<string> GetDependencies(string NodeToDo, bool bFlat = false, bool ECOnly = false)
     {
         var Result = new List<string>();
