Make sure RootEditor and GameEditor nodes both specify the -precompile option. UBT only compiles all plugins when -precompile is specified now, and compiling game editors without it triggers different shared-pch build environments, which causes a significant build and modification of previous build products.

diff --git a/Engine/Source/Programs/AutomationTool/Scripts/GUBP.Automation.cs b/Engine/Source/Programs/AutomationTool/Scripts/GUBP.Automation.cs
--- a/Engine/Source/Programs/AutomationTool/Scripts/GUBP.Automation.cs
+++ b/Engine/Source/Programs/AutomationTool/Scripts/GUBP.Automation.cs
@@ -1458,7 +1458,7 @@ public class GUBP : BuildCommand
         {
             var Agenda = new UE4Build.BuildAgenda();
 
-            string Args = "-nobuilduht -skipactionhistory -skipnonhostplatforms -CopyAppBundleBackToDevice -forceheadergeneration";
+            string Args = "-nobuilduht -skipactionhistory -skipnonhostplatforms -CopyAppBundleBackToDevice -forceheadergeneration -precompile";
 
 			foreach(BranchInfo.BranchUProject GameProj in GameProjects)
 			{
