UAT: Allow processing symlinks on non-Linux platforms.

- The original reason to skip them (UEPLAT-184) was because symlinks were only used to fix case-sensitiveness issues and by skipping them
   we avoid name collision with the original file during packaging.
- Those issues are non-existent on Windows, while skipping the symlinks breaks packaging on deduplicated drives there
  (see https://answers.unrealengine.com/questions/212888/automated-buildjenkins-failing-due-to-symlink-chec.html).

diff --git a/Engine/Source/Programs/AutomationTool/Utils.cs b/Engine/Source/Programs/AutomationTool/Utils.cs
--- a/Engine/Source/Programs/AutomationTool/Utils.cs
+++ b/Engine/Source/Programs/AutomationTool/Utils.cs
@@ -458,24 +458,35 @@ namespace AutomationTool
 				Log.WriteLine(TraceEventType.Information, "FindFiles {0} {1} {2}", Path, SearchPattern, Recursive);
 			}
 
-			// filter out symlinks which can cause problems at the later stage
-			List<string> FileNames = new List<string>();
-			DirectoryInfo DirInfo = new DirectoryInfo(Path);
-			foreach( FileInfo File in DirInfo.EnumerateFiles(SearchPattern, Recursive ? SearchOption.AllDirectories : SearchOption.TopDirectoryOnly))
+			// On Linux, filter out symlinks since we (usually) create them to fix mispelled case-sensitive filenames in content, and if they aren't filtered, 
+			// UAT picks up both the symlink and the original file and considers them duplicates when packaging (pak files are case-insensitive).
+			// Windows needs the symlinks though because that's how deduplication works on Windows server, 
+			// see https://answers.unrealengine.com/questions/212888/automated-buildjenkins-failing-due-to-symlink-chec.html
+			// FIXME: ZFS, JFS and other fs that can be case-insensitive on Linux should use the faster path as well.
+			if (UnrealBuildTool.BuildHostPlatform.Current.Platform != UnrealTargetPlatform.Linux)
 			{
-				if (File.Attributes.HasFlag(FileAttributes.ReparsePoint))
+				return Directory.GetFiles(Path, SearchPattern, Recursive ? SearchOption.AllDirectories : SearchOption.TopDirectoryOnly);
+			}
+			else
+			{
+				List<string> FileNames = new List<string>();
+				DirectoryInfo DirInfo = new DirectoryInfo(Path);
+				foreach( FileInfo File in DirInfo.EnumerateFiles(SearchPattern, Recursive ? SearchOption.AllDirectories : SearchOption.TopDirectoryOnly))
 				{
-					if (!bQuiet)
+					if (File.Attributes.HasFlag(FileAttributes.ReparsePoint))
 					{
-						Log.WriteLine(TraceEventType.Warning, "Ignoring symlink {0}", File.FullName);
+						if (!bQuiet)
+						{
+							Log.WriteLine(TraceEventType.Warning, "Ignoring symlink {0}", File.FullName);
+						}
+						continue;
 					}
-					continue;
+					
+					FileNames.Add(File.FullName);
 				}
 				
-				FileNames.Add(File.FullName);
+				return FileNames.ToArray();
 			}
-			
-			return FileNames.ToArray();
 		}
 
 		/// <summary>
