- Handle argv vs UE4CommandLine.txt conflicts by skipping over multiple .uproject tokens [UE-14121]
- Cleaned up some some dirty code [UE-14120]

diff --git a/Engine/Source/Programs/AutomationTool/Desktop/DesktopPlatform.Automation.cs b/Engine/Source/Programs/AutomationTool/Desktop/DesktopPlatform.Automation.cs
--- a/Engine/Source/Programs/AutomationTool/Desktop/DesktopPlatform.Automation.cs
+++ b/Engine/Source/Programs/AutomationTool/Desktop/DesktopPlatform.Automation.cs
@@ -23,18 +23,23 @@ public class DesktopPlatform : Platform
 		return false;
 	}
 
-	private static List<UnrealTargetPlatform> DesktopPlatforms = new List<UnrealTargetPlatform>() { 
-		UnrealTargetPlatform.Win32,
-		UnrealTargetPlatform.Win64,
-		UnrealTargetPlatform.Mac,
-		UnrealTargetPlatform.Linux,
-	};
+
+	public override UnrealTargetPlatform[] GetStagePlatforms()
+	{
+		return new UnrealTargetPlatform[] 
+		{
+			UnrealTargetPlatform.Win32,
+//			UnrealTargetPlatform.Win64,
+			UnrealTargetPlatform.Mac,
+			UnrealTargetPlatform.Linux,
+		};
+	}
 
 	public override void Package(ProjectParams Params, DeploymentContext SC, int WorkingCL)
 	{
 		SC.bIsCombiningMultiplePlatforms = true;
 		string SavedPlatformDir = SC.PlatformDir;
-		foreach (UnrealTargetPlatform DesktopPlatform in DesktopPlatforms)
+		foreach (UnrealTargetPlatform DesktopPlatform in GetStagePlatforms())
 		{
 			Platform SubPlatform = Platform.Platforms[DesktopPlatform];
 			SC.PlatformDir = DesktopPlatform.ToString();
@@ -48,7 +53,7 @@ public class DesktopPlatform : Platform
 	{
 		SC.bIsCombiningMultiplePlatforms = true;
 		string SavedPlatformDir = SC.PlatformDir;
-		foreach (UnrealTargetPlatform DesktopPlatform in DesktopPlatforms)
+		foreach (UnrealTargetPlatform DesktopPlatform in GetStagePlatforms())
 		{
 			Platform SubPlatform = Platform.Platforms[DesktopPlatform];
 			SC.PlatformDir = DesktopPlatform.ToString();
@@ -72,7 +77,7 @@ public class DesktopPlatform : Platform
 	{
 		SC.bIsCombiningMultiplePlatforms = true;
 		string SavedPlatformDir = SC.PlatformDir;
-		foreach (UnrealTargetPlatform DesktopPlatform in DesktopPlatforms)
+		foreach (UnrealTargetPlatform DesktopPlatform in GetStagePlatforms())
 		{
 			Platform SubPlatform = Platform.Platforms[DesktopPlatform];
 			SC.PlatformDir = DesktopPlatform.ToString();
@@ -88,7 +93,7 @@ public class DesktopPlatform : Platform
 
 		SC.bIsCombiningMultiplePlatforms = true;
 		string SavedPlatformDir = SC.PlatformDir;
-		foreach (UnrealTargetPlatform DesktopPlatform in DesktopPlatforms)
+		foreach (UnrealTargetPlatform DesktopPlatform in GetStagePlatforms())
 		{
 			Platform SubPlatform = Platform.Platforms[DesktopPlatform];
 			SC.PlatformDir = DesktopPlatform.ToString();
diff --git a/Engine/Source/Programs/AutomationTool/Platform.cs b/Engine/Source/Programs/AutomationTool/Platform.cs
--- a/Engine/Source/Programs/AutomationTool/Platform.cs
+++ b/Engine/Source/Programs/AutomationTool/Platform.cs
@@ -508,6 +508,11 @@ namespace AutomationTool
 			return false;
 		}
 
+		public virtual UnrealTargetPlatform[] GetStagePlatforms()
+		{
+			return new UnrealTargetPlatform[] { PlatformType };
+		}
+
 		#endregion
 
 		#region Utilities
diff --git a/Engine/Source/Programs/AutomationTool/Scripts/CopyBuildToStagingDirectory.Automation.cs b/Engine/Source/Programs/AutomationTool/Scripts/CopyBuildToStagingDirectory.Automation.cs
--- a/Engine/Source/Programs/AutomationTool/Scripts/CopyBuildToStagingDirectory.Automation.cs
+++ b/Engine/Source/Programs/AutomationTool/Scripts/CopyBuildToStagingDirectory.Automation.cs
@@ -1380,16 +1380,13 @@ public partial class Project : CommandUtils
 				{
 					string ReceiptBaseDir = Params.IsCodeBasedProject? Path.GetDirectoryName(Params.RawProjectPath) : EngineDir;
 
-					bool bRequireStagedFilesToExist = true;
-					UnrealTargetPlatform[] ReceiptPlatforms = { StagePlatform };
-					// @todo josha: This should ask the platform somehow instead of being hardcoded
-					if (StagePlatform == UnrealTargetPlatform.Desktop)
-					{
-						ReceiptPlatforms = new UnrealTargetPlatform[] { UnrealTargetPlatform.Win32, UnrealTargetPlatform.Win64, UnrealTargetPlatform.Mac, UnrealTargetPlatform.Linux };
-						bRequireStagedFilesToExist = false;
-					}
+					Platform PlatformInstance = Platform.Platforms[StagePlatform];
+					UnrealTargetPlatform[] SubPlatformsToStage = PlatformInstance.GetStagePlatforms();
+
+					// if we are attempting to gathering multiple platforms, the files aren't required
+					bool bRequireStagedFilesToExist = SubPlatformsToStage.Length == 1;
 
-					foreach (UnrealTargetPlatform ReceiptPlatform in ReceiptPlatforms)
+					foreach (UnrealTargetPlatform ReceiptPlatform in SubPlatformsToStage)
 					{
 						// Read the receipt
 						string ReceiptFileName = BuildReceipt.GetDefaultPath(ReceiptBaseDir, Target, ReceiptPlatform, Config, "");
diff --git a/Engine/Source/Runtime/Launch/Private/LaunchEngineLoop.cpp b/Engine/Source/Runtime/Launch/Private/LaunchEngineLoop.cpp
--- a/Engine/Source/Runtime/Launch/Private/LaunchEngineLoop.cpp
+++ b/Engine/Source/Runtime/Launch/Private/LaunchEngineLoop.cpp
@@ -861,6 +861,15 @@ int32 FEngineLoop::PreInit( const TCHAR* CmdLine )
 		Token = FParse::Token( ParsedCmdLine, 0);
 		Token = Token.Trim();
 
+		// if the next token is a project file, then we skip it (which can happen on some platforms that combine
+		// commandlines... this handles extra .uprojects, but if you run with MyGame MyGame, we can't tell if
+		// the second MyGame is a map or not)
+		while (FPaths::GetExtension(Token) == FProjectDescriptor::GetExtension())
+		{
+			Token = FParse::Token(ParsedCmdLine, 0);
+			Token = Token.Trim();
+		}
+
 		if (bFirstTokenIsGameProjectFilePath || bFirstTokenIsGameProjectFileShortName)
 		{
 			// Convert it to relative if possible...
