UAT: fix staging of libLND on Linux.

diff --git a/Engine/Source/Programs/AutomationTool/Linux/LinuxPlatform.Automation.cs b/Engine/Source/Programs/AutomationTool/Linux/LinuxPlatform.Automation.cs
--- a/Engine/Source/Programs/AutomationTool/Linux/LinuxPlatform.Automation.cs
+++ b/Engine/Source/Programs/AutomationTool/Linux/LinuxPlatform.Automation.cs
@@ -25,13 +25,16 @@ public abstract class BaseLinuxPlatform : Platform
 
 	public override void GetFilesToDeployOrStage(ProjectParams Params, DeploymentContext SC)
 	{
+		// FIXME: use build architecture
+		string BuildArchitecture = "x86_64-unknown-linux-gnu";
+
 		if (SC.bStageCrashReporter)
 		{
 			SC.StageFiles(StagedFileType.NonUFS, CombinePaths(SC.LocalRoot, "Engine/Binaries", SC.PlatformDir), "CrashReportClient", false, null, null, true);
 		}
 
 		{
-			SC.StageFiles(StagedFileType.NonUFS, CombinePaths(SC.LocalRoot, "Engine/Binaries/ThirdParty/ICU/icu4c-53_1/", SC.PlatformDir, "x86_64-unknown-linux-gnu"), Params.bDebugBuildsActuallyUseDebugCRT ? "*d.so*" : "*.so*", false, new[] { Params.bDebugBuildsActuallyUseDebugCRT ? "*.so*" : "*d.so*" }, CombinePaths("Engine/Binaries", SC.PlatformDir));
+			SC.StageFiles(StagedFileType.NonUFS, CombinePaths(SC.LocalRoot, "Engine/Binaries/ThirdParty/ICU/icu4c-53_1/", SC.PlatformDir, BuildArchitecture), Params.bDebugBuildsActuallyUseDebugCRT ? "*d.so*" : "*.so*", false, new[] { Params.bDebugBuildsActuallyUseDebugCRT ? "*.so*" : "*d.so*" }, CombinePaths("Engine/Binaries", SC.PlatformDir));
 		}
 
 		// assume that we always have to deploy Steam (FIXME: should be automatic)
@@ -51,7 +54,7 @@ public abstract class BaseLinuxPlatform : Platform
 
 		// stage libLND
 		{
-			SC.StageFiles(StagedFileType.NonUFS, CombinePaths(SC.LocalRoot, "Engine/Binaries", SC.PlatformDir), "libLND*.so");
+			SC.StageFiles(StagedFileType.NonUFS, CombinePaths(SC.LocalRoot, "Engine/Binaries/ThirdParty/LinuxNativeDialogs/", SC.PlatformDir, BuildArchitecture), "libLND*.so");
 		}
 
 		// assume that we always have to deploy OpenAL (FIXME: should be automatic)
