Add a node called "Shared_LabelPromotableAggregate", which allows tacking Rocket publish steps on to the promotable.

diff --git a/Engine/Source/Programs/AutomationTool/Scripts/GUBP.Automation.cs b/Engine/Source/Programs/AutomationTool/Scripts/GUBP.Automation.cs
--- a/Engine/Source/Programs/AutomationTool/Scripts/GUBP.Automation.cs
+++ b/Engine/Source/Programs/AutomationTool/Scripts/GUBP.Automation.cs
@@ -2534,6 +2534,24 @@ public class GUBP : BuildCommand
         }
     }
 
+    public class SharedLabelPromotableSuccessNode : AggregateNode
+    {        
+        public SharedLabelPromotableSuccessNode()
+        {
+			AddDependency(SharedLabelPromotableNode.StaticGetFullName(false));
+        }
+
+        public static string StaticGetFullName()
+        {
+            return SharedLabelPromotableNode.StaticGetFullName(false) + "Aggregate";
+        }
+
+		public override string GetFullName()
+		{
+			return StaticGetFullName();
+		}
+    }
+
     public class WaitForTestShared : AggregateNode
     {
         public WaitForTestShared(GUBP bp)
@@ -6232,6 +6250,7 @@ public class GUBP : BuildCommand
             AddNode(new SharedAggregatePromotableNode(this, HostPlatforms));
             AddNode(new WaitForSharedPromotionUserInput(this, false));
             AddNode(new SharedLabelPromotableNode(this, false));
+            AddNode(new SharedLabelPromotableSuccessNode());
             AddNode(new WaitForTestShared(this));
             AddNode(new WaitForSharedPromotionUserInput(this, true));
             AddNode(new SharedLabelPromotableNode(this, true));
diff --git a/Engine/Source/Programs/AutomationTool/Scripts/RocketBuild.Automation.cs b/Engine/Source/Programs/AutomationTool/Scripts/RocketBuild.Automation.cs
--- a/Engine/Source/Programs/AutomationTool/Scripts/RocketBuild.Automation.cs
+++ b/Engine/Source/Programs/AutomationTool/Scripts/RocketBuild.Automation.cs
@@ -125,6 +125,11 @@ namespace Rocket
 						PromotableNode.AddDependency(BuildGitPromotable.StaticGetFullName(HostPlatform));
 					}
 				}
+
+				// Add a dependency on this being published as part of the shared promotable being labeled
+				GUBP.SharedLabelPromotableSuccessNode LabelPromotableNode = (GUBP.SharedLabelPromotableSuccessNode)bp.FindNode(GUBP.SharedLabelPromotableSuccessNode.StaticGetFullName());
+				LabelPromotableNode.AddDependency(CopyRocketNode.StaticGetFullName(HostPlatform));
+				LabelPromotableNode.AddDependency(CopyRocketSymbolsNode.StaticGetFullName(HostPlatform));
 			}
 		}
 
