Add separate nodes for stripping each target, so they can run in the same ASG instead of waiting for a gather step at the end (which causes particular contention for different source platforms).

diff --git a/Engine/Source/Programs/AutomationTool/Scripts/GUBP.Automation.cs b/Engine/Source/Programs/AutomationTool/Scripts/GUBP.Automation.cs
--- a/Engine/Source/Programs/AutomationTool/Scripts/GUBP.Automation.cs
+++ b/Engine/Source/Programs/AutomationTool/Scripts/GUBP.Automation.cs
@@ -1626,7 +1626,10 @@ public class GUBP : BuildCommand
                 AddPseudodependency(WaitForTestShared.StaticGetFullName());
                 AgentSharingGroup = "TemplateMonolithics" + StaticGetHostPlatformSuffix(InHostPlatform);
             }
-
+			else if (!bp.BranchOptions.bNoInstalledEngine && GameProj == bp.Branch.BaseEngineProject)
+			{
+				AgentSharingGroup = GameProj.GameName + "_MonolithicsGroup_" + InTargetPlatform + StaticGetHostPlatformSuffix(InHostPlatform);
+			}
         }
 
 		public override string GetDisplayGroupName()
diff --git a/Engine/Source/Programs/AutomationTool/Scripts/RocketBuild.Automation.cs b/Engine/Source/Programs/AutomationTool/Scripts/RocketBuild.Automation.cs
--- a/Engine/Source/Programs/AutomationTool/Scripts/RocketBuild.Automation.cs
+++ b/Engine/Source/Programs/AutomationTool/Scripts/RocketBuild.Automation.cs
@@ -87,12 +87,28 @@ namespace Rocket
 					}
 				}
 
+				// Get the temp directory for stripped files for this host
+				string StrippedDir = Path.GetFullPath(CommandUtils.CombinePaths(CommandUtils.CmdEnv.LocalRoot, "Engine", "Saved", "Rocket", HostPlatform.ToString()));
+
+				// Strip the host platform
+				if (StripRocketNode.IsRequiredForPlatform(HostPlatform))
+				{
+					bp.AddNode(new StripRocketToolsNode(HostPlatform, StrippedDir));
+					bp.AddNode(new StripRocketEditorNode(HostPlatform, StrippedDir));
+				}
+
+				// Strip all the target platforms
+				foreach (UnrealTargetPlatform TargetPlatform in TargetPlatforms)
+				{
+					if (GetSourceHostPlatform(bp, HostPlatform, TargetPlatform) == HostPlatform)
+					{
+						bp.AddNode(new StripRocketMonolithicsNode(HostPlatform, bp.Branch.BaseEngineProject, TargetPlatform, CodeTargetPlatforms.Contains(TargetPlatform), StrippedDir));
+					}
+				}
+
 				// Build the DDC
 				bp.AddNode(new BuildDerivedDataCacheNode(HostPlatform, GetCookPlatforms(HostPlatform, TargetPlatforms), CurrentFeaturePacks));
 
-				// Strip symbols for all the targets built on the host platform
-				bp.AddNode(new StripRocketNode(bp, HostPlatform, TargetPlatforms, CodeTargetPlatforms));
-
 				// Generate a list of files that needs to be copied for each target platform
 				bp.AddNode(new FilterRocketNode(bp, HostPlatform, TargetPlatforms, CodeTargetPlatforms, CurrentFeaturePacks, CurrentTemplates));
 
@@ -100,13 +116,13 @@ namespace Rocket
 				if(ShouldDoSeriousThingsLikeP4CheckinAndPostToMCP())
 				{
 					string OutputDir = CommandUtils.CombinePaths(CommandUtils.RootSharedTempStorageDirectory(), "Rocket", "Automated", GetBuildLabel(), CommandUtils.GetGenericPlatformName(HostPlatform));
-					bp.AddNode(new CopyRocketNode(HostPlatform, CodeTargetPlatforms, OutputDir, true));
+					bp.AddNode(new CopyRocketNode(HostPlatform, CodeTargetPlatforms, StrippedDir, OutputDir, true));
 					bp.AddNode(new CopyRocketSymbolsNode(bp, HostPlatform, CodeTargetPlatforms, OutputDir + "Symbols", true));
 				}
 				else
 				{
 					string OutputDir = CommandUtils.CombinePaths(CommandUtils.CmdEnv.LocalRoot, "LocalBuilds", "Rocket", CommandUtils.GetGenericPlatformName(HostPlatform));
-					bp.AddNode(new CopyRocketNode(HostPlatform, CodeTargetPlatforms, OutputDir, false));
+					bp.AddNode(new CopyRocketNode(HostPlatform, CodeTargetPlatforms, StrippedDir, OutputDir, false));
 					bp.AddNode(new CopyRocketSymbolsNode(bp, HostPlatform, CodeTargetPlatforms, OutputDir + "Symbols", false));
 				}
 
@@ -292,45 +308,32 @@ namespace Rocket
 			BuildProducts = new List<string>(PromotableFiles);
 		}
 	}
-	
-	public class StripRocketNode : GUBP.HostPlatformNode
+
+	public abstract class StripRocketNode : GUBP.HostPlatformNode
 	{
-		List<UnrealTargetPlatform> Platforms = new List<UnrealTargetPlatform>();
-		public readonly string StrippedDir;
+		public UnrealTargetPlatform TargetPlatform;
+		public string StrippedDir;
 
-		public StripRocketNode(GUBP bp, UnrealTargetPlatform InHostPlatform, IEnumerable<UnrealTargetPlatform> InTargetPlatforms, IEnumerable<UnrealTargetPlatform> InCodeTargetPlatforms) : base(InHostPlatform)
+		public StripRocketNode(UnrealTargetPlatform InHostPlatform, UnrealTargetPlatform InTargetPlatform, string InStrippedDir) : base(InHostPlatform)
 		{
-			Platforms = new List<UnrealTargetPlatform>(InTargetPlatforms);
-			if(!Platforms.Contains(HostPlatform))
-			{
-				Platforms.Add(HostPlatform);
-			}
-			StrippedDir = Path.GetFullPath(CommandUtils.CombinePaths(CommandUtils.CmdEnv.LocalRoot, "Engine", "Saved", "Rocket", HostPlatform.ToString()));
-
-			// Add all the host nodes
-			AddDependency(GUBP.ToolsForCompileNode.StaticGetFullName(HostPlatform));
-			AddDependency(GUBP.RootEditorNode.StaticGetFullName(HostPlatform));
-			AddDependency(GUBP.ToolsNode.StaticGetFullName(HostPlatform));
-
-			// Add all the target platform nodes
-			foreach(UnrealTargetPlatform TargetPlatform in InTargetPlatforms)
-			{
-				UnrealTargetPlatform SourceHostPlatform = RocketBuild.GetSourceHostPlatform(bp, HostPlatform, TargetPlatform);
-				if(SourceHostPlatform == HostPlatform)
-				{
-					AddDependency(GUBP.GamePlatformMonolithicsNode.StaticGetFullName(HostPlatform, bp.Branch.BaseEngineProject, TargetPlatform, Precompiled: InCodeTargetPlatforms.Contains(TargetPlatform)));
-				}
-			}
+			TargetPlatform = InTargetPlatform;
+			StrippedDir = InStrippedDir;
 		}
 
-		public static string StaticGetFullName(UnrealTargetPlatform HostPlatform)
-		{
-			return "StripRocket" + StaticGetHostPlatformSuffix(HostPlatform);
-		}
+		public override abstract string GetFullName();
 
-		public override string GetFullName()
+		public static bool IsRequiredForPlatform(UnrealTargetPlatform Platform)
 		{
-			return StaticGetFullName(HostPlatform);
+			switch(Platform)
+			{
+				case UnrealTargetPlatform.Win32:
+				case UnrealTargetPlatform.Win64:
+				case UnrealTargetPlatform.Android:
+				case UnrealTargetPlatform.Linux:
+					return true;
+				default:
+					return false;
+			}
 		}
 
 		public override void DoBuild(GUBP bp)
@@ -338,32 +341,29 @@ namespace Rocket
 			BuildProducts = new List<string>();
 
 			string InputDir = Path.GetFullPath(CommandUtils.CmdEnv.LocalRoot);
-
 			string RulesFileName = CommandUtils.CombinePaths(CommandUtils.CmdEnv.LocalRoot, "Engine", "Build", "InstalledEngineFilters.ini");
-			foreach(UnrealTargetPlatform Platform in Platforms)
-			{
-				// Read the filter for files on this platform
-				FileFilter StripFilter = new FileFilter();
-				StripFilter.ReadRulesFromFile(RulesFileName, "StripSymbols." + Platform.ToString(), HostPlatform.ToString());
 
-				// Apply the filter to the build products
-				List<string> SourcePaths = new List<string>();
-				List<string> TargetPaths = new List<string>();
-				foreach(string DependencyBuildProduct in AllDependencyBuildProducts)
+			// Read the filter for files on this platform
+			FileFilter StripFilter = new FileFilter();
+			StripFilter.ReadRulesFromFile(RulesFileName, "StripSymbols." + TargetPlatform.ToString(), HostPlatform.ToString());
+
+			// Apply the filter to the build products
+			List<string> SourcePaths = new List<string>();
+			List<string> TargetPaths = new List<string>();
+			foreach(string DependencyBuildProduct in AllDependencyBuildProducts)
+			{
+				string RelativePath = CommandUtils.StripBaseDirectory(Path.GetFullPath(DependencyBuildProduct), InputDir);
+				if(StripFilter.Matches(RelativePath))
 				{
-					string RelativePath = CommandUtils.StripBaseDirectory(Path.GetFullPath(DependencyBuildProduct), InputDir);
-					if(StripFilter.Matches(RelativePath))
-					{
-						SourcePaths.Add(CommandUtils.CombinePaths(InputDir, RelativePath));
-						TargetPaths.Add(CommandUtils.CombinePaths(StrippedDir, RelativePath));
-					}
+					SourcePaths.Add(CommandUtils.CombinePaths(InputDir, RelativePath));
+					TargetPaths.Add(CommandUtils.CombinePaths(StrippedDir, RelativePath));
 				}
-
-				// Strip the files and add them to the build products
-				StripSymbols(Platform, SourcePaths.ToArray(), TargetPaths.ToArray());
-				BuildProducts.AddRange(TargetPaths);
 			}
 
+			// Strip the files and add them to the build products
+			StripSymbols(TargetPlatform, SourcePaths.ToArray(), TargetPaths.ToArray());
+			BuildProducts.AddRange(TargetPaths);
+
 			SaveRecordOfSuccessAndAddToBuildProducts();
 		}
 
@@ -389,6 +389,78 @@ namespace Rocket
 		}
 	}
 
+	public class StripRocketToolsNode : StripRocketNode
+	{
+		public StripRocketToolsNode(UnrealTargetPlatform InHostPlatform, string InStrippedDir)
+			: base(InHostPlatform, InHostPlatform, InStrippedDir)
+		{
+			AddDependency(GUBP.ToolsForCompileNode.StaticGetFullName(HostPlatform));
+			AddDependency(GUBP.ToolsNode.StaticGetFullName(HostPlatform));
+			AgentSharingGroup = "ToolsGroup" + StaticGetHostPlatformSuffix(InHostPlatform);
+		}
+
+		public static string StaticGetFullName(UnrealTargetPlatform InHostPlatform)
+		{
+			return "Strip" + GUBP.ToolsNode.StaticGetFullName(InHostPlatform);
+		}
+
+		public override string GetFullName()
+		{
+			return StaticGetFullName(HostPlatform);
+		}
+	}
+
+	public class StripRocketEditorNode : StripRocketNode
+	{
+		public StripRocketEditorNode(UnrealTargetPlatform InHostPlatform, string InStrippedDir)
+			: base(InHostPlatform, InHostPlatform, InStrippedDir)
+		{
+			AddDependency(GUBP.RootEditorNode.StaticGetFullName(HostPlatform));
+			AgentSharingGroup = "Editor" + StaticGetHostPlatformSuffix(HostPlatform);
+		}
+
+		public static string StaticGetFullName(UnrealTargetPlatform InHostPlatform)
+		{
+			return "Strip" + GUBP.RootEditorNode.StaticGetFullName(InHostPlatform);
+		}
+
+		public override string GetFullName()
+		{
+			return StaticGetFullName(HostPlatform);
+		}
+	}
+
+	public class StripRocketMonolithicsNode : StripRocketNode
+	{
+		BranchInfo.BranchUProject Project;
+		bool bIsCodeTargetPlatform;
+
+		public StripRocketMonolithicsNode(UnrealTargetPlatform InHostPlatform, BranchInfo.BranchUProject InProject, UnrealTargetPlatform InTargetPlatform, bool bInIsCodeTargetPlatform, string InStrippedDir) : base(InHostPlatform, InTargetPlatform, InStrippedDir)
+		{
+			Project = InProject;
+			bIsCodeTargetPlatform = bInIsCodeTargetPlatform;
+
+			AddDependency(GUBP.GamePlatformMonolithicsNode.StaticGetFullName(HostPlatform, Project, InTargetPlatform, Precompiled: bIsCodeTargetPlatform));
+
+			AgentSharingGroup = Project.GameName + "_MonolithicsGroup_" + InTargetPlatform + StaticGetHostPlatformSuffix(InHostPlatform);
+		}
+
+		public static string StaticGetFullName(UnrealTargetPlatform InHostPlatform, BranchInfo.BranchUProject InProject, UnrealTargetPlatform InTargetPlatform, bool bIsCodeTargetPlatform)
+		{
+			string Name = InProject.GameName + "_" + InTargetPlatform + "_Mono";
+			if(bIsCodeTargetPlatform)
+			{
+				Name += "_Precompiled";
+			}
+			return Name + "_Strip" + StaticGetHostPlatformSuffix(InHostPlatform);
+		}
+
+		public override string GetFullName()
+		{
+			return StaticGetFullName(HostPlatform, Project, TargetPlatform, bIsCodeTargetPlatform);
+		}
+	}
+
 	public class FilterRocketNode : GUBP.HostPlatformNode
 	{
 		List<UnrealTargetPlatform> SourceHostPlatforms;
@@ -398,6 +470,7 @@ namespace Rocket
 		string[] CurrentTemplates;
 		public readonly string DepotManifestPath;
 		public readonly string StrippedManifestPath;
+		List<string> StrippedNodes = new List<string>();
 
 		public FilterRocketNode(GUBP bp, UnrealTargetPlatform InHostPlatform, List<UnrealTargetPlatform> InTargetPlatforms, List<UnrealTargetPlatform> InCodeTargetPlatforms, string[] InCurrentFeaturePacks, string[] InCurrentTemplates)
 			: base(InHostPlatform)
@@ -416,10 +489,18 @@ namespace Rocket
 			AddDependency(GUBP.ToolsNode.StaticGetFullName(HostPlatform));
 
 			// Get all the external file lists. Android/HTML5 on Mac uses the same files as Win64.
-			foreach(UnrealTargetPlatform CodeTargetPlatform in CodeTargetPlatforms)
+			foreach(UnrealTargetPlatform TargetPlatform in TargetPlatforms)
 			{
-				UnrealTargetPlatform SourceHostPlatform = RocketBuild.GetSourceHostPlatform(bp, HostPlatform, CodeTargetPlatform);
-				AddDependency(GUBP.GamePlatformMonolithicsNode.StaticGetFullName(SourceHostPlatform, bp.Branch.BaseEngineProject, CodeTargetPlatform, Precompiled: true));
+				UnrealTargetPlatform SourceHostPlatform = RocketBuild.GetSourceHostPlatform(bp, HostPlatform, TargetPlatform);
+				bool bIsCodeTargetPlatform = CodeTargetPlatforms.Contains(TargetPlatform);
+				AddDependency(GUBP.GamePlatformMonolithicsNode.StaticGetFullName(SourceHostPlatform, bp.Branch.BaseEngineProject, TargetPlatform, Precompiled: bIsCodeTargetPlatform));
+
+				if(StripRocketNode.IsRequiredForPlatform(TargetPlatform))
+				{
+					string StripNode = StripRocketMonolithicsNode.StaticGetFullName(SourceHostPlatform, bp.Branch.BaseEngineProject, TargetPlatform, bIsCodeTargetPlatform);
+					AddDependency(StripNode);
+					StrippedNodes.Add(StripNode);
+				}
 			}
 
 			// Add win64 tools on Mac, to get the win64 build of UBT, UAT and IPP
@@ -439,10 +520,14 @@ namespace Rocket
 				SourceHostPlatforms.Add(HostPlatform);
 			}
 
-			// Add the stripped target platforms
-			foreach(UnrealTargetPlatform SourceHostPlatform in SourceHostPlatforms)
+			// Add the stripped host platforms
+			if(StripRocketNode.IsRequiredForPlatform(HostPlatform))
 			{
-				AddDependency(StripRocketNode.StaticGetFullName(SourceHostPlatform));
+				AddDependency(StripRocketToolsNode.StaticGetFullName(HostPlatform));
+				StrippedNodes.Add(StripRocketToolsNode.StaticGetFullName(HostPlatform));
+
+				AddDependency(StripRocketEditorNode.StaticGetFullName(HostPlatform));
+				StrippedNodes.Add(StripRocketEditorNode.StaticGetFullName(HostPlatform));
 			}
 		}
 
@@ -524,9 +609,9 @@ namespace Rocket
 
 			// Run the filter on the stripped symbols, and remove those files from the copy filter
 			List<string> StrippedFiles = new List<string>();
-			foreach(UnrealTargetPlatform SourceHostPlatform in SourceHostPlatforms)
+			foreach(string StrippedNode in StrippedNodes)
 			{
-				StripRocketNode StripNode = (StripRocketNode)bp.FindNode(StripRocketNode.StaticGetFullName(SourceHostPlatform));
+				StripRocketNode StripNode = (StripRocketNode)bp.FindNode(StrippedNode);
 				foreach(string BuildProduct in StripNode.BuildProducts)
 				{
 					if(Utils.IsFileUnderDirectory(BuildProduct, StripNode.StrippedDir))
@@ -593,11 +678,13 @@ namespace Rocket
 
 	public class CopyRocketNode : GUBP.HostPlatformNode
 	{
+		public readonly string StrippedDir;
 		public readonly string OutputDir;
 		public List<UnrealTargetPlatform> CodeTargetPlatforms;
 
-		public CopyRocketNode(UnrealTargetPlatform HostPlatform, List<UnrealTargetPlatform> InCodeTargetPlatforms, string InOutputDir, bool bWaitForPromotion) : base(HostPlatform)
+		public CopyRocketNode(UnrealTargetPlatform HostPlatform, List<UnrealTargetPlatform> InCodeTargetPlatforms, string InStrippedDir, string InOutputDir, bool bWaitForPromotion) : base(HostPlatform)
 		{
+			StrippedDir = InStrippedDir;
 			OutputDir = InOutputDir;
 			CodeTargetPlatforms = new List<UnrealTargetPlatform>(InCodeTargetPlatforms);
 
@@ -612,7 +699,6 @@ namespace Rocket
 			}
 
 			AddDependency(FilterRocketNode.StaticGetFullName(HostPlatform));
-			AddDependency(StripRocketNode.StaticGetFullName(HostPlatform));
 			AddDependency(BuildDerivedDataCacheNode.StaticGetFullName(HostPlatform));
 		}
 
@@ -644,10 +730,7 @@ namespace Rocket
 			// Copy the depot files to the output directory
 			FilterRocketNode FilterNode = (FilterRocketNode)bp.FindNode(FilterRocketNode.StaticGetFullName(HostPlatform));
 			CopyManifestFilesToOutput(FilterNode.DepotManifestPath, CommandUtils.CmdEnv.LocalRoot, OutputDir);
-
-			// Copy the stripped symbols to the output directory
-			StripRocketNode StripNode = (StripRocketNode)bp.FindNode(StripRocketNode.StaticGetFullName(HostPlatform));
-			CopyManifestFilesToOutput(FilterNode.StrippedManifestPath, StripNode.StrippedDir, OutputDir);
+			CopyManifestFilesToOutput(FilterNode.StrippedManifestPath, StrippedDir, OutputDir);
 
 			// Copy the DDC to the output directory
 			BuildDerivedDataCacheNode DerivedDataCacheNode = (BuildDerivedDataCacheNode)bp.FindNode(BuildDerivedDataCacheNode.StaticGetFullName(HostPlatform));
