PR #1089: Fixed spelling in error message (Contributed by [REDACTED_PERSON])

diff --git a/Engine/Source/Programs/AutomationTool/ProjectUtils.cs b/Engine/Source/Programs/AutomationTool/ProjectUtils.cs
--- a/Engine/Source/Programs/AutomationTool/ProjectUtils.cs
+++ b/Engine/Source/Programs/AutomationTool/ProjectUtils.cs
@@ -278,7 +278,7 @@ namespace AutomationTool
 				// should never ask for engine targets if we can't compile
 				if (String.IsNullOrEmpty(RawProjectPath))
 				{
-					throw new AutomationException("Cannot dtermine engine targets if we can't compile.");
+					throw new AutomationException("Cannot determine engine targets if we can't compile.");
 				}
 
 				Properties.bIsCodeBasedProject = Properties.bWasGenerated;
