Group the strip nodes with monolithic nodes in the EC list.

diff --git a/Engine/Source/Programs/AutomationTool/Scripts/RocketBuild.Automation.cs b/Engine/Source/Programs/AutomationTool/Scripts/RocketBuild.Automation.cs
--- a/Engine/Source/Programs/AutomationTool/Scripts/RocketBuild.Automation.cs
+++ b/Engine/Source/Programs/AutomationTool/Scripts/RocketBuild.Automation.cs
@@ -445,6 +445,11 @@ namespace Rocket
 			AgentSharingGroup = Project.GameName + "_MonolithicsGroup_" + InTargetPlatform + StaticGetHostPlatformSuffix(InHostPlatform);
 		}
 
+		public override string GetDisplayGroupName()
+		{
+			return Project.GameName + "_Monolithics" + (bIsCodeTargetPlatform? "_Precompiled" : "");
+		}
+
 		public static string StaticGetFullName(UnrealTargetPlatform InHostPlatform, BranchInfo.BranchUProject InProject, UnrealTargetPlatform InTargetPlatform, bool bIsCodeTargetPlatform)
 		{
 			string Name = InProject.GameName + "_" + InTargetPlatform + "_Mono";
