Fix missing templates in installed builds

diff --git a/Engine/Build/InstalledEngineFilters.ini b/Engine/Build/InstalledEngineFilters.ini
--- a/Engine/Build/InstalledEngineFilters.ini
+++ b/Engine/Build/InstalledEngineFilters.ini
@@ -41,6 +41,15 @@
 ; Shaders
 /Engine/Shaders/...
 
+; Valid templates will be included by name from the build script, but we need to exclude intermediates
+-/Templates/*/Binaries/...
+-/Templates/*/Build/Receipts/...
+-/Templates/*/Intermediate/...
+-/Templates/*/DerivedDataCache/...
+-/Templates/*/Saved/...
+-/Templates/*/manifest.json
+-/Templates/*/contents.txt
+
 ; Win64-specific dependencies
 {Win64} /Engine/Binaries/ThirdParty/ImgTec/PVRTexToolCLI.exe
 
diff --git a/Engine/Source/Programs/AutomationTool/Scripts/RocketBuild.Automation.cs b/Engine/Source/Programs/AutomationTool/Scripts/RocketBuild.Automation.cs
--- a/Engine/Source/Programs/AutomationTool/Scripts/RocketBuild.Automation.cs
+++ b/Engine/Source/Programs/AutomationTool/Scripts/RocketBuild.Automation.cs
@@ -590,10 +590,7 @@ namespace Rocket
 			foreach (string Template in CurrentTemplates)
 			{
 				BranchInfo.BranchUProject Project = bp.Branch.FindGameChecked(Template);
-				string TemplateRelativeDirectory = "/" + Utils.StripBaseDirectory(Path.GetDirectoryName(Project.FilePath), CommandUtils.CmdEnv.LocalRoot) + "/";
-				Filter.Include(TemplateRelativeDirectory + "/...");
-				Filter.Exclude(TemplateRelativeDirectory + "/.../Binaries/...");
-				Filter.Exclude(TemplateRelativeDirectory + "/.../Intermediate/...");
+				Filter.Include("/" + Utils.StripBaseDirectory(Path.GetDirectoryName(Project.FilePath), CommandUtils.CmdEnv.LocalRoot).Replace('\\', '/') + "/...");
 			}
 
 			// Include all the standard rules
