Fixed localization automation script so that it doesn't upload the local, source controlled "FromOneSky" PO files.

diff --git a/Engine/Source/Programs/AutomationTool/Scripts/Localisation.Automation.cs b/Engine/Source/Programs/AutomationTool/Scripts/Localisation.Automation.cs
--- a/Engine/Source/Programs/AutomationTool/Scripts/Localisation.Automation.cs
+++ b/Engine/Source/Programs/AutomationTool/Scripts/Localisation.Automation.cs
@@ -250,14 +250,14 @@ class Localise : BuildCommand
         if (nativeCultureDirectory != null)
         {
             string[] files = Directory.GetFiles(nativeCultureDirectory.FullName, fileExtension, SearchOption.AllDirectories);
-            UploadFilesToProjectForLocale(project, files, nativeCultureDirectory.Name);
+            UploadFilesToProjectForLocale(project, files.Where((FilePath) => { return !FilePath.Contains("_FromOneSky"); }).ToArray(), nativeCultureDirectory.Name);
         }
 
         // Upload all other POs asynchronously.
         foreach (var foreignCultureDirectory in Array.FindAll(cultureDirectories, d => { return d != nativeCultureDirectory; }))
         {
             string[] files = Directory.GetFiles(foreignCultureDirectory.FullName, fileExtension, SearchOption.AllDirectories);
-            UploadFilesToProjectForLocale(project, files, foreignCultureDirectory.Name);
+            UploadFilesToProjectForLocale(project, files.Where((FilePath) => { return !FilePath.Contains("_FromOneSky");}).ToArray(), foreignCultureDirectory.Name);
         }
     }
 }
