Don't compute dependent frequencies until after the frequency barriers have been applied. Increase frequency for promotables in the release branches to 2h20m.

diff --git a/Engine/Source/Programs/AutomationTool/Scripts/GUBP.Automation.cs b/Engine/Source/Programs/AutomationTool/Scripts/GUBP.Automation.cs
--- a/Engine/Source/Programs/AutomationTool/Scripts/GUBP.Automation.cs
+++ b/Engine/Source/Programs/AutomationTool/Scripts/GUBP.Automation.cs
@@ -6358,10 +6358,6 @@ public class GUBP : BuildCommand
                     }
                 }
             }
-            foreach (var NodeToDo in GUBPNodes)
-            {
-                ComputeDependentCISFrequencyQuantumShift(NodeToDo.Key);
-            }
 
 			// Make sure that everything that's listed as a frequency barrier is completed with the given interval
 			foreach (KeyValuePair<string, sbyte> Barrier in BranchOptions.FrequencyBarriers)
@@ -6389,11 +6385,19 @@ public class GUBP : BuildCommand
 				{
 					if (IncludedNodes.Contains(NodePair.Key))
 					{
-						NodePair.Value.ComputedDependentCISFrequencyQuantumShift = Math.Min(NodePair.Value.ComputedDependentCISFrequencyQuantumShift, Barrier.Value);
+						if(NodePair.Value.ComputedDependentCISFrequencyQuantumShift == -1 || NodePair.Value.ComputedDependentCISFrequencyQuantumShift > Barrier.Value)
+						{
+							NodePair.Value.ComputedDependentCISFrequencyQuantumShift = Barrier.Value;
+						}
 					}
 				}
 			}
 
+            foreach (var NodeToDo in GUBPNodes)
+            {
+                ComputeDependentCISFrequencyQuantumShift(NodeToDo.Key);
+            }
+
             foreach (var NodeToDo in GUBPNodes)
             {
                 var Deps = GUBPNodes[NodeToDo.Key].DependentPromotions;
