Split Rocket publish steps up, so chunking is done from a build assembled locally without passing build products through GUBP, and the final copy to network happens downstream of everything else.

diff --git a/Engine/Source/Programs/AutomationTool/Scripts/RocketBuild.Automation.cs b/Engine/Source/Programs/AutomationTool/Scripts/RocketBuild.Automation.cs
--- a/Engine/Source/Programs/AutomationTool/Scripts/RocketBuild.Automation.cs
+++ b/Engine/Source/Programs/AutomationTool/Scripts/RocketBuild.Automation.cs
@@ -112,19 +112,9 @@ namespace Rocket
 				// Generate a list of files that needs to be copied for each target platform
 				bp.AddNode(new FilterRocketNode(bp, HostPlatform, TargetPlatforms, CodeTargetPlatforms, CurrentFeaturePacks, CurrentTemplates));
 
-				// Copy the install to the output directory. This node adds a dependency on triggering a shared promotable if it's copying to the network.
-				if(ShouldDoSeriousThingsLikeP4CheckinAndPostToMCP())
-				{
-					string OutputDir = CommandUtils.CombinePaths(CommandUtils.RootSharedTempStorageDirectory(), "Rocket", "Automated", GetBuildLabel(), CommandUtils.GetGenericPlatformName(HostPlatform));
-					bp.AddNode(new CopyRocketNode(HostPlatform, CodeTargetPlatforms, StrippedDir, OutputDir, true));
-					bp.AddNode(new CopyRocketSymbolsNode(bp, HostPlatform, CodeTargetPlatforms, OutputDir + "Symbols", true));
-				}
-				else
-				{
-					string OutputDir = CommandUtils.CombinePaths(CommandUtils.CmdEnv.LocalRoot, "LocalBuilds", "Rocket", CommandUtils.GetGenericPlatformName(HostPlatform));
-					bp.AddNode(new CopyRocketNode(HostPlatform, CodeTargetPlatforms, StrippedDir, OutputDir, false));
-					bp.AddNode(new CopyRocketSymbolsNode(bp, HostPlatform, CodeTargetPlatforms, OutputDir + "Symbols", false));
-				}
+				// Copy the install to the output directory
+				string LocalOutputDir = CommandUtils.CombinePaths(CommandUtils.CmdEnv.LocalRoot, "LocalBuilds", "Rocket", CommandUtils.GetGenericPlatformName(HostPlatform));
+				bp.AddNode(new CopyRocketNode(HostPlatform, CodeTargetPlatforms, StrippedDir, LocalOutputDir));
 
 				// Add the aggregate node for the entire install
 				GUBP.GUBPNode PromotableNode = bp.FindNode(GUBP.SharedAggregatePromotableNode.StaticGetFullName());
@@ -142,10 +132,15 @@ namespace Rocket
 					}
 				}
 
+				// Publish the install to the network
+				string RemoteOutputDir = CommandUtils.CombinePaths(CommandUtils.RootSharedTempStorageDirectory(), "Rocket", "Automated", GetBuildLabel(), CommandUtils.GetGenericPlatformName(HostPlatform));
+				bp.AddNode(new PublishRocketNode(HostPlatform, LocalOutputDir, RemoteOutputDir));
+				bp.AddNode(new PublishRocketSymbolsNode(bp, HostPlatform, CodeTargetPlatforms, RemoteOutputDir + "Symbols"));
+
 				// Add a dependency on this being published as part of the shared promotable being labeled
 				GUBP.SharedLabelPromotableSuccessNode LabelPromotableNode = (GUBP.SharedLabelPromotableSuccessNode)bp.FindNode(GUBP.SharedLabelPromotableSuccessNode.StaticGetFullName());
-				LabelPromotableNode.AddDependency(CopyRocketNode.StaticGetFullName(HostPlatform));
-				LabelPromotableNode.AddDependency(CopyRocketSymbolsNode.StaticGetFullName(HostPlatform));
+				LabelPromotableNode.AddDependency(PublishRocketNode.StaticGetFullName(HostPlatform));
+				LabelPromotableNode.AddDependency(PublishRocketSymbolsNode.StaticGetFullName(HostPlatform));
 			}
 		}
 
@@ -684,24 +679,17 @@ namespace Rocket
 		public readonly string OutputDir;
 		public List<UnrealTargetPlatform> CodeTargetPlatforms;
 
-		public CopyRocketNode(UnrealTargetPlatform HostPlatform, List<UnrealTargetPlatform> InCodeTargetPlatforms, string InStrippedDir, string InOutputDir, bool bWaitForPromotion) : base(HostPlatform)
+		public CopyRocketNode(UnrealTargetPlatform HostPlatform, List<UnrealTargetPlatform> InCodeTargetPlatforms, string InStrippedDir, string InOutputDir) : base(HostPlatform)
 		{
 			StrippedDir = InStrippedDir;
 			OutputDir = InOutputDir;
 			CodeTargetPlatforms = new List<UnrealTargetPlatform>(InCodeTargetPlatforms);
 
-			// If this is a CIS build, we need to wait for the promotable trigger before running. If not, make an agent sharing group so that the samples node can use our build products.
-			if(bWaitForPromotion)
-			{
-				AddDependency(GUBP.WaitForSharedPromotionUserInput.StaticGetFullName(false));
-			}
-			else
-			{
-				AgentSharingGroup = "RocketGroup" + StaticGetHostPlatformSuffix(HostPlatform);
-			}
-
+			AddDependency(GUBP.WaitForSharedPromotionUserInput.StaticGetFullName(false));
 			AddDependency(FilterRocketNode.StaticGetFullName(HostPlatform));
 			AddDependency(BuildDerivedDataCacheNode.StaticGetFullName(HostPlatform));
+
+			AgentSharingGroup = "RocketGroup" + StaticGetHostPlatformSuffix(HostPlatform);
 		}
 
 		public static string StaticGetFullName(UnrealTargetPlatform HostPlatform)
@@ -763,23 +751,57 @@ namespace Rocket
 		}
 	}
 
-	public class CopyRocketSymbolsNode : GUBP.HostPlatformNode
+	public class PublishRocketNode : GUBP.HostPlatformNode
 	{
-		string SymbolsOutputDir;
+		string LocalDir;
+		string PublishDir;
 
-		public CopyRocketSymbolsNode(GUBP bp, UnrealTargetPlatform HostPlatform, IEnumerable<UnrealTargetPlatform> TargetPlatforms, string InSymbolsOutputDir, bool bWaitForPromotion) : base(HostPlatform)
+		public PublishRocketNode(UnrealTargetPlatform HostPlatform, string InLocalDir, string InPublishDir) : base(HostPlatform)
 		{
-			SymbolsOutputDir = InSymbolsOutputDir;
+			LocalDir = InLocalDir;
+			PublishDir = InPublishDir;
 
-			if(bWaitForPromotion)
-			{
-				AddDependency(GUBP.WaitForSharedPromotionUserInput.StaticGetFullName(false));
-			}
-			else
+			AddDependency(CopyRocketNode.StaticGetFullName(HostPlatform));
+
+			AgentSharingGroup = "RocketGroup" + StaticGetHostPlatformSuffix(HostPlatform);
+		}
+
+		public static string StaticGetFullName(UnrealTargetPlatform HostPlatform)
+		{
+			return "PublishRocket" + StaticGetHostPlatformSuffix(HostPlatform);
+		}
+
+		public override string GetFullName()
+		{
+			return StaticGetFullName(HostPlatform);
+		}
+
+		public override float Priority()
+		{
+			return -10000.0f; // Do this and PublishRocketSymbolsNode at the end
+		}
+
+		public override void DoBuild(GUBP bp)
+		{
+			if(RocketBuild.ShouldDoSeriousThingsLikeP4CheckinAndPostToMCP())
 			{
-				AgentSharingGroup = "RocketGroup" + StaticGetHostPlatformSuffix(HostPlatform);
+				CommandUtils.ThreadedCopyFiles(LocalDir, PublishDir);
 			}
 
+			BuildProducts = new List<string>();
+			SaveRecordOfSuccessAndAddToBuildProducts();
+		}
+	}
+
+	public class PublishRocketSymbolsNode : GUBP.HostPlatformNode
+	{
+		string SymbolsOutputDir;
+
+		public PublishRocketSymbolsNode(GUBP bp, UnrealTargetPlatform HostPlatform, IEnumerable<UnrealTargetPlatform> TargetPlatforms, string InSymbolsOutputDir) : base(HostPlatform)
+		{
+			SymbolsOutputDir = InSymbolsOutputDir;
+
+			AddDependency(GUBP.WaitForSharedPromotionUserInput.StaticGetFullName(false));
 			AddDependency(GUBP.ToolsForCompileNode.StaticGetFullName(HostPlatform));
 			AddDependency(GUBP.RootEditorNode.StaticGetFullName(HostPlatform));
 			AddDependency(GUBP.ToolsNode.StaticGetFullName(HostPlatform));
@@ -789,11 +811,13 @@ namespace Rocket
 				UnrealTargetPlatform SourceHostPlatform = RocketBuild.GetSourceHostPlatform(bp, HostPlatform, TargetPlatform);
 				AddDependency(GUBP.GamePlatformMonolithicsNode.StaticGetFullName(SourceHostPlatform, bp.Branch.BaseEngineProject, TargetPlatform, Precompiled: true));
 			}
+
+			AgentSharingGroup = "RocketGroup" + StaticGetHostPlatformSuffix(HostPlatform);
 		}
 
 		public static string StaticGetFullName(UnrealTargetPlatform HostPlatform)
 		{
-			return "CopyRocketSymbols" + StaticGetHostPlatformSuffix(HostPlatform);
+			return "PublishRocketSymbols" + StaticGetHostPlatformSuffix(HostPlatform);
 		}
 
 		public override string GetFullName()
@@ -801,17 +825,25 @@ namespace Rocket
 			return StaticGetFullName(HostPlatform);
 		}
 
+		public override float Priority()
+		{
+			return -100000.0f; // Do this last; it's not mission critical
+		}
+
 		public override void DoBuild(GUBP bp)
 		{
-			// Make a lookup for all the known debug extensions, and filter all the dependency build products against that
-			HashSet<string> DebugExtensions = new HashSet<string>(Platform.Platforms.Values.SelectMany(x => x.GetDebugFileExtentions()).Distinct().ToArray(), StringComparer.InvariantCultureIgnoreCase);
-			foreach(string InputFileName in AllDependencyBuildProducts)
+			if(RocketBuild.ShouldDoSeriousThingsLikeP4CheckinAndPostToMCP())
 			{
-				string Extension = Path.GetExtension(InputFileName);
-				if(DebugExtensions.Contains(Extension))
+				// Make a lookup for all the known debug extensions, and filter all the dependency build products against that
+				HashSet<string> DebugExtensions = new HashSet<string>(Platform.Platforms.Values.SelectMany(x => x.GetDebugFileExtentions()).Distinct().ToArray(), StringComparer.InvariantCultureIgnoreCase);
+				foreach(string InputFileName in AllDependencyBuildProducts)
 				{
-					string OutputFileName = CommandUtils.MakeRerootedFilePath(InputFileName, CommandUtils.CmdEnv.LocalRoot, SymbolsOutputDir);
-					CommandUtils.CopyFile(InputFileName, OutputFileName);
+					string Extension = Path.GetExtension(InputFileName);
+					if(DebugExtensions.Contains(Extension))
+					{
+						string OutputFileName = CommandUtils.MakeRerootedFilePath(InputFileName, CommandUtils.CmdEnv.LocalRoot, SymbolsOutputDir);
+						CommandUtils.CopyFile(InputFileName, OutputFileName);
+					}
 				}
 			}
 
