Back out changelist 2524890

diff --git a/Engine/Source/Programs/AutomationTool/Scripts/RocketBuild.Automation.cs b/Engine/Source/Programs/AutomationTool/Scripts/RocketBuild.Automation.cs
--- a/Engine/Source/Programs/AutomationTool/Scripts/RocketBuild.Automation.cs
+++ b/Engine/Source/Programs/AutomationTool/Scripts/RocketBuild.Automation.cs
@@ -112,9 +112,19 @@ namespace Rocket
 				// Generate a list of files that needs to be copied for each target platform
 				bp.AddNode(new FilterRocketNode(bp, HostPlatform, TargetPlatforms, CodeTargetPlatforms, CurrentFeaturePacks, CurrentTemplates));
 
-				// Copy the install to the output directory
-				string LocalOutputDir = CommandUtils.CombinePaths(CommandUtils.CmdEnv.LocalRoot, "LocalBuilds", "Rocket", CommandUtils.GetGenericPlatformName(HostPlatform));
-				bp.AddNode(new CopyRocketNode(HostPlatform, CodeTargetPlatforms, StrippedDir, LocalOutputDir));
+				// Copy the install to the output directory. This node adds a dependency on triggering a shared promotable if it's copying to the network.
+				if(ShouldDoSeriousThingsLikeP4CheckinAndPostToMCP())
+				{
+					string OutputDir = CommandUtils.CombinePaths(CommandUtils.RootSharedTempStorageDirectory(), "Rocket", "Automated", GetBuildLabel(), CommandUtils.GetGenericPlatformName(HostPlatform));
+					bp.AddNode(new CopyRocketNode(HostPlatform, CodeTargetPlatforms, StrippedDir, OutputDir, true));
+					bp.AddNode(new CopyRocketSymbolsNode(bp, HostPlatform, CodeTargetPlatforms, OutputDir + "Symbols", true));
+				}
+				else
+				{
+					string OutputDir = CommandUtils.CombinePaths(CommandUtils.CmdEnv.LocalRoot, "LocalBuilds", "Rocket", CommandUtils.GetGenericPlatformName(HostPlatform));
+					bp.AddNode(new CopyRocketNode(HostPlatform, CodeTargetPlatforms, StrippedDir, OutputDir, false));
+					bp.AddNode(new CopyRocketSymbolsNode(bp, HostPlatform, CodeTargetPlatforms, OutputDir + "Symbols", false));
+				}
 
 				// Add the aggregate node for the entire install
 				GUBP.GUBPNode PromotableNode = bp.FindNode(GUBP.SharedAggregatePromotableNode.StaticGetFullName());
@@ -132,15 +142,10 @@ namespace Rocket
 					}
 				}
 
-				// Publish the install to the network
-				string RemoteOutputDir = CommandUtils.CombinePaths(CommandUtils.RootSharedTempStorageDirectory(), "Rocket", "Automated", GetBuildLabel(), CommandUtils.GetGenericPlatformName(HostPlatform));
-				bp.AddNode(new PublishRocketNode(HostPlatform, LocalOutputDir, RemoteOutputDir));
-				bp.AddNode(new PublishRocketSymbolsNode(bp, HostPlatform, CodeTargetPlatforms, RemoteOutputDir + "Symbols"));
-
 				// Add a dependency on this being published as part of the shared promotable being labeled
 				GUBP.SharedLabelPromotableSuccessNode LabelPromotableNode = (GUBP.SharedLabelPromotableSuccessNode)bp.FindNode(GUBP.SharedLabelPromotableSuccessNode.StaticGetFullName());
-				LabelPromotableNode.AddDependency(PublishRocketNode.StaticGetFullName(HostPlatform));
-				LabelPromotableNode.AddDependency(PublishRocketSymbolsNode.StaticGetFullName(HostPlatform));
+				LabelPromotableNode.AddDependency(CopyRocketNode.StaticGetFullName(HostPlatform));
+				LabelPromotableNode.AddDependency(CopyRocketSymbolsNode.StaticGetFullName(HostPlatform));
 			}
 		}
 
@@ -679,17 +684,24 @@ namespace Rocket
 		public readonly string OutputDir;
 		public List<UnrealTargetPlatform> CodeTargetPlatforms;
 
-		public CopyRocketNode(UnrealTargetPlatform HostPlatform, List<UnrealTargetPlatform> InCodeTargetPlatforms, string InStrippedDir, string InOutputDir) : base(HostPlatform)
+		public CopyRocketNode(UnrealTargetPlatform HostPlatform, List<UnrealTargetPlatform> InCodeTargetPlatforms, string InStrippedDir, string InOutputDir, bool bWaitForPromotion) : base(HostPlatform)
 		{
 			StrippedDir = InStrippedDir;
 			OutputDir = InOutputDir;
 			CodeTargetPlatforms = new List<UnrealTargetPlatform>(InCodeTargetPlatforms);
 
-			AddDependency(GUBP.WaitForSharedPromotionUserInput.StaticGetFullName(false));
+			// If this is a CIS build, we need to wait for the promotable trigger before running. If not, make an agent sharing group so that the samples node can use our build products.
+			if(bWaitForPromotion)
+			{
+				AddDependency(GUBP.WaitForSharedPromotionUserInput.StaticGetFullName(false));
+			}
+			else
+			{
+				AgentSharingGroup = "RocketGroup" + StaticGetHostPlatformSuffix(HostPlatform);
+			}
+
 			AddDependency(FilterRocketNode.StaticGetFullName(HostPlatform));
 			AddDependency(BuildDerivedDataCacheNode.StaticGetFullName(HostPlatform));
-
-			AgentSharingGroup = "RocketGroup" + StaticGetHostPlatformSuffix(HostPlatform);
 		}
 
 		public static string StaticGetFullName(UnrealTargetPlatform HostPlatform)
@@ -751,57 +763,23 @@ namespace Rocket
 		}
 	}
 
-	public class PublishRocketNode : GUBP.HostPlatformNode
+	public class CopyRocketSymbolsNode : GUBP.HostPlatformNode
 	{
-		string LocalDir;
-		string PublishDir;
-
-		public PublishRocketNode(UnrealTargetPlatform HostPlatform, string InLocalDir, string InPublishDir) : base(HostPlatform)
-		{
-			LocalDir = InLocalDir;
-			PublishDir = InPublishDir;
-
-			AddDependency(CopyRocketNode.StaticGetFullName(HostPlatform));
-
-			AgentSharingGroup = "RocketGroup" + StaticGetHostPlatformSuffix(HostPlatform);
-		}
-
-		public static string StaticGetFullName(UnrealTargetPlatform HostPlatform)
-		{
-			return "PublishRocket" + StaticGetHostPlatformSuffix(HostPlatform);
-		}
-
-		public override string GetFullName()
-		{
-			return StaticGetFullName(HostPlatform);
-		}
+		string SymbolsOutputDir;
 
-		public override float Priority()
+		public CopyRocketSymbolsNode(GUBP bp, UnrealTargetPlatform HostPlatform, IEnumerable<UnrealTargetPlatform> TargetPlatforms, string InSymbolsOutputDir, bool bWaitForPromotion) : base(HostPlatform)
 		{
-			return -10000.0f; // Do this and PublishRocketSymbolsNode at the end
-		}
+			SymbolsOutputDir = InSymbolsOutputDir;
 
-		public override void DoBuild(GUBP bp)
-		{
-			if(RocketBuild.ShouldDoSeriousThingsLikeP4CheckinAndPostToMCP())
+			if(bWaitForPromotion)
 			{
-				CommandUtils.ThreadedCopyFiles(LocalDir, PublishDir);
+				AddDependency(GUBP.WaitForSharedPromotionUserInput.StaticGetFullName(false));
+			}
+			else
+			{
+				AgentSharingGroup = "RocketGroup" + StaticGetHostPlatformSuffix(HostPlatform);
 			}
 
-			BuildProducts = new List<string>();
-			SaveRecordOfSuccessAndAddToBuildProducts();
-		}
-	}
-
-	public class PublishRocketSymbolsNode : GUBP.HostPlatformNode
-	{
-		string SymbolsOutputDir;
-
-		public PublishRocketSymbolsNode(GUBP bp, UnrealTargetPlatform HostPlatform, IEnumerable<UnrealTargetPlatform> TargetPlatforms, string InSymbolsOutputDir) : base(HostPlatform)
-		{
-			SymbolsOutputDir = InSymbolsOutputDir;
-
-			AddDependency(GUBP.WaitForSharedPromotionUserInput.StaticGetFullName(false));
 			AddDependency(GUBP.ToolsForCompileNode.StaticGetFullName(HostPlatform));
 			AddDependency(GUBP.RootEditorNode.StaticGetFullName(HostPlatform));
 			AddDependency(GUBP.ToolsNode.StaticGetFullName(HostPlatform));
@@ -811,13 +789,11 @@ namespace Rocket
 				UnrealTargetPlatform SourceHostPlatform = RocketBuild.GetSourceHostPlatform(bp, HostPlatform, TargetPlatform);
 				AddDependency(GUBP.GamePlatformMonolithicsNode.StaticGetFullName(SourceHostPlatform, bp.Branch.BaseEngineProject, TargetPlatform, Precompiled: true));
 			}
-
-			AgentSharingGroup = "RocketGroup" + StaticGetHostPlatformSuffix(HostPlatform);
 		}
 
 		public static string StaticGetFullName(UnrealTargetPlatform HostPlatform)
 		{
-			return "PublishRocketSymbols" + StaticGetHostPlatformSuffix(HostPlatform);
+			return "CopyRocketSymbols" + StaticGetHostPlatformSuffix(HostPlatform);
 		}
 
 		public override string GetFullName()
@@ -825,25 +801,17 @@ namespace Rocket
 			return StaticGetFullName(HostPlatform);
 		}
 
-		public override float Priority()
-		{
-			return -100000.0f; // Do this last; it's not mission critical
-		}
-
 		public override void DoBuild(GUBP bp)
 		{
-			if(RocketBuild.ShouldDoSeriousThingsLikeP4CheckinAndPostToMCP())
+			// Make a lookup for all the known debug extensions, and filter all the dependency build products against that
+			HashSet<string> DebugExtensions = new HashSet<string>(Platform.Platforms.Values.SelectMany(x => x.GetDebugFileExtentions()).Distinct().ToArray(), StringComparer.InvariantCultureIgnoreCase);
+			foreach(string InputFileName in AllDependencyBuildProducts)
 			{
-				// Make a lookup for all the known debug extensions, and filter all the dependency build products against that
-				HashSet<string> DebugExtensions = new HashSet<string>(Platform.Platforms.Values.SelectMany(x => x.GetDebugFileExtentions()).Distinct().ToArray(), StringComparer.InvariantCultureIgnoreCase);
-				foreach(string InputFileName in AllDependencyBuildProducts)
+				string Extension = Path.GetExtension(InputFileName);
+				if(DebugExtensions.Contains(Extension))
 				{
-					string Extension = Path.GetExtension(InputFileName);
-					if(DebugExtensions.Contains(Extension))
-					{
-						string OutputFileName = CommandUtils.MakeRerootedFilePath(InputFileName, CommandUtils.CmdEnv.LocalRoot, SymbolsOutputDir);
-						CommandUtils.CopyFile(InputFileName, OutputFileName);
-					}
+					string OutputFileName = CommandUtils.MakeRerootedFilePath(InputFileName, CommandUtils.CmdEnv.LocalRoot, SymbolsOutputDir);
+					CommandUtils.CopyFile(InputFileName, OutputFileName);
 				}
 			}
 
