Fix order of node frequency calculation.

diff --git a/Engine/Source/Programs/AutomationTool/Scripts/GUBP.Automation.cs b/Engine/Source/Programs/AutomationTool/Scripts/GUBP.Automation.cs
--- a/Engine/Source/Programs/AutomationTool/Scripts/GUBP.Automation.cs
+++ b/Engine/Source/Programs/AutomationTool/Scripts/GUBP.Automation.cs
@@ -6358,6 +6358,10 @@ public class GUBP : BuildCommand
                     }
                 }
             }
+            foreach (var NodeToDo in GUBPNodes)
+            {
+                ComputeDependentCISFrequencyQuantumShift(NodeToDo.Key);
+            }
 
 			// Make sure that everything that's listed as a frequency barrier is completed with the given interval
 			foreach (KeyValuePair<string, sbyte> Barrier in BranchOptions.FrequencyBarriers)
@@ -6385,19 +6389,11 @@ public class GUBP : BuildCommand
 				{
 					if (IncludedNodes.Contains(NodePair.Key))
 					{
-						if(NodePair.Value.ComputedDependentCISFrequencyQuantumShift == -1 || NodePair.Value.ComputedDependentCISFrequencyQuantumShift > Barrier.Value)
-						{
-							NodePair.Value.ComputedDependentCISFrequencyQuantumShift = Barrier.Value;
-						}
+						NodePair.Value.ComputedDependentCISFrequencyQuantumShift = Math.Min(NodePair.Value.ComputedDependentCISFrequencyQuantumShift, Barrier.Value);
 					}
 				}
 			}
 
-            foreach (var NodeToDo in GUBPNodes)
-            {
-                ComputeDependentCISFrequencyQuantumShift(NodeToDo.Key);
-            }
-
             foreach (var NodeToDo in GUBPNodes)
             {
                 var Deps = GUBPNodes[NodeToDo.Key].DependentPromotions;
