Proper fix for frequency barriers. Nodes behind a barrier now include all pseudodependencies, and frequency is correctly propagated to dependents.

diff --git a/Engine/Source/Programs/AutomationTool/Scripts/GUBP.Automation.cs b/Engine/Source/Programs/AutomationTool/Scripts/GUBP.Automation.cs
--- a/Engine/Source/Programs/AutomationTool/Scripts/GUBP.Automation.cs
+++ b/Engine/Source/Programs/AutomationTool/Scripts/GUBP.Automation.cs
@@ -1617,19 +1617,11 @@ public class GUBP : BuildCommand
                     AddPseudodependency(GamePlatformMonolithicsNode.StaticGetFullName(InHostPlatform, bp.Branch.BaseEngineProject, InHostPlatform, Precompiled: Precompiled));
                 }
             }
-			if(!WithXp && !Precompiled && HasPrecompiledTargets(GameProj, HostPlatform, TargetPlatform))
-			{
-				AddDependency(GamePlatformMonolithicsNode.StaticGetFullName(InHostPlatform, GameProj, InTargetPlatform, WithXp, Precompiled: true));
-			}
             if (InGameProj.Options(InHostPlatform).bTestWithShared)  /// compiling templates is only for testing purposes, and we will group them to avoid saturating the farm
             {
                 AddPseudodependency(WaitForTestShared.StaticGetFullName());
                 AgentSharingGroup = "TemplateMonolithics" + StaticGetHostPlatformSuffix(InHostPlatform);
             }
-			else if (!bp.BranchOptions.bNoInstalledEngine && GameProj == bp.Branch.BaseEngineProject)
-			{
-				AgentSharingGroup = GameProj.GameName + "_MonolithicsGroup_" + InTargetPlatform + StaticGetHostPlatformSuffix(InHostPlatform);
-			}
         }
 
 		public override string GetDisplayGroupName()
@@ -1697,6 +1689,16 @@ public class GUBP : BuildCommand
 			return false;
 		}
 
+		public override float Priority()
+		{
+			float Result = base.Priority();
+			if(Precompiled)
+			{
+				Result += 1.0f;
+			}
+			return Result;
+		}
+
 		public override void DoBuild(GUBP bp)
 		{
 			base.DoBuild(bp);
@@ -4057,24 +4059,31 @@ public class GUBP : BuildCommand
         return FrequencyString;
     }
 
-    public int ComputeDependentCISFrequencyQuantumShift(string NodeToDo)
+    public int ComputeDependentCISFrequencyQuantumShift(string NodeToDo, Dictionary<string, int> FrequencyOverrides)
     {
         int Result = GUBPNodes[NodeToDo].ComputedDependentCISFrequencyQuantumShift;        
         if (Result < 0)
         {
             Result = GUBPNodes[NodeToDo].CISFrequencyQuantumShift(this);
             Result = GetFrequencyForNode(this, GUBPNodes[NodeToDo].GetFullName(), Result);
+
+			int FrequencyOverride;
+			if(FrequencyOverrides.TryGetValue(NodeToDo, out FrequencyOverride) && Result > FrequencyOverride)
+			{
+				Result = FrequencyOverride;
+			}
+
             foreach (var Dep in GUBPNodes[NodeToDo].FullNamesOfDependencies)
             {
-                Result = Math.Max(ComputeDependentCISFrequencyQuantumShift(Dep), Result);
+                Result = Math.Max(ComputeDependentCISFrequencyQuantumShift(Dep, FrequencyOverrides), Result);
             }
             foreach (var Dep in GUBPNodes[NodeToDo].FullNamesOfPseudosependencies)
             {
-                Result = Math.Max(ComputeDependentCISFrequencyQuantumShift(Dep), Result);
+                Result = Math.Max(ComputeDependentCISFrequencyQuantumShift(Dep, FrequencyOverrides), Result);
             }
 			foreach (var Dep in GUBPNodes[NodeToDo].CompletedDependencies)
 			{
-				Result = Math.Max(ComputeDependentCISFrequencyQuantumShift(Dep), Result);
+				Result = Math.Max(ComputeDependentCISFrequencyQuantumShift(Dep, FrequencyOverrides), Result);
 			}
             if (Result < 0)
             {
@@ -6358,12 +6367,9 @@ public class GUBP : BuildCommand
                     }
                 }
             }
-            foreach (var NodeToDo in GUBPNodes)
-            {
-                ComputeDependentCISFrequencyQuantumShift(NodeToDo.Key);
-            }
 
 			// Make sure that everything that's listed as a frequency barrier is completed with the given interval
+			Dictionary<string, int> FrequencyOverrides = new Dictionary<string,int>();
 			foreach (KeyValuePair<string, sbyte> Barrier in BranchOptions.FrequencyBarriers)
 			{
 				// All the nodes which are dependencies of the barrier node
@@ -6374,7 +6380,7 @@ public class GUBP : BuildCommand
 				for (int Idx = 0; Idx < SearchNodes.Count; Idx++)
 				{
 					GUBPNode Node = GUBPNodes[SearchNodes[Idx]];
-					foreach (string DependencyName in Node.FullNamesOfDependencies)
+					foreach (string DependencyName in Node.FullNamesOfDependencies.Union(Node.FullNamesOfPseudosependencies))
 					{
 						if (!IncludedNodes.Contains(DependencyName))
 						{
@@ -6389,11 +6395,26 @@ public class GUBP : BuildCommand
 				{
 					if (IncludedNodes.Contains(NodePair.Key))
 					{
-						NodePair.Value.ComputedDependentCISFrequencyQuantumShift = Math.Min(NodePair.Value.ComputedDependentCISFrequencyQuantumShift, Barrier.Value);
+						int Frequency;
+						if(FrequencyOverrides.TryGetValue(NodePair.Key, out Frequency))
+						{
+							Frequency = Math.Min(Frequency, Barrier.Value);
+						}
+						else
+						{
+							Frequency = Barrier.Value;
+						}
+						FrequencyOverrides[NodePair.Key] = Frequency;
 					}
 				}
 			}
 
+			// Compute all the frequencies
+            foreach (var NodeToDo in GUBPNodes)
+            {
+                ComputeDependentCISFrequencyQuantumShift(NodeToDo.Key, FrequencyOverrides);
+            }
+
             foreach (var NodeToDo in GUBPNodes)
             {
                 var Deps = GUBPNodes[NodeToDo.Key].DependentPromotions;
diff --git a/Engine/Source/Programs/AutomationTool/Scripts/RocketBuild.Automation.cs b/Engine/Source/Programs/AutomationTool/Scripts/RocketBuild.Automation.cs
--- a/Engine/Source/Programs/AutomationTool/Scripts/RocketBuild.Automation.cs
+++ b/Engine/Source/Programs/AutomationTool/Scripts/RocketBuild.Automation.cs
@@ -102,7 +102,7 @@ namespace Rocket
 				{
 					if (GetSourceHostPlatform(bp, HostPlatform, TargetPlatform) == HostPlatform)
 					{
-						bp.AddNode(new StripRocketMonolithicsNode(HostPlatform, bp.Branch.BaseEngineProject, TargetPlatform, CodeTargetPlatforms.Contains(TargetPlatform), StrippedDir));
+						bp.AddNode(new StripRocketMonolithicsNode(bp, HostPlatform, TargetPlatform, CodeTargetPlatforms.Contains(TargetPlatform), StrippedDir));
 					}
 				}
 
@@ -435,14 +435,19 @@ namespace Rocket
 		BranchInfo.BranchUProject Project;
 		bool bIsCodeTargetPlatform;
 
-		public StripRocketMonolithicsNode(UnrealTargetPlatform InHostPlatform, BranchInfo.BranchUProject InProject, UnrealTargetPlatform InTargetPlatform, bool bInIsCodeTargetPlatform, string InStrippedDir) : base(InHostPlatform, InTargetPlatform, InStrippedDir)
+		public StripRocketMonolithicsNode(GUBP bp, UnrealTargetPlatform InHostPlatform, UnrealTargetPlatform InTargetPlatform, bool bInIsCodeTargetPlatform, string InStrippedDir) : base(InHostPlatform, InTargetPlatform, InStrippedDir)
 		{
-			Project = InProject;
+			Project = bp.Branch.BaseEngineProject;
 			bIsCodeTargetPlatform = bInIsCodeTargetPlatform;
 
-			AddDependency(GUBP.GamePlatformMonolithicsNode.StaticGetFullName(HostPlatform, Project, InTargetPlatform, Precompiled: bIsCodeTargetPlatform));
+			GUBP.GUBPNode Node = bp.FindNode(GUBP.GamePlatformMonolithicsNode.StaticGetFullName(HostPlatform, Project, InTargetPlatform, Precompiled: bIsCodeTargetPlatform));
+			if(String.IsNullOrEmpty(Node.AgentSharingGroup))
+			{
+				Node.AgentSharingGroup = bp.Branch.BaseEngineProject.GameName + "_MonolithicsGroup_" + InTargetPlatform + StaticGetHostPlatformSuffix(InHostPlatform);
+			}
+			AddDependency(Node.GetFullName());
 
-			AgentSharingGroup = Project.GameName + "_MonolithicsGroup_" + InTargetPlatform + StaticGetHostPlatformSuffix(InHostPlatform);
+			AgentSharingGroup = Node.AgentSharingGroup;
 		}
 
 		public override string GetDisplayGroupName()
