Prevent overlapping build products from stripping the editor and tools nodes.

diff --git a/Engine/Source/Programs/AutomationTool/Scripts/RocketBuild.Automation.cs b/Engine/Source/Programs/AutomationTool/Scripts/RocketBuild.Automation.cs
--- a/Engine/Source/Programs/AutomationTool/Scripts/RocketBuild.Automation.cs
+++ b/Engine/Source/Programs/AutomationTool/Scripts/RocketBuild.Automation.cs
@@ -313,6 +313,7 @@ namespace Rocket
 	{
 		public UnrealTargetPlatform TargetPlatform;
 		public string StrippedDir;
+		public List<string> NodesToStrip = new List<string>();
 
 		public StripRocketNode(UnrealTargetPlatform InHostPlatform, UnrealTargetPlatform InTargetPlatform, string InStrippedDir) : base(InHostPlatform)
 		{
@@ -322,6 +323,12 @@ namespace Rocket
 
 		public override abstract string GetFullName();
 
+		public void AddNodeToStrip(string NodeName)
+		{
+			NodesToStrip.Add(NodeName);
+			AddDependency(NodeName);
+		}
+
 		public static bool IsRequiredForPlatform(UnrealTargetPlatform Platform)
 		{
 			switch(Platform)
@@ -350,13 +357,17 @@ namespace Rocket
 			// Apply the filter to the build products
 			List<string> SourcePaths = new List<string>();
 			List<string> TargetPaths = new List<string>();
-			foreach(string DependencyBuildProduct in AllDependencyBuildProducts)
+			foreach(string NodeToStrip in NodesToStrip)
 			{
-				string RelativePath = CommandUtils.StripBaseDirectory(Path.GetFullPath(DependencyBuildProduct), InputDir);
-				if(StripFilter.Matches(RelativePath))
+				GUBP.GUBPNode Node = bp.FindNode(NodeToStrip);
+				foreach(string DependencyBuildProduct in Node.BuildProducts)
 				{
-					SourcePaths.Add(CommandUtils.CombinePaths(InputDir, RelativePath));
-					TargetPaths.Add(CommandUtils.CombinePaths(StrippedDir, RelativePath));
+					string RelativePath = CommandUtils.StripBaseDirectory(Path.GetFullPath(DependencyBuildProduct), InputDir);
+					if(StripFilter.Matches(RelativePath))
+					{
+						SourcePaths.Add(CommandUtils.CombinePaths(InputDir, RelativePath));
+						TargetPaths.Add(CommandUtils.CombinePaths(StrippedDir, RelativePath));
+					}
 				}
 			}
 
@@ -394,8 +405,8 @@ namespace Rocket
 		public StripRocketToolsNode(UnrealTargetPlatform InHostPlatform, string InStrippedDir)
 			: base(InHostPlatform, InHostPlatform, InStrippedDir)
 		{
-			AddDependency(GUBP.ToolsForCompileNode.StaticGetFullName(HostPlatform));
-			AddDependency(GUBP.ToolsNode.StaticGetFullName(HostPlatform));
+			AddNodeToStrip(GUBP.ToolsForCompileNode.StaticGetFullName(HostPlatform));
+			AddNodeToStrip(GUBP.ToolsNode.StaticGetFullName(HostPlatform));
 			AgentSharingGroup = "ToolsGroup" + StaticGetHostPlatformSuffix(InHostPlatform);
 		}
 
@@ -415,7 +426,7 @@ namespace Rocket
 		public StripRocketEditorNode(UnrealTargetPlatform InHostPlatform, string InStrippedDir)
 			: base(InHostPlatform, InHostPlatform, InStrippedDir)
 		{
-			AddDependency(GUBP.RootEditorNode.StaticGetFullName(HostPlatform));
+			AddNodeToStrip(GUBP.RootEditorNode.StaticGetFullName(HostPlatform));
 			AgentSharingGroup = "Editor" + StaticGetHostPlatformSuffix(HostPlatform);
 		}
 
@@ -445,7 +456,7 @@ namespace Rocket
 			{
 				Node.AgentSharingGroup = bp.Branch.BaseEngineProject.GameName + "_MonolithicsGroup_" + InTargetPlatform + StaticGetHostPlatformSuffix(InHostPlatform);
 			}
-			AddDependency(Node.GetFullName());
+			AddNodeToStrip(Node.GetFullName());
 
 			AgentSharingGroup = Node.AgentSharingGroup;
 		}
