Change spec used to determine failure emails; was only checking source folders under a top level directory (not plugins, samples, etc...)

diff --git a/Engine/Source/Programs/AutomationTool/Scripts/GUBP.Automation.cs b/Engine/Source/Programs/AutomationTool/Scripts/GUBP.Automation.cs
--- a/Engine/Source/Programs/AutomationTool/Scripts/GUBP.Automation.cs
+++ b/Engine/Source/Programs/AutomationTool/Scripts/GUBP.Automation.cs
@@ -4171,8 +4171,8 @@ public class GUBP : BuildCommand
             if (LastOutputForChanges > 1990000)
             {
                 string Cmd = String.Format("{0}@{1},{2} {3}@{4},{5}",
-                    CombinePaths(PathSeparator.Slash, P4Env.BuildRootP4, "*", "Source", "..."), LastOutputForChanges + 1, TopCL,
-                    CombinePaths(PathSeparator.Slash, P4Env.BuildRootP4, "*", "Build", "..."), LastOutputForChanges + 1, TopCL
+                    CombinePaths(PathSeparator.Slash, P4Env.BuildRootP4, "...", "Source", "..."), LastOutputForChanges + 1, TopCL,
+                    CombinePaths(PathSeparator.Slash, P4Env.BuildRootP4, "...", "Build", "..."), LastOutputForChanges + 1, TopCL
                     );
                 List<P4Connection.ChangeRecord> ChangeRecords;
 				if (P4.Changes(out ChangeRecords, Cmd, false, true, LongComment: true))
