Merging using UE4-To-UE4-LauncherDev

diff --git a/Engine/Source/Programs/AutomationTool/Scripts/GUBP.Automation.cs b/Engine/Source/Programs/AutomationTool/Scripts/GUBP.Automation.cs
--- a/Engine/Source/Programs/AutomationTool/Scripts/GUBP.Automation.cs
+++ b/Engine/Source/Programs/AutomationTool/Scripts/GUBP.Automation.cs
@@ -1,4 +1,4 @@
-﻿// Copyright 1998-2015 Epic Games, Inc. All Rights Reserved.
+// Copyright 1998-2015 Epic Games, Inc. All Rights Reserved.
 using System;
 using System.Collections.Generic;
 using System.Text;
@@ -3846,7 +3846,7 @@ public class GUBP : BuildCommand
         if (!GUBPNodes.ContainsKey(Node))
         {
             throw new AutomationException("Node {0} not found", Node);
-        }        
+        }
         GUBPNodes[Node].RemovePseudodependency(Dep);
     }
 
