Use separate manifests for each node that strips symbols, so we can mix-and-match on different host platforms. Using Win64 Android build products on Mac has a different output folder than the target host platform.

diff --git a/Engine/Source/Programs/AutomationTool/Scripts/RocketBuild.Automation.cs b/Engine/Source/Programs/AutomationTool/Scripts/RocketBuild.Automation.cs
--- a/Engine/Source/Programs/AutomationTool/Scripts/RocketBuild.Automation.cs
+++ b/Engine/Source/Programs/AutomationTool/Scripts/RocketBuild.Automation.cs
@@ -114,7 +114,7 @@ namespace Rocket
 
 				// Copy the install to the output directory
 				string LocalOutputDir = CommandUtils.CombinePaths(CommandUtils.CmdEnv.LocalRoot, "LocalBuilds", "Rocket", CommandUtils.GetGenericPlatformName(HostPlatform));
-				bp.AddNode(new GatherRocketNode(HostPlatform, CodeTargetPlatforms, StrippedDir, LocalOutputDir));
+				bp.AddNode(new GatherRocketNode(HostPlatform, CodeTargetPlatforms, LocalOutputDir));
 
 				// Add the aggregate node for the entire install
 				GUBP.GUBPNode PromotableNode = bp.FindNode(GUBP.SharedAggregatePromotableNode.StaticGetFullName());
@@ -485,8 +485,7 @@ namespace Rocket
 		string[] CurrentFeaturePacks;
 		string[] CurrentTemplates;
 		public readonly string DepotManifestPath;
-		public readonly string StrippedManifestPath;
-		List<string> StrippedNodes = new List<string>();
+		public Dictionary<string, string> StrippedNodeManifestPaths = new Dictionary<string, string>();
 
 		public FilterRocketNode(GUBP bp, UnrealTargetPlatform InHostPlatform, List<UnrealTargetPlatform> InTargetPlatforms, List<UnrealTargetPlatform> InCodeTargetPlatforms, string[] InCurrentFeaturePacks, string[] InCurrentTemplates)
 			: base(InHostPlatform)
@@ -495,8 +494,7 @@ namespace Rocket
 			CodeTargetPlatforms = new List<UnrealTargetPlatform>(InCodeTargetPlatforms);
 			CurrentFeaturePacks = InCurrentFeaturePacks;
 			CurrentTemplates = InCurrentTemplates;
-			DepotManifestPath = CommandUtils.CombinePaths(CommandUtils.CmdEnv.LocalRoot, "Engine", "Saved", "Rocket", HostPlatform.ToString(), "DistillDepot.txt");
-			StrippedManifestPath = CommandUtils.CombinePaths(CommandUtils.CmdEnv.LocalRoot, "Engine", "Saved", "Rocket", HostPlatform.ToString(), "DistillStripped.txt");
+			DepotManifestPath = CommandUtils.CombinePaths(CommandUtils.CmdEnv.LocalRoot, "Engine", "Saved", "Rocket", HostPlatform.ToString(), "Filter.txt");
 
 			// Add the editor
 			AddDependency(GUBP.VersionFilesNode.StaticGetFullName());
@@ -505,6 +503,7 @@ namespace Rocket
 			AddDependency(GUBP.ToolsNode.StaticGetFullName(HostPlatform));
 
 			// Get all the external file lists. Android/HTML5 on Mac uses the same files as Win64.
+			List<string> StrippedNodeNames = new List<string>();
 			foreach(UnrealTargetPlatform TargetPlatform in TargetPlatforms)
 			{
 				UnrealTargetPlatform SourceHostPlatform = RocketBuild.GetSourceHostPlatform(bp, HostPlatform, TargetPlatform);
@@ -515,7 +514,7 @@ namespace Rocket
 				{
 					string StripNode = StripRocketMonolithicsNode.StaticGetFullName(SourceHostPlatform, bp.Branch.BaseEngineProject, TargetPlatform, bIsCodeTargetPlatform);
 					AddDependency(StripNode);
-					StrippedNodes.Add(StripNode);
+					StrippedNodeNames.Add(StripNode);
 				}
 			}
 
@@ -540,10 +539,16 @@ namespace Rocket
 			if(StripRocketNode.IsRequiredForPlatform(HostPlatform))
 			{
 				AddDependency(StripRocketToolsNode.StaticGetFullName(HostPlatform));
-				StrippedNodes.Add(StripRocketToolsNode.StaticGetFullName(HostPlatform));
+				StrippedNodeNames.Add(StripRocketToolsNode.StaticGetFullName(HostPlatform));
 
 				AddDependency(StripRocketEditorNode.StaticGetFullName(HostPlatform));
-				StrippedNodes.Add(StripRocketEditorNode.StaticGetFullName(HostPlatform));
+				StrippedNodeNames.Add(StripRocketEditorNode.StaticGetFullName(HostPlatform));
+			}
+
+			// Set all the stripped manifest paths
+			foreach(string StrippedNodeName in StrippedNodeNames)
+			{
+				StrippedNodeManifestPaths.Add(StrippedNodeName, Path.Combine(Path.GetDirectoryName(DepotManifestPath), "Filter_" + StrippedNodeName + ".txt"));
 			}
 		}
 
@@ -564,6 +569,7 @@ namespace Rocket
 
 		public override void DoBuild(GUBP bp)
 		{
+			BuildProducts = new List<string>();
 			FileFilter Filter = new FileFilter();
 
 			// Include all the editor products
@@ -621,10 +627,12 @@ namespace Rocket
 			Filter.ExcludeConfidentialFolders();
 
 			// Run the filter on the stripped symbols, and remove those files from the copy filter
-			List<string> StrippedFiles = new List<string>();
-			foreach(string StrippedNode in StrippedNodes)
+			List<string> AllStrippedFiles = new List<string>();
+			foreach(KeyValuePair<string, string> StrippedNodeManifestPath in StrippedNodeManifestPaths)
 			{
-				StripRocketNode StripNode = (StripRocketNode)bp.FindNode(StrippedNode);
+				List<string> StrippedFiles = new List<string>();
+
+				StripRocketNode StripNode = (StripRocketNode)bp.FindNode(StrippedNodeManifestPath.Key);
 				foreach(string BuildProduct in StripNode.BuildProducts)
 				{
 					if(Utils.IsFileUnderDirectory(BuildProduct, StripNode.StrippedDir))
@@ -633,16 +641,20 @@ namespace Rocket
 						if(Filter.Matches(RelativePath))
 						{
 							StrippedFiles.Add(RelativePath);
+							AllStrippedFiles.Add(RelativePath);
 							Filter.Exclude("/" + RelativePath);
 						}
 					}
 				}
+
+				WriteManifest(StrippedNodeManifestPath.Value, StrippedFiles);
+				BuildProducts.Add(StrippedNodeManifestPath.Value);
 			}
-			WriteManifest(StrippedManifestPath, StrippedFiles);
 
 			// Write the filtered list of depot files to disk, removing any symlinks
 			List<string> DepotFiles = Filter.ApplyToDirectory(CommandUtils.CmdEnv.LocalRoot, true).ToList();
 			WriteManifest(DepotManifestPath, DepotFiles);
+			BuildProducts.Add(DepotManifestPath);
 
 			// Sort the list of output files
 			SortedDictionary<string, bool> SortedFiles = new SortedDictionary<string,bool>(StringComparer.InvariantCultureIgnoreCase);
@@ -650,7 +662,7 @@ namespace Rocket
 			{
 				SortedFiles.Add(DepotFile, false);
 			}
-			foreach(string StrippedFile in StrippedFiles)
+			foreach(string StrippedFile in AllStrippedFiles)
 			{
 				SortedFiles.Add(StrippedFile, true);
 			}
@@ -661,9 +673,6 @@ namespace Rocket
 			{
 				CommandUtils.Log("  {0}{1}", SortedFile.Key, SortedFile.Value? " (stripped)" : "");
 			}
-
-			// Add the build products
-			BuildProducts = new List<string>{ StrippedManifestPath, DepotManifestPath };
 		}
 
 		static void AddRuleForBuildProducts(FileFilter Filter, GUBP bp, string NodeName, FileFilterType Type)
@@ -691,17 +700,20 @@ namespace Rocket
 
 	public class GatherRocketNode : GUBP.HostPlatformNode
 	{
-		public readonly string StrippedDir;
 		public readonly string OutputDir;
 		public List<UnrealTargetPlatform> CodeTargetPlatforms;
 
-		public GatherRocketNode(UnrealTargetPlatform HostPlatform, List<UnrealTargetPlatform> InCodeTargetPlatforms, string InStrippedDir, string InOutputDir) : base(HostPlatform)
+		public GatherRocketNode(UnrealTargetPlatform HostPlatform, List<UnrealTargetPlatform> InCodeTargetPlatforms, string InOutputDir) : base(HostPlatform)
 		{
-			StrippedDir = InStrippedDir;
 			OutputDir = InOutputDir;
 			CodeTargetPlatforms = new List<UnrealTargetPlatform>(InCodeTargetPlatforms);
 
-			AddDependency(GUBP.WaitForSharedPromotionUserInput.StaticGetFullName(false));
+			// Optimization for testing locally; don't do other things that are part of the shared promotable
+			if(RocketBuild.ShouldDoSeriousThingsLikeP4CheckinAndPostToMCP())
+			{
+				AddDependency(GUBP.WaitForSharedPromotionUserInput.StaticGetFullName(false));
+			}
+
 			AddDependency(FilterRocketNode.StaticGetFullName(HostPlatform));
 			AddDependency(BuildDerivedDataCacheNode.StaticGetFullName(HostPlatform));
 
@@ -736,7 +748,13 @@ namespace Rocket
 			// Copy the depot files to the output directory
 			FilterRocketNode FilterNode = (FilterRocketNode)bp.FindNode(FilterRocketNode.StaticGetFullName(HostPlatform));
 			CopyManifestFilesToOutput(FilterNode.DepotManifestPath, CommandUtils.CmdEnv.LocalRoot, OutputDir);
-			CopyManifestFilesToOutput(FilterNode.StrippedManifestPath, StrippedDir, OutputDir);
+
+			// Copy the stripped files to the output directory
+			foreach(KeyValuePair<string, string> StrippedManifestPath in FilterNode.StrippedNodeManifestPaths)
+			{
+				StripRocketNode StripNode = (StripRocketNode)bp.FindNode(StrippedManifestPath.Key);
+				CopyManifestFilesToOutput(StrippedManifestPath.Value, StripNode.StrippedDir, OutputDir);
+			}
 
 			// Copy the DDC to the output directory
 			BuildDerivedDataCacheNode DerivedDataCacheNode = (BuildDerivedDataCacheNode)bp.FindNode(BuildDerivedDataCacheNode.StaticGetFullName(HostPlatform));
