UAT - test command to look for overlapping .obj files

diff --git a/Engine/Source/Programs/AutomationTool/Scripts/Tests.Automation.cs b/Engine/Source/Programs/AutomationTool/Scripts/Tests.Automation.cs
--- a/Engine/Source/Programs/AutomationTool/Scripts/Tests.Automation.cs
+++ b/Engine/Source/Programs/AutomationTool/Scripts/Tests.Automation.cs
@@ -1968,6 +1968,58 @@ class TestStopProcess : BuildCommand
 	}
 }
 
+[Help("Looks through an XGE xml for overlapping obj files")]
+[Help("Source", "full path of xml to look at")]
+class LookForOverlappingBuildProducts : BuildCommand
+{
+	public override void ExecuteBuild()
+	{
+		var SourcePath = ParseParamValue("Source=");
+		if (String.IsNullOrEmpty(SourcePath))
+		{
+			SourcePath = "D:\\UAT_XGE.xml";
+		}
+		if (!FileExists_NoExceptions(SourcePath))
+		{
+			throw new AutomationException("Source path not found, please use -source=Path");
+		}
+		var Objs = new HashSet<string>( StringComparer.InvariantCultureIgnoreCase );
+//    /Fo&quot;D:\BuildFarm\buildmachine_++depot+UE4\Engine\Intermediate\Build\Win64\UE4Editor\Development\Projects\Module.Projects.cpp.obj&quot;
+		var FileText = ReadAllText(SourcePath);
+		string Start = "/Fo&quot;";
+		string End = "&quot;";
+		while (true)
+		{
+			var Index = FileText.IndexOf(Start);
+			if (Index >= 0)
+			{
+				FileText = FileText.Substring(Index + Start.Length);
+				Index = FileText.IndexOf(End);
+				if (Index >= 0)
+				{
+					var ObjFile = FileText.Substring(0, Index);
+					if (Objs.Contains(ObjFile))
+					{
+						LogError("Duplicate obj: {0}", ObjFile);
+					}
+					else
+					{
+						Objs.Add(ObjFile);
+					}
+					FileText = FileText.Substring(Index + End.Length);
+				}
+				else
+				{
+					break;
+				}
+			}
+			else
+			{
+				break;
+			}
+		}
+	}
+}
 
 [Help("Copies all files from source directory to destination directory using ThreadedCopyFiles")]
 [Help("Source", "Source path")]
