Fix missing headers for editor targets (which prevents creation of new code projects), and add UBT/UAT from Windows into Mac install (which was causing an "Engine source not found" error).

diff --git a/Engine/Source/Programs/AutomationTool/Scripts/RocketBuild.Automation.cs b/Engine/Source/Programs/AutomationTool/Scripts/RocketBuild.Automation.cs
--- a/Engine/Source/Programs/AutomationTool/Scripts/RocketBuild.Automation.cs
+++ b/Engine/Source/Programs/AutomationTool/Scripts/RocketBuild.Automation.cs
@@ -568,8 +568,15 @@ namespace Rocket
 			AddRuleForBuildProducts(Filter, bp, GUBP.RootEditorNode.StaticGetFullName(HostPlatform), FileFilterType.Include);
 			AddRuleForBuildProducts(Filter, bp, GUBP.ToolsNode.StaticGetFullName(HostPlatform), FileFilterType.Include);
 
+			// Include win64 tools on Mac, to get the win64 build of UBT, UAT and IPP
+			if (HostPlatform == UnrealTargetPlatform.Mac && bp.HostPlatforms.Contains(UnrealTargetPlatform.Win64))
+			{
+				AddRuleForBuildProducts(Filter, bp, GUBP.ToolsNode.StaticGetFullName(UnrealTargetPlatform.Win64), FileFilterType.Include);
+				AddRuleForBuildProducts(Filter, bp, GUBP.ToolsForCompileNode.StaticGetFullName(UnrealTargetPlatform.Win64), FileFilterType.Include);
+			}
+
 			// Include the editor headers
-			UnzipAndAddRuleForHeaders(GUBP.RootEditorNode.StaticGetFullName(HostPlatform), Filter, FileFilterType.Include);
+			UnzipAndAddRuleForHeaders(GUBP.RootEditorNode.StaticGetArchivedHeadersPath(HostPlatform), Filter, FileFilterType.Include);
 
 			// Include the build dependencies for every code platform
 			foreach(UnrealTargetPlatform CodeTargetPlatform in CodeTargetPlatforms)
@@ -577,7 +584,7 @@ namespace Rocket
 				UnrealTargetPlatform SourceHostPlatform = RocketBuild.GetSourceHostPlatform(bp, HostPlatform, CodeTargetPlatform);
 				string FileListPath = GUBP.GamePlatformMonolithicsNode.StaticGetBuildDependenciesPath(SourceHostPlatform, bp.Branch.BaseEngineProject, CodeTargetPlatform);
 				Filter.AddRuleForFiles(UnrealBuildTool.Utils.ReadClass<UnrealBuildTool.ExternalFileList>(FileListPath).FileNames, CommandUtils.CmdEnv.LocalRoot, FileFilterType.Include);
-				UnzipAndAddRuleForHeaders(GUBP.GamePlatformMonolithicsNode.StaticGetArchivedHeadersPath(HostPlatform, bp.Branch.BaseEngineProject, CodeTargetPlatform), Filter, FileFilterType.Include);
+				UnzipAndAddRuleForHeaders(GUBP.GamePlatformMonolithicsNode.StaticGetArchivedHeadersPath(SourceHostPlatform, bp.Branch.BaseEngineProject, CodeTargetPlatform), Filter, FileFilterType.Include);
 			}
 
 			// Add the monolithic binaries
