Include the base engine project when looking for which platforms to build for. HTML5 is only built for the base engine project, so was not being included in CIS/Rocket.

diff --git a/Engine/Source/Programs/AutomationTool/Scripts/GUBP.Automation.cs b/Engine/Source/Programs/AutomationTool/Scripts/GUBP.Automation.cs
--- a/Engine/Source/Programs/AutomationTool/Scripts/GUBP.Automation.cs
+++ b/Engine/Source/Programs/AutomationTool/Scripts/GUBP.Automation.cs
@@ -5478,39 +5478,35 @@ public class GUBP : BuildCommand
         if (IsBuildMachine || ParseParam("AllPlatforms"))
         {
             ActivePlatforms = new List<UnrealTargetPlatform>();
-            var BranchCodeProjects = new List<BranchInfo.BranchUProject>();
-            foreach (var GameProj in Branch.CodeProjects)
+            
+			List<BranchInfo.BranchUProject> BranchCodeProjects = new List<BranchInfo.BranchUProject>();
+			BranchCodeProjects.Add(Branch.BaseEngineProject);
+			BranchCodeProjects.AddRange(Branch.CodeProjects);
+			BranchCodeProjects.RemoveAll(Project => BranchOptions.ExcludeNodes.Contains(Project.GameName));
+
+			foreach (var GameProj in BranchCodeProjects)
             {
-                if (BranchOptions.ExcludeNodes.Contains(GameProj.GameName))
-                {
-                    continue;
-                }
-                else
+                foreach (var Kind in BranchInfo.MonolithicKinds)
                 {
-                    BranchCodeProjects.Add(GameProj);
-                    foreach (var Kind in BranchInfo.MonolithicKinds)
+                    if (GameProj.Properties.Targets.ContainsKey(Kind))
                     {
-                        if (GameProj.Properties.Targets.ContainsKey(Kind))
+                        var Target = GameProj.Properties.Targets[Kind];
+                        foreach (var HostPlatform in HostPlatforms)
                         {
-                            var Target = GameProj.Properties.Targets[Kind];
-                            foreach (var HostPlatform in HostPlatforms)
+                            var Platforms = Target.Rules.GUBP_GetPlatforms_MonolithicOnly(HostPlatform);
+                            var AdditionalPlatforms = Target.Rules.GUBP_GetBuildOnlyPlatforms_MonolithicOnly(HostPlatform);
+                            var AllPlatforms = Platforms.Union(AdditionalPlatforms);
+                            foreach (var Plat in AllPlatforms)
                             {
-                                var Platforms = Target.Rules.GUBP_GetPlatforms_MonolithicOnly(HostPlatform);
-                                var AdditionalPlatforms = Target.Rules.GUBP_GetBuildOnlyPlatforms_MonolithicOnly(HostPlatform);
-                                var AllPlatforms = Platforms.Union(AdditionalPlatforms);
-                                foreach (var Plat in AllPlatforms)
+                                if (Target.Rules.SupportsPlatform(Plat) && !ActivePlatforms.Contains(Plat))
                                 {
-                                    if (Target.Rules.SupportsPlatform(Plat) && !ActivePlatforms.Contains(Plat))
-                                    {
-                                        ActivePlatforms.Add(Plat);
-                                    }
+                                    ActivePlatforms.Add(Plat);
                                 }
                             }
                         }
                     }
-                    Branch.CodeProjects = BranchCodeProjects;
                 }
-            }
+			}
         }
         else
         {
