Fix missing headers when gathering Rocket builds - incorrect path specified for root editor node.

diff --git a/Engine/Source/Programs/AutomationTool/Scripts/RocketBuild.Automation.cs b/Engine/Source/Programs/AutomationTool/Scripts/RocketBuild.Automation.cs
--- a/Engine/Source/Programs/AutomationTool/Scripts/RocketBuild.Automation.cs
+++ b/Engine/Source/Programs/AutomationTool/Scripts/RocketBuild.Automation.cs
@@ -733,7 +733,7 @@ namespace Rocket
 			CommandUtils.DeleteDirectoryContents(OutputDir);
 
 			// Extract the editor headers
-			CommandUtils.UnzipFiles(GUBP.RootEditorNode.StaticGetFullName(HostPlatform), CommandUtils.CmdEnv.LocalRoot);
+			CommandUtils.UnzipFiles(GUBP.RootEditorNode.StaticGetArchivedHeadersPath(HostPlatform), CommandUtils.CmdEnv.LocalRoot);
 
 			// Extract all the headers for code target platforms
 			foreach(UnrealTargetPlatform CodeTargetPlatform in CodeTargetPlatforms)
