Linux: packaging fixes (UE-12940).

- Fix for build receipt (related: UE-13763).
- Handling root:/ etc paths which are sometimes passed to the lower level without expansion.

diff --git a/Engine/Source/Programs/AutomationTool/Scripts/CopyBuildToStagingDirectory.Automation.cs b/Engine/Source/Programs/AutomationTool/Scripts/CopyBuildToStagingDirectory.Automation.cs
--- a/Engine/Source/Programs/AutomationTool/Scripts/CopyBuildToStagingDirectory.Automation.cs
+++ b/Engine/Source/Programs/AutomationTool/Scripts/CopyBuildToStagingDirectory.Automation.cs
@@ -1388,8 +1388,9 @@ public partial class Project : CommandUtils
 
 					foreach (UnrealTargetPlatform ReceiptPlatform in SubPlatformsToStage)
 					{
-						// Read the receipt
-						string ReceiptFileName = BuildReceipt.GetDefaultPath(ReceiptBaseDir, Target, ReceiptPlatform, Config, "");
+						// Read the receipt (FIXME: store the architecture in project params?)
+						string Architecture = UEBuildPlatform.GetBuildPlatform(ReceiptPlatform).GetActiveArchitecture();
+						string ReceiptFileName = BuildReceipt.GetDefaultPath(ReceiptBaseDir, Target, ReceiptPlatform, Config, Architecture);
 						if(!File.Exists(ReceiptFileName))
 						{
 							if (bRequireStagedFilesToExist)
diff --git a/Engine/Source/Runtime/Core/Private/Linux/LinuxPlatformFile.cpp b/Engine/Source/Runtime/Core/Private/Linux/LinuxPlatformFile.cpp
--- a/Engine/Source/Runtime/Core/Private/Linux/LinuxPlatformFile.cpp
+++ b/Engine/Source/Runtime/Core/Private/Linux/LinuxPlatformFile.cpp
@@ -499,6 +499,13 @@ public:
 	 */
 	int32 OpenCaseInsensitiveRead(const FString & Filename, FString & MappedToFilename)
 	{
+		// We can get some "absolute" filenames like "D:/Blah/" here (e.g. non-Linux paths to source files embedded in assets).
+		// In that case, fail silently.
+		if (Filename.IsEmpty() || Filename[0] != TEXT('/'))
+		{
+			return -1;
+		}
+
 		// try opening right away
 		int32 Handle = open(TCHAR_TO_UTF8(*Filename), O_RDONLY);
 		if (Handle != -1)
