Cleaned up unnecessary log message from TestTempStorage.

diff --git a/Engine/Source/Programs/AutomationTool/Scripts/Tests.Automation.cs b/Engine/Source/Programs/AutomationTool/Scripts/Tests.Automation.cs
--- a/Engine/Source/Programs/AutomationTool/Scripts/Tests.Automation.cs
+++ b/Engine/Source/Programs/AutomationTool/Scripts/Tests.Automation.cs
@@ -483,8 +483,6 @@ class TestTempStorage : BuildCommand
 	{
 		Log("TestTempStorage********");
 
-		Log("Resolved Ocean to {0}", ResolveSharedBuildDirectory("Ocean"));
-
 		DeleteLocalTempStorageManifests(CmdEnv);
 		DeleteSharedTempStorageManifests(CmdEnv, "Test");
 		if (TempStorageExists(CmdEnv, "Test"))
