Enable symbol stripping nodes for Mac and IOS.

diff --git a/Engine/Source/Programs/AutomationTool/Scripts/RocketBuild.Automation.cs b/Engine/Source/Programs/AutomationTool/Scripts/RocketBuild.Automation.cs
--- a/Engine/Source/Programs/AutomationTool/Scripts/RocketBuild.Automation.cs
+++ b/Engine/Source/Programs/AutomationTool/Scripts/RocketBuild.Automation.cs
@@ -326,16 +326,7 @@ namespace Rocket
 
 		public static bool IsRequiredForPlatform(UnrealTargetPlatform Platform)
 		{
-			switch(Platform)
-			{
-				case UnrealTargetPlatform.Win32:
-				case UnrealTargetPlatform.Win64:
-				case UnrealTargetPlatform.Android:
-				case UnrealTargetPlatform.Linux:
-					return true;
-				default:
-					return false;
-			}
+			return Platform != UnrealTargetPlatform.HTML5;
 		}
 
 		public override void DoBuild(GUBP bp)
