Back out changelist 2528840

diff --git a/Engine/Source/Programs/AutomationTool/Scripts/RocketBuild.Automation.cs b/Engine/Source/Programs/AutomationTool/Scripts/RocketBuild.Automation.cs
--- a/Engine/Source/Programs/AutomationTool/Scripts/RocketBuild.Automation.cs
+++ b/Engine/Source/Programs/AutomationTool/Scripts/RocketBuild.Automation.cs
@@ -326,7 +326,16 @@ namespace Rocket
 
 		public static bool IsRequiredForPlatform(UnrealTargetPlatform Platform)
 		{
-			return Platform != UnrealTargetPlatform.HTML5;
+			switch(Platform)
+			{
+				case UnrealTargetPlatform.Win32:
+				case UnrealTargetPlatform.Win64:
+				case UnrealTargetPlatform.Android:
+				case UnrealTargetPlatform.Linux:
+					return true;
+				default:
+					return false;
+			}
 		}
 
 		public override void DoBuild(GUBP bp)
