Second attempt to fix stripping for iOS and Mac builds. Uncovered lots of problems relating to the source host platform not being correct.

diff --git a/Engine/Source/Programs/AutomationTool/Scripts/RocketBuild.Automation.cs b/Engine/Source/Programs/AutomationTool/Scripts/RocketBuild.Automation.cs
--- a/Engine/Source/Programs/AutomationTool/Scripts/RocketBuild.Automation.cs
+++ b/Engine/Source/Programs/AutomationTool/Scripts/RocketBuild.Automation.cs
@@ -77,16 +77,6 @@ namespace Rocket
 				// Find all the target platforms for this host platform.
 				List<UnrealTargetPlatform> TargetPlatforms = GetTargetPlatforms(bp, HostPlatform);
 
-				// Get a list of all the code target platforms
-				List<UnrealTargetPlatform> CodeTargetPlatforms = new List<UnrealTargetPlatform>();
-				foreach(UnrealTargetPlatform TargetPlatform in TargetPlatforms)
-				{
-					if(IsCodeTargetPlatform(HostPlatform, TargetPlatform))
-					{
-						CodeTargetPlatforms.Add(TargetPlatform);
-					}
-				}
-
 				// Get the temp directory for stripped files for this host
 				string StrippedDir = Path.GetFullPath(CommandUtils.CombinePaths(CommandUtils.CmdEnv.LocalRoot, "Engine", "Saved", "Rocket", HostPlatform.ToString()));
 
@@ -97,12 +87,12 @@ namespace Rocket
 					bp.AddNode(new StripRocketEditorNode(HostPlatform, StrippedDir));
 				}
 
-				// Strip all the target platforms
+				// Strip all the target platforms that are built on this host
 				foreach (UnrealTargetPlatform TargetPlatform in TargetPlatforms)
 				{
-					if (GetSourceHostPlatform(bp, HostPlatform, TargetPlatform) == HostPlatform)
+					if (GetSourceHostPlatform(bp, HostPlatform, TargetPlatform) == HostPlatform && StripRocketNode.IsRequiredForPlatform(TargetPlatform))
 					{
-						bp.AddNode(new StripRocketMonolithicsNode(bp, HostPlatform, TargetPlatform, CodeTargetPlatforms.Contains(TargetPlatform), StrippedDir));
+						bp.AddNode(new StripRocketMonolithicsNode(bp, HostPlatform, TargetPlatform, StrippedDir));
 					}
 				}
 
@@ -110,11 +100,11 @@ namespace Rocket
 				bp.AddNode(new BuildDerivedDataCacheNode(HostPlatform, GetCookPlatforms(HostPlatform, TargetPlatforms), CurrentFeaturePacks));
 
 				// Generate a list of files that needs to be copied for each target platform
-				bp.AddNode(new FilterRocketNode(bp, HostPlatform, TargetPlatforms, CodeTargetPlatforms, CurrentFeaturePacks, CurrentTemplates));
+				bp.AddNode(new FilterRocketNode(bp, HostPlatform, TargetPlatforms, CurrentFeaturePacks, CurrentTemplates));
 
 				// Copy the install to the output directory
 				string LocalOutputDir = CommandUtils.CombinePaths(CommandUtils.CmdEnv.LocalRoot, "LocalBuilds", "Rocket", CommandUtils.GetGenericPlatformName(HostPlatform));
-				bp.AddNode(new GatherRocketNode(HostPlatform, CodeTargetPlatforms, LocalOutputDir));
+				bp.AddNode(new GatherRocketNode(HostPlatform, TargetPlatforms, LocalOutputDir));
 
 				// Add the aggregate node for the entire install
 				GUBP.GUBPNode PromotableNode = bp.FindNode(GUBP.SharedAggregatePromotableNode.StaticGetFullName());
@@ -135,7 +125,7 @@ namespace Rocket
 				// Publish the install to the network
 				string RemoteOutputDir = CommandUtils.CombinePaths(CommandUtils.RootSharedTempStorageDirectory(), "Rocket", "Automated", GetBuildLabel(), CommandUtils.GetGenericPlatformName(HostPlatform));
 				bp.AddNode(new PublishRocketNode(HostPlatform, LocalOutputDir, RemoteOutputDir));
-				bp.AddNode(new PublishRocketSymbolsNode(bp, HostPlatform, CodeTargetPlatforms, RemoteOutputDir + "Symbols"));
+				bp.AddNode(new PublishRocketSymbolsNode(bp, HostPlatform, TargetPlatforms, RemoteOutputDir + "Symbols"));
 
 				// Add a dependency on this being published as part of the shared promotable being labeled
 				GUBP.SharedLabelPromotableSuccessNode LabelPromotableNode = (GUBP.SharedLabelPromotableSuccessNode)bp.FindNode(GUBP.SharedLabelPromotableSuccessNode.StaticGetFullName());
@@ -326,16 +316,7 @@ namespace Rocket
 
 		public static bool IsRequiredForPlatform(UnrealTargetPlatform Platform)
 		{
-			switch(Platform)
-			{
-				case UnrealTargetPlatform.Win32:
-				case UnrealTargetPlatform.Win64:
-				case UnrealTargetPlatform.Android:
-				case UnrealTargetPlatform.Linux:
-					return true;
-				default:
-					return false;
-			}
+			return Platform != UnrealTargetPlatform.HTML5;
 		}
 
 		public override void DoBuild(GUBP bp)
@@ -432,10 +413,10 @@ namespace Rocket
 		BranchInfo.BranchUProject Project;
 		bool bIsCodeTargetPlatform;
 
-		public StripRocketMonolithicsNode(GUBP bp, UnrealTargetPlatform InHostPlatform, UnrealTargetPlatform InTargetPlatform, bool bInIsCodeTargetPlatform, string InStrippedDir) : base(InHostPlatform, InTargetPlatform, InStrippedDir)
+		public StripRocketMonolithicsNode(GUBP bp, UnrealTargetPlatform InHostPlatform, UnrealTargetPlatform InTargetPlatform, string InStrippedDir) : base(InHostPlatform, InTargetPlatform, InStrippedDir)
 		{
 			Project = bp.Branch.BaseEngineProject;
-			bIsCodeTargetPlatform = bInIsCodeTargetPlatform;
+			bIsCodeTargetPlatform = RocketBuild.IsCodeTargetPlatform(InHostPlatform, InTargetPlatform);
 
 			GUBP.GUBPNode Node = bp.FindNode(GUBP.GamePlatformMonolithicsNode.StaticGetFullName(HostPlatform, Project, InTargetPlatform, Precompiled: bIsCodeTargetPlatform));
 			if(String.IsNullOrEmpty(Node.AgentSharingGroup))
@@ -472,17 +453,15 @@ namespace Rocket
 	{
 		List<UnrealTargetPlatform> SourceHostPlatforms;
 		List<UnrealTargetPlatform> TargetPlatforms;
-		List<UnrealTargetPlatform> CodeTargetPlatforms;
 		string[] CurrentFeaturePacks;
 		string[] CurrentTemplates;
 		public readonly string DepotManifestPath;
 		public Dictionary<string, string> StrippedNodeManifestPaths = new Dictionary<string, string>();
 
-		public FilterRocketNode(GUBP bp, UnrealTargetPlatform InHostPlatform, List<UnrealTargetPlatform> InTargetPlatforms, List<UnrealTargetPlatform> InCodeTargetPlatforms, string[] InCurrentFeaturePacks, string[] InCurrentTemplates)
+		public FilterRocketNode(GUBP bp, UnrealTargetPlatform InHostPlatform, List<UnrealTargetPlatform> InTargetPlatforms, string[] InCurrentFeaturePacks, string[] InCurrentTemplates)
 			: base(InHostPlatform)
 		{
 			TargetPlatforms = new List<UnrealTargetPlatform>(InTargetPlatforms);
-			CodeTargetPlatforms = new List<UnrealTargetPlatform>(InCodeTargetPlatforms);
 			CurrentFeaturePacks = InCurrentFeaturePacks;
 			CurrentTemplates = InCurrentTemplates;
 			DepotManifestPath = CommandUtils.CombinePaths(CommandUtils.CmdEnv.LocalRoot, "Engine", "Saved", "Rocket", HostPlatform.ToString(), "Filter.txt");
@@ -493,17 +472,22 @@ namespace Rocket
 			AddDependency(GUBP.RootEditorNode.StaticGetFullName(HostPlatform));
 			AddDependency(GUBP.ToolsNode.StaticGetFullName(HostPlatform));
 
-			// Get all the external file lists. Android/HTML5 on Mac uses the same files as Win64.
-			List<string> StrippedNodeNames = new List<string>();
+			// Add all the monolithic builds from their appropriate source host platform
 			foreach(UnrealTargetPlatform TargetPlatform in TargetPlatforms)
 			{
 				UnrealTargetPlatform SourceHostPlatform = RocketBuild.GetSourceHostPlatform(bp, HostPlatform, TargetPlatform);
-				bool bIsCodeTargetPlatform = CodeTargetPlatforms.Contains(TargetPlatform);
+				bool bIsCodeTargetPlatform = RocketBuild.IsCodeTargetPlatform(SourceHostPlatform, TargetPlatform);
 				AddDependency(GUBP.GamePlatformMonolithicsNode.StaticGetFullName(SourceHostPlatform, bp.Branch.BaseEngineProject, TargetPlatform, Precompiled: bIsCodeTargetPlatform));
+			}
 
+			// Also add stripped symbols for all the target platforms that require it
+			List<string> StrippedNodeNames = new List<string>();
+			foreach(UnrealTargetPlatform TargetPlatform in TargetPlatforms)
+			{
 				if(StripRocketNode.IsRequiredForPlatform(TargetPlatform))
 				{
-					string StripNode = StripRocketMonolithicsNode.StaticGetFullName(SourceHostPlatform, bp.Branch.BaseEngineProject, TargetPlatform, bIsCodeTargetPlatform);
+					UnrealTargetPlatform SourceHostPlatform = RocketBuild.GetSourceHostPlatform(bp, HostPlatform, TargetPlatform);
+					string StripNode = StripRocketMonolithicsNode.StaticGetFullName(SourceHostPlatform, bp.Branch.BaseEngineProject, TargetPlatform, RocketBuild.IsCodeTargetPlatform(SourceHostPlatform, TargetPlatform));
 					AddDependency(StripNode);
 					StrippedNodeNames.Add(StripNode);
 				}
@@ -579,19 +563,23 @@ namespace Rocket
 			UnzipAndAddRuleForHeaders(GUBP.RootEditorNode.StaticGetArchivedHeadersPath(HostPlatform), Filter, FileFilterType.Include);
 
 			// Include the build dependencies for every code platform
-			foreach(UnrealTargetPlatform CodeTargetPlatform in CodeTargetPlatforms)
+			foreach(UnrealTargetPlatform TargetPlatform in TargetPlatforms)
 			{
-				UnrealTargetPlatform SourceHostPlatform = RocketBuild.GetSourceHostPlatform(bp, HostPlatform, CodeTargetPlatform);
-				string FileListPath = GUBP.GamePlatformMonolithicsNode.StaticGetBuildDependenciesPath(SourceHostPlatform, bp.Branch.BaseEngineProject, CodeTargetPlatform);
-				Filter.AddRuleForFiles(UnrealBuildTool.Utils.ReadClass<UnrealBuildTool.ExternalFileList>(FileListPath).FileNames, CommandUtils.CmdEnv.LocalRoot, FileFilterType.Include);
-				UnzipAndAddRuleForHeaders(GUBP.GamePlatformMonolithicsNode.StaticGetArchivedHeadersPath(SourceHostPlatform, bp.Branch.BaseEngineProject, CodeTargetPlatform), Filter, FileFilterType.Include);
+				if(RocketBuild.IsCodeTargetPlatform(HostPlatform, TargetPlatform))
+				{
+					UnrealTargetPlatform SourceHostPlatform = RocketBuild.GetSourceHostPlatform(bp, HostPlatform, TargetPlatform);
+					string FileListPath = GUBP.GamePlatformMonolithicsNode.StaticGetBuildDependenciesPath(SourceHostPlatform, bp.Branch.BaseEngineProject, TargetPlatform);
+					Filter.AddRuleForFiles(UnrealBuildTool.Utils.ReadClass<UnrealBuildTool.ExternalFileList>(FileListPath).FileNames, CommandUtils.CmdEnv.LocalRoot, FileFilterType.Include);
+					UnzipAndAddRuleForHeaders(GUBP.GamePlatformMonolithicsNode.StaticGetArchivedHeadersPath(SourceHostPlatform, bp.Branch.BaseEngineProject, TargetPlatform), Filter, FileFilterType.Include);
+				}
 			}
 
 			// Add the monolithic binaries
-			foreach(UnrealTargetPlatform CodeTargetPlatform in CodeTargetPlatforms)
+			foreach(UnrealTargetPlatform TargetPlatform in TargetPlatforms)
 			{
-				UnrealTargetPlatform SourceHostPlatform = RocketBuild.GetSourceHostPlatform(bp, HostPlatform, CodeTargetPlatform);
-				AddRuleForBuildProducts(Filter, bp, GUBP.GamePlatformMonolithicsNode.StaticGetFullName(SourceHostPlatform, bp.Branch.BaseEngineProject, CodeTargetPlatform, Precompiled: true), FileFilterType.Include);
+				UnrealTargetPlatform SourceHostPlatform = RocketBuild.GetSourceHostPlatform(bp, HostPlatform, TargetPlatform);
+				bool bIsCodeTargetPlatform = RocketBuild.IsCodeTargetPlatform(SourceHostPlatform, TargetPlatform);
+				AddRuleForBuildProducts(Filter, bp, GUBP.GamePlatformMonolithicsNode.StaticGetFullName(SourceHostPlatform, bp.Branch.BaseEngineProject, TargetPlatform, Precompiled: bIsCodeTargetPlatform), FileFilterType.Include);
 			}
 
 			// Include the feature packs
@@ -840,8 +828,11 @@ namespace Rocket
 
 			foreach(UnrealTargetPlatform TargetPlatform in TargetPlatforms)
 			{
-				UnrealTargetPlatform SourceHostPlatform = RocketBuild.GetSourceHostPlatform(bp, HostPlatform, TargetPlatform);
-				AddDependency(GUBP.GamePlatformMonolithicsNode.StaticGetFullName(SourceHostPlatform, bp.Branch.BaseEngineProject, TargetPlatform, Precompiled: true));
+				if(HostPlatform == RocketBuild.GetSourceHostPlatform(bp, HostPlatform, TargetPlatform))
+				{
+					bool bIsCodeTargetPlatform = RocketBuild.IsCodeTargetPlatform(HostPlatform, TargetPlatform);
+					AddDependency(GUBP.GamePlatformMonolithicsNode.StaticGetFullName(HostPlatform, bp.Branch.BaseEngineProject, TargetPlatform, Precompiled: bIsCodeTargetPlatform));
+				}
 			}
 
 			AgentSharingGroup = "RocketGroup" + StaticGetHostPlatformSuffix(HostPlatform);
