Fix inability to load some C# assemblies (UE-12883).

diff --git a/Engine/Source/Programs/AutomationTool/Scripts/GUBP.Automation.cs b/Engine/Source/Programs/AutomationTool/Scripts/GUBP.Automation.cs
--- a/Engine/Source/Programs/AutomationTool/Scripts/GUBP.Automation.cs
+++ b/Engine/Source/Programs/AutomationTool/Scripts/GUBP.Automation.cs
@@ -150,6 +150,21 @@ public class GUBP : BuildCommand
     }
 
     private static List<GUBPNodeAdder> Adders;
+
+	Type[] GetTypesFromAssembly(Assembly Asm)
+	{
+		Type[] AllTypesWeCanGet;
+		try
+		{
+			AllTypesWeCanGet = Asm.GetTypes();
+		}
+		catch(ReflectionTypeLoadException Exc)
+		{
+			AllTypesWeCanGet = Exc.Types;
+		}
+		return AllTypesWeCanGet;
+	}
+
     private void AddCustomNodes(UnrealTargetPlatform InHostPlatform)
     {
         if (Adders == null)
@@ -158,7 +173,7 @@ public class GUBP : BuildCommand
             Assembly[] LoadedAssemblies = AppDomain.CurrentDomain.GetAssemblies();
             foreach (var Dll in LoadedAssemblies)
             {
-                Type[] AllTypes = Dll.GetTypes();
+				Type[] AllTypes = GetTypesFromAssembly(Dll);
                 foreach (var PotentialConfigType in AllTypes)
                 {
                     if (PotentialConfigType != typeof(GUBPNodeAdder) && typeof(GUBPNodeAdder).IsAssignableFrom(PotentialConfigType))
@@ -207,7 +222,7 @@ public class GUBP : BuildCommand
             Assembly[] LoadedAssemblies = AppDomain.CurrentDomain.GetAssemblies();
             foreach (var Dll in LoadedAssemblies)
             {
-                Type[] AllTypes = Dll.GetTypes();
+				Type[] AllTypes = GetTypesFromAssembly(Dll);
                 foreach (var PotentialConfigType in AllTypes)
                 {
                     if (PotentialConfigType != typeof(GUBPBranchHacker) && typeof(GUBPBranchHacker).IsAssignableFrom(PotentialConfigType))
@@ -268,7 +283,7 @@ public class GUBP : BuildCommand
             Assembly[] LoadedAssemblies = AppDomain.CurrentDomain.GetAssemblies();
             foreach (var Dll in LoadedAssemblies)
             {
-                Type[] AllTypes = Dll.GetTypes();
+				Type[] AllTypes = GetTypesFromAssembly(Dll);
                 foreach (var PotentialConfigType in AllTypes)
                 {
                     if (PotentialConfigType != typeof(GUBPEmailHacker) && typeof(GUBPEmailHacker).IsAssignableFrom(PotentialConfigType))
@@ -330,7 +345,7 @@ public class GUBP : BuildCommand
             Assembly[] LoadedAssemblies = AppDomain.CurrentDomain.GetAssemblies();
             foreach (var Dll in LoadedAssemblies)
             {
-                Type[] AllTypes = Dll.GetTypes();
+				Type[] AllTypes = GetTypesFromAssembly(Dll);
                 foreach (var PotentialConfigType in AllTypes)
                 {
                     if (PotentialConfigType != typeof(GUBPFrequencyHacker) && typeof(GUBPFrequencyHacker).IsAssignableFrom(PotentialConfigType))
