Make the promotion trigger depend on the Rocket promotable steps succeeding, so we can select the Shared_LabelPromoted node when manually triggering a build again.

diff --git a/Engine/Source/Programs/AutomationTool/Scripts/RocketBuild.Automation.cs b/Engine/Source/Programs/AutomationTool/Scripts/RocketBuild.Automation.cs
--- a/Engine/Source/Programs/AutomationTool/Scripts/RocketBuild.Automation.cs
+++ b/Engine/Source/Programs/AutomationTool/Scripts/RocketBuild.Automation.cs
@@ -131,6 +131,11 @@ namespace Rocket
 				GUBP.SharedLabelPromotableSuccessNode LabelPromotableNode = (GUBP.SharedLabelPromotableSuccessNode)bp.FindNode(GUBP.SharedLabelPromotableSuccessNode.StaticGetFullName());
 				LabelPromotableNode.AddDependency(PublishRocketNode.StaticGetFullName(HostPlatform));
 				LabelPromotableNode.AddDependency(PublishRocketSymbolsNode.StaticGetFullName(HostPlatform));
+
+				// Add dependencies on a promotable to do these steps too
+				GUBP.WaitForSharedPromotionUserInput WaitForPromotionNode = (GUBP.WaitForSharedPromotionUserInput)bp.FindNode(GUBP.WaitForSharedPromotionUserInput.StaticGetFullName(true));
+				WaitForPromotionNode.AddDependency(PublishRocketNode.StaticGetFullName(HostPlatform));
+				WaitForPromotionNode.AddDependency(PublishRocketSymbolsNode.StaticGetFullName(HostPlatform));
 			}
 		}
 
@@ -694,12 +699,6 @@ namespace Rocket
 			OutputDir = InOutputDir;
 			CodeTargetPlatforms = new List<UnrealTargetPlatform>(InCodeTargetPlatforms);
 
-			// Optimization for testing locally; don't do other things that are part of the shared promotable
-			if(RocketBuild.ShouldDoSeriousThingsLikeP4CheckinAndPostToMCP())
-			{
-				AddDependency(GUBP.WaitForSharedPromotionUserInput.StaticGetFullName(false));
-			}
-
 			AddDependency(FilterRocketNode.StaticGetFullName(HostPlatform));
 			AddDependency(BuildDerivedDataCacheNode.StaticGetFullName(HostPlatform));
 
