Add dependency on the promotable trigger back in.

diff --git a/Engine/Source/Programs/AutomationTool/Scripts/RocketBuild.Automation.cs b/Engine/Source/Programs/AutomationTool/Scripts/RocketBuild.Automation.cs
--- a/Engine/Source/Programs/AutomationTool/Scripts/RocketBuild.Automation.cs
+++ b/Engine/Source/Programs/AutomationTool/Scripts/RocketBuild.Automation.cs
@@ -699,6 +699,12 @@ namespace Rocket
 			OutputDir = InOutputDir;
 			CodeTargetPlatforms = new List<UnrealTargetPlatform>(InCodeTargetPlatforms);
 
+			// Optimization for testing locally; don't do other things that are part of the shared promotable
+			if(RocketBuild.ShouldDoSeriousThingsLikeP4CheckinAndPostToMCP())
+			{
+				AddDependency(GUBP.WaitForSharedPromotionUserInput.StaticGetFullName(false));
+			}
+
 			AddDependency(FilterRocketNode.StaticGetFullName(HostPlatform));
 			AddDependency(BuildDerivedDataCacheNode.StaticGetFullName(HostPlatform));
 
