Add the -WithRocketPromotable argument to allow testing Rocket builds behind promotion triggers locally.

diff --git a/Engine/Source/Programs/AutomationTool/Scripts/RocketBuild.Automation.cs b/Engine/Source/Programs/AutomationTool/Scripts/RocketBuild.Automation.cs
--- a/Engine/Source/Programs/AutomationTool/Scripts/RocketBuild.Automation.cs
+++ b/Engine/Source/Programs/AutomationTool/Scripts/RocketBuild.Automation.cs
@@ -136,6 +136,18 @@ namespace Rocket
 				GUBP.WaitForSharedPromotionUserInput WaitForPromotionNode = (GUBP.WaitForSharedPromotionUserInput)bp.FindNode(GUBP.WaitForSharedPromotionUserInput.StaticGetFullName(true));
 				WaitForPromotionNode.AddDependency(PublishRocketNode.StaticGetFullName(HostPlatform));
 				WaitForPromotionNode.AddDependency(PublishRocketSymbolsNode.StaticGetFullName(HostPlatform));
+
+				// Push everything behind the promotion triggers if we're doing things on the build machines
+				if(ShouldDoSeriousThingsLikeP4CheckinAndPostToMCP() || bp.ParseParam("WithRocketPromotable"))
+				{
+					string WaitForTrigger = GUBP.WaitForSharedPromotionUserInput.StaticGetFullName(false);
+
+					GatherRocketNode GatherRocket = (GatherRocketNode)bp.FindNode(GatherRocketNode.StaticGetFullName(HostPlatform));
+					GatherRocket.AddDependency(WaitForTrigger);
+
+					PublishRocketSymbolsNode PublishRocketSymbols = (PublishRocketSymbolsNode)bp.FindNode(PublishRocketSymbolsNode.StaticGetFullName(HostPlatform));
+					PublishRocketSymbols.AddDependency(WaitForTrigger);
+				}
 			}
 		}
 
@@ -699,12 +711,6 @@ namespace Rocket
 			OutputDir = InOutputDir;
 			CodeTargetPlatforms = new List<UnrealTargetPlatform>(InCodeTargetPlatforms);
 
-			// Optimization for testing locally; don't do other things that are part of the shared promotable
-			if(RocketBuild.ShouldDoSeriousThingsLikeP4CheckinAndPostToMCP())
-			{
-				AddDependency(GUBP.WaitForSharedPromotionUserInput.StaticGetFullName(false));
-			}
-
 			AddDependency(FilterRocketNode.StaticGetFullName(HostPlatform));
 			AddDependency(BuildDerivedDataCacheNode.StaticGetFullName(HostPlatform));
 
