Linux: do not stage LND for dedicated servers.

diff --git a/Engine/Source/Programs/AutomationTool/Linux/LinuxPlatform.Automation.cs b/Engine/Source/Programs/AutomationTool/Linux/LinuxPlatform.Automation.cs
--- a/Engine/Source/Programs/AutomationTool/Linux/LinuxPlatform.Automation.cs
+++ b/Engine/Source/Programs/AutomationTool/Linux/LinuxPlatform.Automation.cs
@@ -53,6 +53,7 @@ public abstract class BaseLinuxPlatform : Platform
 		}
 
 		// stage libLND
+		if (!SC.DedicatedServer)
 		{
 			SC.StageFiles(StagedFileType.NonUFS, CombinePaths(SC.LocalRoot, "Engine/Binaries/ThirdParty/LinuxNativeDialogs/", SC.PlatformDir, BuildArchitecture), "libLND*.so");
 		}
