UE-14741 - OBB Fails to deploy to Android devices with 4.1.1 installed
- Changed OBB install script to deal with this setup

diff --git a/Engine/Source/Programs/AutomationTool/Android/AndroidPlatform.Automation.cs b/Engine/Source/Programs/AutomationTool/Android/AndroidPlatform.Automation.cs
--- a/Engine/Source/Programs/AutomationTool/Android/AndroidPlatform.Automation.cs
+++ b/Engine/Source/Programs/AutomationTool/Android/AndroidPlatform.Automation.cs
@@ -205,6 +205,7 @@ public class AndroidPlatform : Platform
 						"#!/bin/sh",
 						"cd \"`dirname \"$0\"`\"",
                         "ADB=",
+                        "PREFIX=Android/",
 						"if [ \"$ANDROID_HOME\" != \"\" ]; then ADB=$ANDROID_HOME/platform-tools/adb; else ADB=" +Environment.GetEnvironmentVariable("ANDROID_HOME") + "; fi",
 						"DEVICE=",
 						"if [ \"$1\" != \"\" ]; then DEVICE=\"-s $1\"; fi",
@@ -217,14 +218,14 @@ public class AndroidPlatform : Platform
 						"if [ $? -eq 0 ]; then",
 						"\techo",
                         "\tNEWVERSION=$(echo \"`$ADB $DEVICE shell 'getprop ro.build.version.release'`\" | awk -F\".\" '{gsub(/-.*/,\"\",$4); if ($1 > 4 || ($1 == 4 && $2 >= 2 )) print 1; else print 0 }')",
-                        "\t if [ $NEWVERSION -eq 1 ]; then STORAGE=/mnt/shell/emulated; else STORAGE=$(echo \"`$ADB $DEVICE shell 'echo $EXTERNAL_STORAGE'`\" | cat -v | tr -d '^M'); fi",
+                        "\t if [ $NEWVERSION -eq 1 ]; then STORAGE=/mnt/shell/emulated; PREFIX=; else STORAGE=$(echo \"`$ADB $DEVICE shell 'echo $EXTERNAL_STORAGE'`\" | cat -v | tr -d '^M'); fi",
 						"\techo Removing old data. Failures here are usually fine - indicating the files were not on the device.",
                         "\t$ADB $DEVICE shell rm -r $STORAGE/UE4Game/" + Params.ShortProjectName,
 						"\t$ADB $DEVICE shell rm -r $STORAGE/UE4Game/UE4CommandLine.txt",
 						"\t$ADB $DEVICE shell rm -r $STORAGE/" + TargetAndroidLocation + PackageName,
 						bPackageDataInsideApk ? "" : "\techo",
 						bPackageDataInsideApk ? "" : "\techo Installing new data. Failures here indicate storage problems \\(missing SD card or bad permissions\\) and are fatal.",
-						bPackageDataInsideApk ? "" : "\t$ADB $DEVICE push " + Path.GetFileName(ObbName) + " $STORAGE/" + DeviceObbName,
+						bPackageDataInsideApk ? "" : "\t$ADB $DEVICE push " + Path.GetFileName(ObbName) + " $STORAGE/$PREFIX" + DeviceObbName,
 						bPackageDataInsideApk ? "if [ 1 ]; then" : "\tif [ $? -eq 0 ]; then",
 						"\t\techo",
 						"\t\techo Installation successful",
@@ -250,15 +251,22 @@ public class AndroidPlatform : Platform
                         "if \"%ANDROIDHOME%\"==\"\" set ANDROIDHOME="+Environment.GetEnvironmentVariable("ANDROID_HOME"),
 						"set ADB=%ANDROIDHOME%\\platform-tools\\adb.exe",
 						"set DEVICE=",
+                        "set PREFIX=Android/",
                         "if not \"%1\"==\"\" set DEVICE=-s %1",
                         "for /f \"delims=\" %%A in ('%ADB% %DEVICE% " + GetStorageQueryCommand() +"') do @set STORAGE=%%A",
                         "for /f \"tokens=1-2 delims=.\" %%M in ('%ADB% %DEVICE% shell getprop ro.build.version.release') do (",
                         "    if \"%%M\"==\"\" (goto completed)",
                         "    if \"%%N\"==\"\" (goto completed)",
                         "    if %%M equ 4 (",     
-                        "        if %%N geq 2 (set STORAGE=/mnt/shell/emulated)",
+                        "        if %%N geq 2 (",
+                        "            set STORAGE=/mnt/shell/emulated",
+                        "            set PREFIX=",
+                        "        );",
                         "    ) else (",
-                        "       if %%M gtr 4 (set STORAGE=/mnt/shell/emulated)",
+                        "       if %%M gtr 4 (",
+                        "            set STORAGE=/mnt/shell/emulated",
+                        "            set PREFIX=",
+                        "       );",
                         "    );",                        
                         "    @echo.",
                         ")",
@@ -275,7 +283,7 @@ public class AndroidPlatform : Platform
 						"%ADB% %DEVICE% shell rm -r %STORAGE%/" + TargetAndroidLocation + PackageName,
 						bPackageDataInsideApk ? "" : "@echo.",
 						bPackageDataInsideApk ? "" : "@echo Installing new data. Failures here indicate storage problems (missing SD card or bad permissions) and are fatal.",
-						bPackageDataInsideApk ? "" : "%ADB% %DEVICE% push " + Path.GetFileName(ObbName) + " %STORAGE%/" + DeviceObbName,
+						bPackageDataInsideApk ? "" : "%ADB% %DEVICE% push " + Path.GetFileName(ObbName) + " %STORAGE%/%PREFIX%" + DeviceObbName,
 						bPackageDataInsideApk ? "" : "if \"%ERRORLEVEL%\" NEQ \"0\" goto Error",
 						"@echo.",
 						"@echo Installation successful",
