fix for incorrect working directory
UE-14735
#uat

diff --git a/Engine/Source/Programs/AutomationTool/ProjectUtils.cs b/Engine/Source/Programs/AutomationTool/ProjectUtils.cs
--- a/Engine/Source/Programs/AutomationTool/ProjectUtils.cs
+++ b/Engine/Source/Programs/AutomationTool/ProjectUtils.cs
@@ -181,12 +181,15 @@ namespace AutomationTool
 
 			// Change the working directory to be the Engine/Source folder. We are running from Engine/Binaries/DotNET
 			string oldCWD = Directory.GetCurrentDirectory();
-			string EngineSourceDirectory = Path.Combine(UnrealBuildTool.Utils.GetExecutingAssemblyDirectory(), "..", "..", "..", "Engine", "Source");
-			if (!Directory.Exists(EngineSourceDirectory)) // only set the directory if it exists, this should only happen if we are launching the editor from an artist sync
+			if (BuildConfiguration.RelativeEnginePath == "../../Engine/")
 			{
-				EngineSourceDirectory = Path.Combine(UnrealBuildTool.Utils.GetExecutingAssemblyDirectory(), "..", "..", "..", "Engine", "Binaries");
+				string EngineSourceDirectory = Path.Combine(UnrealBuildTool.Utils.GetExecutingAssemblyDirectory(), "..", "..", "..", "Engine", "Source");
+				if (!Directory.Exists(EngineSourceDirectory)) // only set the directory if it exists, this should only happen if we are launching the editor from an artist sync
+				{
+					EngineSourceDirectory = Path.Combine(UnrealBuildTool.Utils.GetExecutingAssemblyDirectory(), "..", "..", "..", "Engine", "Binaries");
+				}
+				Directory.SetCurrentDirectory(EngineSourceDirectory);
 			}
-			Directory.SetCurrentDirectory(EngineSourceDirectory);
 
 			// Read the project descriptor, and find all the plugins available to this project
 			ProjectDescriptor Project = ProjectDescriptor.FromFile(RawProjectPath);
