Merging using UE4-To-UE4-LauncherDev

diff --git a/Engine/Source/Programs/AutomationTool/Scripts/LauncherLocalization.Automation.cs b/Engine/Source/Programs/AutomationTool/Scripts/LauncherLocalization.Automation.cs
--- a/Engine/Source/Programs/AutomationTool/Scripts/LauncherLocalization.Automation.cs
+++ b/Engine/Source/Programs/AutomationTool/Scripts/LauncherLocalization.Automation.cs
@@ -24,9 +24,9 @@ class LauncherLocalization : BuildCommand
 			P4.Sync(P4Env.BuildRootP4 + "/Engine/Content/...");
 			P4.Sync(P4Env.BuildRootP4 + "/Engine/Source/...");
 
-			P4.Sync(P4Env.BuildRootP4 + "/Engine/Programs/NoRedist/UnrealEngineLauncher/Config/...");
-			P4.Sync(P4Env.BuildRootP4 + "/Engine/Programs/NoRedist/UnrealEngineLauncher/Content/...");
-			//P4.Sync(P4Env.BuildRootP4 + "/Engine/Source/..."); <- takes care of syncing Launcher source already
+            P4.Sync(P4Env.BuildRootP4 + "/Portal/Config/...");
+            P4.Sync(P4Env.BuildRootP4 + "/Portal/Content/...");
+            P4.Sync(P4Env.BuildRootP4 + "/Portal/Source/...");
 
 			Log("Localize from label {0}", P4Env.LabelToSync);
 		}
@@ -43,7 +43,7 @@ class LauncherLocalization : BuildCommand
 			//Export
 			if (appFile != null)
 			{
-				ExportFileToDirectory(appFile, new DirectoryInfo(CmdEnv.LocalRoot + "/Engine/Programs/NoRedist/UnrealEngineLauncher/Content/Localization/App"), launcherGroup.EnabledCultures);
+				ExportFileToDirectory(appFile, new DirectoryInfo(CmdEnv.LocalRoot + "/Portal/Content/Localization/App"), launcherGroup.EnabledCultures);
 			}
 		}
 
@@ -66,8 +66,8 @@ class LauncherLocalization : BuildCommand
 		// Setup commandlet arguments with configurations.
 		var CommandletArgumentSets = new string[] 
 			{
-				String.Format("-config={0}", @"./Programs/NoRedist/UnrealEngineLauncher/Config/Localization/App.ini") + (string.IsNullOrEmpty(CommandletSCCArguments) ? "" : " " + CommandletSCCArguments),
-				String.Format("-config={0}", @"./Programs/NoRedist/UnrealEngineLauncher/Config/Localization/WordCount.ini"),
+				String.Format("-config={0}", @"../Portal/Config/Localization/App.ini") + (string.IsNullOrEmpty(CommandletSCCArguments) ? "" : " " + CommandletSCCArguments),
+				String.Format("-config={0}", @"../Portal/Config/Localization/WordCount.ini"),
 			};
 
 		// Execute commandlet for each set of arguments.
@@ -87,7 +87,7 @@ class LauncherLocalization : BuildCommand
 		}
 
 		// Upload Launcher text to OneSky
-		UploadDirectoryToProject(GetAppProject(oneSkyService), new DirectoryInfo(CmdEnv.LocalRoot + "/Engine/Programs/NoRedist/UnrealEngineLauncher/Content/Localization/App"), "*.po");
+		UploadDirectoryToProject(GetAppProject(oneSkyService), new DirectoryInfo(CmdEnv.LocalRoot + "/Portal/Content/Localization/App"), "*.po");
 	}
 
 	private static ProjectGroup GetLauncherGroup(OneSkyService oneSkyService)
