Add UAT script to build a plugin for distribution. Includes UBT support for writing plugin descriptors, and compiling single plugins at a time.

diff --git a/Engine/Source/Programs/AutomationTool/Scripts/AutomationScripts.Automation.csproj b/Engine/Source/Programs/AutomationTool/Scripts/AutomationScripts.Automation.csproj
--- a/Engine/Source/Programs/AutomationTool/Scripts/AutomationScripts.Automation.csproj
+++ b/Engine/Source/Programs/AutomationTool/Scripts/AutomationScripts.Automation.csproj
@@ -64,6 +64,7 @@
     <Compile Include="BuildCommonTools.Automation.cs" />
     <Compile Include="BuildCookRun.Automation.cs" />
     <Compile Include="BuildDocumentation.Automation.cs" />
+    <Compile Include="BuildPluginCommand.Automation.cs" />
     <Compile Include="BuildProjectCommand.Automation.cs" />
     <Compile Include="BuildThirdPartyLibs.Automation.cs" />
     <Compile Include="CodeSurgery.Automation.cs" />
diff --git a/Engine/Source/Programs/AutomationTool/Scripts/BuildPluginCommand.Automation.cs b/Engine/Source/Programs/AutomationTool/Scripts/BuildPluginCommand.Automation.cs
new file mode 100644
--- /dev/null
+++ b/Engine/Source/Programs/AutomationTool/Scripts/BuildPluginCommand.Automation.cs
@@ -0,0 +1,180 @@
+﻿// Copyright 1998-2015 Epic Games, Inc. All Rights Reserved.
+using System;
+using System.Collections.Generic;
+using System.IO;
+using System.Threading;
+using System.Linq;
+using System.Reflection;
+using AutomationTool;
+using UnrealBuildTool;
+
+[Help("Builds a plugin, and packages it for distribution")]
+class BuildPlugin : BuildCommand
+{
+	public override void ExecuteBuild()
+	{
+		// Get the plugin filename
+		string PluginFileName = ParseParamValue("Plugin");
+		if(PluginFileName == null)
+		{
+			throw new AutomationException("Plugin file name was not specified via the -plugin argument");
+		}
+
+		// Get the platforms to build for
+		List<UnrealTargetPlatform> EditorPlatforms = ParseParamPlatforms("EditorPlatforms");
+		List<UnrealTargetPlatform> GamePlatforms = ParseParamPlatforms("GamePlatforms");
+
+		// Read the plugin
+		PluginDescriptor Plugin = PluginDescriptor.FromFile(PluginFileName);
+
+		// Clean the intermediate build directory
+		string IntermediateBuildDirectory = Path.Combine(Path.GetDirectoryName(PluginFileName), "Intermediate", "Build");
+		if(CommandUtils.DirectoryExists(IntermediateBuildDirectory))
+		{
+			CommandUtils.DeleteDirectory(IntermediateBuildDirectory);
+		}
+
+		// Build everything
+		List<string> ReceiptFileNames = new List<string>();
+		UE4Build.BuildAgenda Agenda = new UE4Build.BuildAgenda();
+		AddPluginToAgenda(Agenda, PluginFileName, Plugin, EditorPlatforms, GamePlatforms, ReceiptFileNames);
+		UE4Build Build = new UE4Build(this);
+		Build.Build(Agenda, InDeleteBuildProducts: true, InUpdateVersionFiles: false);
+
+		// Package the plugin to the output folder
+		string PackageDirectory = ParseParamValue("Package");
+		if(PackageDirectory != null)
+		{
+			List<BuildProduct> BuildProducts = GetBuildProductsFromReceipts(ReceiptFileNames);
+			PackagePlugin(PluginFileName, BuildProducts, PackageDirectory);
+		}
+	}
+
+	List<UnrealTargetPlatform> ParseParamPlatforms(string ParamName)
+	{
+		// Parse the raw parameter
+		string ParamValue = ParseParamValue(ParamName, null);
+
+		// Convert it to a list of target platform enums
+		List<UnrealTargetPlatform> Platforms = new List<UnrealTargetPlatform>();
+		if(ParamValue != null)
+		{
+			foreach(string PlatformName in ParamValue.Split(new char[]{ '+' }, StringSplitOptions.RemoveEmptyEntries))
+			{
+				UnrealTargetPlatform Platform;
+				if(!Enum.TryParse<UnrealTargetPlatform>(PlatformName, true, out Platform))
+				{
+					throw new AutomationException("'{0}' is not a valid platform name; valid platforms are: {1}.", PlatformName, String.Join(", ", Enum.GetNames(typeof(UnrealTargetPlatform))));
+				}
+				Platforms.Add(Platform);
+			}
+		}
+		return Platforms;
+	}
+
+	static void AddPluginToAgenda(UE4Build.BuildAgenda Agenda, string PluginFileName, PluginDescriptor Plugin, List<UnrealTargetPlatform> EditorPlatforms, List<UnrealTargetPlatform> GamePlatforms, List<string> ReceiptFileNames)
+	{
+		// Add the editor targets
+		foreach(UnrealTargetPlatform EditorPlatform in EditorPlatforms)
+		{
+			AddPluginToAgenda(Agenda, PluginFileName, Plugin, "UE4Editor", TargetRules.TargetType.Editor, EditorPlatform, UnrealTargetConfiguration.Development, ReceiptFileNames);
+		}
+
+		// Add the game targets
+		foreach(UnrealTargetPlatform GamePlatform in GamePlatforms)
+		{
+			AddPluginToAgenda(Agenda, PluginFileName, Plugin, "UE4Game", TargetRules.TargetType.Game, GamePlatform, UnrealTargetConfiguration.Development, ReceiptFileNames);
+			AddPluginToAgenda(Agenda, PluginFileName, Plugin, "UE4Game", TargetRules.TargetType.Game, GamePlatform, UnrealTargetConfiguration.Shipping, ReceiptFileNames);
+		}
+	}
+
+	static void AddPluginToAgenda(UE4Build.BuildAgenda Agenda, string PluginFileName, PluginDescriptor Plugin, string TargetName, TargetRules.TargetType TargetType, UnrealTargetPlatform Platform, UnrealTargetConfiguration Configuration, List<string> ReceiptFileNames)
+	{
+		// Find a list of modules that need to be built for this plugin
+		List<string> ModuleNames = new List<string>();
+		foreach(ModuleDescriptor Module in Plugin.Modules)
+		{
+			if(Module.IsCompiledInConfiguration(Platform, TargetType))
+			{
+				ModuleNames.Add(Module.Name);
+			}
+		}
+
+		// Add these modules to the build agenda
+		if(ModuleNames.Count > 0)
+		{
+			string Arguments = String.Format("-plugin {0}", CommandUtils.MakePathSafeToUseWithCommandLine(PluginFileName));
+			foreach(string ModuleName in ModuleNames)
+			{
+				Arguments += String.Format(" -module {0}", ModuleName);
+			}
+
+			string ReceiptFileName = BuildReceipt.GetDefaultPath(Path.GetDirectoryName(PluginFileName), TargetName, Platform, Configuration, "");
+			Arguments += String.Format(" -receipt {0}", CommandUtils.MakePathSafeToUseWithCommandLine(ReceiptFileName));
+			ReceiptFileNames.Add(ReceiptFileName);
+
+			Agenda.AddTarget(TargetName, Platform, Configuration, InAddArgs: Arguments);
+		}
+	}
+
+	static List<BuildProduct> GetBuildProductsFromReceipts(List<string> ReceiptFileNames)
+	{
+		List<BuildProduct> BuildProducts = new List<BuildProduct>();
+		foreach(string ReceiptFileName in ReceiptFileNames)
+		{
+			BuildReceipt Receipt = BuildReceipt.Read(ReceiptFileName);
+			BuildProducts.AddRange(Receipt.BuildProducts);
+		}
+		return BuildProducts;
+	}
+
+	static void PackagePlugin(string PluginFileName, List<BuildProduct> BuildProducts, string PackageDirectory)
+	{
+		// Clear the output directory
+		CommandUtils.DeleteDirectoryContents(PackageDirectory);
+
+		// Copy all the files to the output directory
+		List<string> MatchingFileNames = FilterPluginFiles(PluginFileName, BuildProducts);
+		foreach(string MatchingFileName in MatchingFileNames)
+		{
+			string SourceFileName = Path.Combine(Path.GetDirectoryName(PluginFileName), MatchingFileName);
+			string TargetFileName = Path.Combine(PackageDirectory, MatchingFileName);
+			CommandUtils.CopyFile(SourceFileName, TargetFileName);
+		}
+
+		// Get the output plugin filename
+		string TargetPluginFileName = CommandUtils.MakeRerootedFilePath(Path.GetFullPath(PluginFileName), Path.GetDirectoryName(Path.GetFullPath(PluginFileName)), PackageDirectory);
+		PluginDescriptor NewDescriptor = PluginDescriptor.FromFile(TargetPluginFileName);
+		NewDescriptor.bEnabledByDefault = false;
+		NewDescriptor.bInstalled = true;
+		NewDescriptor.Save(TargetPluginFileName);
+	}
+
+	static List<string> FilterPluginFiles(string PluginFileName, List<BuildProduct> BuildProducts)
+	{
+		string PluginDirectory = Path.GetDirectoryName(PluginFileName);
+
+		// Set up the default filter
+		FileFilter Filter = new FileFilter();
+		Filter.AddRuleForFile(PluginFileName, PluginDirectory, FileFilterType.Include);
+		Filter.AddRuleForFiles(BuildProducts.Select(x => x.Path), PluginDirectory, FileFilterType.Include);
+		Filter.Include("/Resources/...");
+		Filter.Include("/Content/...");
+		Filter.Include("/Intermediate/Build/.../Inc/...");
+		Filter.Include("/Source/...");
+
+		// Add custom rules for each platform
+		string FilterFileName = Path.Combine(Path.GetDirectoryName(PluginFileName), "Config", "FilterPlugin.ini");
+		if(File.Exists(FilterFileName))
+		{
+			Filter.ReadRulesFromFile(FilterFileName, "FilterPlugin");
+		}
+
+		// Apply the standard exclusion rules
+		Filter.ExcludeConfidentialFolders();
+		Filter.ExcludeConfidentialPlatforms();
+
+		// Apply the filter to the plugin directory
+		return new List<string>(Filter.ApplyToDirectory(PluginDirectory, true));
+	}
+}
diff --git a/Engine/Source/Programs/UnrealBuildTool/Configuration/UEBuildTarget.cs b/Engine/Source/Programs/UnrealBuildTool/Configuration/UEBuildTarget.cs
--- a/Engine/Source/Programs/UnrealBuildTool/Configuration/UEBuildTarget.cs
+++ b/Engine/Source/Programs/UnrealBuildTool/Configuration/UEBuildTarget.cs
@@ -168,6 +168,8 @@ namespace UnrealBuildTool
 		public List<OnlyModule> OnlyModules;
 		public bool bPrecompile;
 		public bool bUsePrecompiled;
+		public List<string> ForeignPlugins;
+		public string ForceReceiptFileName;
 	}
 
 
@@ -217,6 +219,8 @@ namespace UnrealBuildTool
 			var Configuration = UnrealTargetConfiguration.Unknown;
 			string RemoteRoot = null;
 			var OnlyModules = new List<OnlyModule>();
+			List<string> ForeignPlugins = new List<string>();
+			string ForceReceiptFileName = null;
 
 			// If true, the recompile was launched by the editor.
 			bool bIsEditorRecompile = false;
@@ -269,6 +273,28 @@ namespace UnrealBuildTool
 							}
 							break;
 
+						case "-PLUGIN":
+							{
+								if (ArgumentIndex + 1 >= Arguments.Count)
+								{
+									throw new BuildException("Expected plugin filename after -Plugin argument, but found nothing.");
+								}
+
+								ForeignPlugins.Add(Arguments[++ArgumentIndex]);
+							}
+							break;
+
+						case "-RECEIPT":
+							{
+								if (ArgumentIndex + 1 >= Arguments.Count)
+								{
+									throw new BuildException("Expected path to the generated receipt after -Receipt argument, but found nothing.");
+								}
+
+								ForceReceiptFileName = Arguments[++ArgumentIndex];
+							}
+							break;
+
 						// Configuration names:
 						case "DEBUG":
 							Configuration = UnrealTargetConfiguration.Debug;
@@ -470,6 +496,8 @@ namespace UnrealBuildTool
 							OnlyModules = OnlyModules,
 							bPrecompile = bPrecompile,
 							bUsePrecompiled = bUsePrecompiled,
+							ForeignPlugins = ForeignPlugins,
+							ForceReceiptFileName = ForceReceiptFileName
 						} );
 					break;
 				}
@@ -495,7 +523,7 @@ namespace UnrealBuildTool
 
 				// Scan the disk to find source files for all known targets and modules, and generate "in memory" project
 				// file data that will be used to determine what to build
-				RulesCompiler.SetAssemblyNameAndGameFolders( PossibleAssemblyName, UEBuildTarget.DiscoverAllGameFolders() );
+				RulesCompiler.SetAssemblyNameAndGameFolders( PossibleAssemblyName, UEBuildTarget.DiscoverAllGameFolders(), Desc.ForeignPlugins );
 			}
 
 			// Try getting it from the RulesCompiler
@@ -703,6 +731,10 @@ namespace UnrealBuildTool
 		[NonSerialized]
 		public List<PluginInfo> EnabledPlugins;
 
+		/** Additional plugin filenames which are foreign to this target */
+		[NonSerialized]
+		public List<string> ForeignPlugins;
+
 		/** All application binaries; may include binaries not built by this target. */
 		[NonSerialized]
 		public List<UEBuildBinary> AppBinaries = new List<UEBuildBinary>();
@@ -732,6 +764,10 @@ namespace UnrealBuildTool
 		/** The receipt for this target, which contains a record of this build. */
 		private BuildReceipt Receipt;
 
+		/** Force output of the receipt to an additional filename */
+		[NonSerialized]
+		private string ForceReceiptFileName;
+
 		/// <summary>
 		/// Whether this target should be compiled in monolithic mode
 		/// </summary>
@@ -758,6 +794,8 @@ namespace UnrealBuildTool
 			bEditorRecompile = InDesc.bIsEditorRecompile;
 			bPrecompile = InDesc.bPrecompile;
 			bUsePrecompiled = InDesc.bUsePrecompiled;
+			ForeignPlugins = InDesc.ForeignPlugins;
+			ForceReceiptFileName = InDesc.ForceReceiptFileName;
 
 			{
 				bCompileMonolithic = Rules.ShouldCompileMonolithic(InDesc.Platform, InDesc.Configuration);
@@ -1304,7 +1342,10 @@ namespace UnrealBuildTool
 			}
 
 			IUEBuildPlatform BuildPlatform = UEBuildPlatform.GetBuildPlatform(Platform);
-			Manifest.AddBuildProduct(BuildReceipt.GetDefaultPath(ProjectDirectory, TargetName, Platform, Configuration, BuildPlatform.GetActiveArchitecture()));
+			if(OnlyModules.Count == 0)
+			{
+				Manifest.AddBuildProduct(BuildReceipt.GetDefaultPath(ProjectDirectory, TargetName, Platform, Configuration, BuildPlatform.GetActiveArchitecture()));
+			}
 
 			if (UEBuildConfiguration.bCleanProject)
 			{
@@ -1344,9 +1385,17 @@ namespace UnrealBuildTool
 			if(Receipt != null)
 			{
 				IUEBuildPlatform BuildPlatform = UEBuildPlatform.GetBuildPlatform(Platform);
-				string ReceiptFileName = BuildReceipt.GetDefaultPath(ProjectDirectory, TargetName, Platform, Configuration, BuildPlatform.GetActiveArchitecture());
-				Directory.CreateDirectory(Path.GetDirectoryName(ReceiptFileName));
-				Receipt.Write(ReceiptFileName);
+				if(OnlyModules.Count == 0)
+				{
+					string ReceiptFileName = BuildReceipt.GetDefaultPath(ProjectDirectory, TargetName, Platform, Configuration, BuildPlatform.GetActiveArchitecture());
+					Directory.CreateDirectory(Path.GetDirectoryName(ReceiptFileName));
+					Receipt.Write(ReceiptFileName);
+				}
+				if(ForceReceiptFileName != null)
+				{
+					Directory.CreateDirectory(Path.GetDirectoryName(ForceReceiptFileName));
+					Receipt.Write(ForceReceiptFileName);
+				}
 			}
 		}
 
@@ -2513,6 +2562,17 @@ namespace UnrealBuildTool
 			{
 				BuildPlugins = new List<PluginInfo>(EnabledPlugins);
 			}
+
+			// Add any foreign plugins to the list
+			if(ForeignPlugins != null)
+			{
+				foreach(string ForeignPlugin in ForeignPlugins)
+				{
+					PluginInfo ForeignPluginInfo = new PluginInfo(ForeignPlugin, PluginLoadedFrom.GameProject);
+					ValidPlugins.Add(ForeignPluginInfo);
+					BuildPlugins.Add(ForeignPluginInfo);
+				}
+			}
 		}
 
 		/** Sets up the binaries for the target. */
diff --git a/Engine/Source/Programs/UnrealBuildTool/System/ModuleDescriptor.cs b/Engine/Source/Programs/UnrealBuildTool/System/ModuleDescriptor.cs
--- a/Engine/Source/Programs/UnrealBuildTool/System/ModuleDescriptor.cs
+++ b/Engine/Source/Programs/UnrealBuildTool/System/ModuleDescriptor.cs
@@ -96,6 +96,65 @@ namespace UnrealBuildTool
 			return Module;
 		}
 
+		/// <summary>
+		/// Write this module to a JsonWriter
+		/// </summary>
+		/// <param name="Writer">Writer to output to</param>
+		void Write(JsonWriter Writer)
+		{
+			Writer.WriteObjectStart();
+			Writer.WriteValue("Name", Name);
+			Writer.WriteValue("Type", Type.ToString());
+			Writer.WriteValue("LoadingPhase", LoadingPhase.ToString());
+			if (WhitelistPlatforms != null && WhitelistPlatforms.Length > 0)
+			{
+				Writer.WriteArrayStart("WhitelistPlatforms");
+				foreach(UnrealTargetPlatform WhitelistPlatform in WhitelistPlatforms)
+				{
+					Writer.WriteValue(WhitelistPlatform.ToString());
+				}
+				Writer.WriteArrayEnd();
+			}
+			if (BlacklistPlatforms != null && BlacklistPlatforms.Length > 0)
+			{
+				Writer.WriteArrayStart("BlacklistPlatforms");
+				foreach(UnrealTargetPlatform BlacklistPlatform in BlacklistPlatforms)
+				{
+					Writer.WriteValue(BlacklistPlatform.ToString());
+				}
+				Writer.WriteArrayEnd();
+			}
+			if (AdditionalDependencies != null && AdditionalDependencies.Length > 0)
+			{
+				Writer.WriteArrayStart("AdditionalDependencies");
+				foreach(string AdditionalDependency in AdditionalDependencies)
+				{
+					Writer.WriteValue(AdditionalDependency);
+				}
+				Writer.WriteArrayEnd();
+			}
+			Writer.WriteObjectEnd();
+		}
+
+		/// <summary>
+		/// Write an array of module descriptors
+		/// </summary>
+		/// <param name="Writer">The Json writer to output to</param>
+		/// <param name="Name">Name of the array</param>
+		/// <param name="Modules">Array of modules</param>
+		public static void WriteArray(JsonWriter Writer, string Name, ModuleDescriptor[] Modules)
+		{
+			if(Modules.Length > 0)
+			{
+				Writer.WriteArrayStart(Name);
+				foreach(ModuleDescriptor Module in Modules)
+				{
+					Module.Write(Writer);
+				}
+				Writer.WriteArrayEnd();
+			}
+		}
+
 		/// <summary>
 		/// Determines whether the given plugin module is part of the current build.
 		/// </summary>
@@ -125,7 +184,7 @@ namespace UnrealBuildTool
 					return UEBuildConfiguration.bBuildDeveloperTools;
 				case ModuleHostType.Editor:
 				case ModuleHostType.EditorNoCommandlet:
-					return UEBuildConfiguration.bBuildEditor;
+					return TargetType == TargetRules.TargetType.Editor || UEBuildConfiguration.bBuildEditor;
 				case ModuleHostType.Program:
 					return TargetType == TargetRules.TargetType.Program;
 			}
diff --git a/Engine/Source/Programs/UnrealBuildTool/System/PluginDescriptor.cs b/Engine/Source/Programs/UnrealBuildTool/System/PluginDescriptor.cs
--- a/Engine/Source/Programs/UnrealBuildTool/System/PluginDescriptor.cs
+++ b/Engine/Source/Programs/UnrealBuildTool/System/PluginDescriptor.cs
@@ -1,6 +1,7 @@
 ﻿using System;
 using System.Collections.Generic;
 using System.Diagnostics;
+using System.IO;
 using System.Linq;
 using System.Text;
 
@@ -142,6 +143,37 @@ namespace UnrealBuildTool
 				throw new JsonParseException("{0} (in {1})", ParseException.Message, FileName);
 			}
 		}
+
+		/// <summary>
+		/// Saves the descriptor to disk
+		/// </summary>
+		/// <param name="FileName">The filename to write to</param>
+		public void Save(string FileName)
+		{
+			using(JsonWriter Writer = new JsonWriter(FileName))
+			{
+				Writer.WriteObjectStart();
+
+				Writer.WriteValue("FileVersion", (int)ProjectDescriptorVersion.Latest);
+				Writer.WriteValue("Version", Version);
+				Writer.WriteValue("VersionName", VersionName);
+				Writer.WriteValue("FriendlyName", FriendlyName);
+				Writer.WriteValue("Description", Description);
+				Writer.WriteValue("Category", Category);
+				Writer.WriteValue("CreatedBy", CreatedBy);
+				Writer.WriteValue("CreatedByURL", CreatedByURL);
+				Writer.WriteValue("DocsURL", DocsURL);
+
+				ModuleDescriptor.WriteArray(Writer, "Modules", Modules);
+
+				Writer.WriteValue("EnabledByDefault", bEnabledByDefault);
+				Writer.WriteValue("CanContainContent", bCanContainContent);
+				Writer.WriteValue("IsBetaVersion", bIsBetaVersion);
+				Writer.WriteValue("Installed", bInstalled);
+
+				Writer.WriteObjectEnd();
+			}
+		}
 	}
 
 	[DebuggerDisplay("Name={Name}")]
diff --git a/Engine/Source/Programs/UnrealBuildTool/System/RulesCompiler.cs b/Engine/Source/Programs/UnrealBuildTool/System/RulesCompiler.cs
--- a/Engine/Source/Programs/UnrealBuildTool/System/RulesCompiler.cs
+++ b/Engine/Source/Programs/UnrealBuildTool/System/RulesCompiler.cs
@@ -1134,6 +1134,9 @@ namespace UnrealBuildTool
 			private set;
 		}
 
+		/// External folders to also search for rules files
+		static List<string> ForeignPlugins;
+
 		/// <summary>
 		/// Sets which game folders to look at when harvesting for rules source files.  This must be called before
 		/// other functions in the RulesCompiler.  The idea here is that we can actually cache rules files for multiple
@@ -1141,12 +1144,15 @@ namespace UnrealBuildTool
 		/// when generating project files.
 		/// </summary>
 		/// <param name="GameFolders">List of all game folders that rules files will ever be requested for</param>
-		public static void SetAssemblyNameAndGameFolders( string AssemblyName, List<string> GameFolders )
+		/// <param name="InExtraPluginFolders">List of additional folders </param>
+		public static void SetAssemblyNameAndGameFolders( string AssemblyName, List<string> GameFolders, List<string> InForeignPlugins = null)
 		{
 			RulesCompiler.AssemblyName = AssemblyName + "ModuleRules";
 
 			AllGameFolders = new List<string>();
 			AllGameFolders.AddRange( GameFolders );
+
+			ForeignPlugins = (InForeignPlugins == null)? null : new List<string>(InForeignPlugins);
 		}
 
 
@@ -1176,6 +1182,16 @@ namespace UnrealBuildTool
 				}
 			}
 
+			// Add all the extra plugin folders
+			if( ForeignPlugins != null )
+			{
+				foreach(string ForeignPlugin in ForeignPlugins)
+				{
+					string PluginDirectory = Path.GetDirectoryName(Path.GetFullPath(ForeignPlugin));
+					Folders.Add(Path.Combine(PluginDirectory, "Source"));
+				}
+			}
+
 			// Add in the game folders to search
 			if( AllGameFolders != null )
 			{
diff --git a/Engine/Source/Programs/UnrealBuildTool/Utilities/JsonObject.cs b/Engine/Source/Programs/UnrealBuildTool/Utilities/JsonObject.cs
--- a/Engine/Source/Programs/UnrealBuildTool/Utilities/JsonObject.cs
+++ b/Engine/Source/Programs/UnrealBuildTool/Utilities/JsonObject.cs
@@ -185,4 +185,110 @@ namespace UnrealBuildTool
 			}
 		}
 	}
+
+	public class JsonWriter : IDisposable
+	{
+		StreamWriter Writer;
+		bool bRequiresComma;
+		string Indent;
+
+		public JsonWriter(string FileName)
+		{
+			Writer = new StreamWriter(FileName);
+			Indent = "";
+		}
+
+		public void Dispose()
+		{
+			Writer.Dispose();
+		}
+
+		public void WriteObjectStart()
+		{
+			WriteCommaNewline();
+
+			Writer.Write(Indent);
+			Writer.Write("{");
+
+			Indent += "\t";
+			bRequiresComma = false;
+		}
+
+		public void WriteObjectEnd()
+		{
+			Indent = Indent.Substring(0, Indent.Length - 1);
+
+			Writer.WriteLine();
+			Writer.Write(Indent);
+			Writer.Write("}");
+
+			bRequiresComma = true;
+		}
+
+		public void WriteArrayStart(string ArrayName)
+		{
+			WriteCommaNewline();
+
+			Writer.WriteLine("{0}\"{1}\" :", Indent, ArrayName);
+			Writer.Write("{0}[", Indent);
+
+			Indent += "\t";
+			bRequiresComma = false;
+		}
+
+		public void WriteArrayEnd()
+		{
+			Indent = Indent.Substring(0, Indent.Length - 1);
+
+			Writer.WriteLine();
+			Writer.Write("{0}]", Indent);
+
+			bRequiresComma = true;
+		}
+
+		public void WriteValue(string Value)
+		{
+			WriteCommaNewline();
+
+			Writer.Write("{0}\"{1}\"", Indent, Value);
+
+			bRequiresComma = true;
+		}
+
+		public void WriteValue(string Name, string Value)
+		{
+			WriteValueInternal(Name, '"' + Value + '"');
+		}
+
+		public void WriteValue(string Name, int Value)
+		{
+			WriteValueInternal(Name, Value.ToString());
+		}
+
+		public void WriteValue(string Name, bool Value)
+		{
+			WriteValueInternal(Name, Value? "true" : "false");
+		}
+
+		void WriteCommaNewline()
+		{
+			if(bRequiresComma)
+			{
+				Writer.WriteLine(",");
+			}
+			else if(Indent.Length > 0)
+			{
+				Writer.WriteLine();
+			}
+		}
+
+		void WriteValueInternal(string Name, string Value)
+		{
+			WriteCommaNewline();
+
+			Writer.Write("{0}\"{1}\" : {2}", Indent, Name, Value);
+
+			bRequiresComma = true;
+		}
+	}
 }
