Setup word count and conflict report generation for engine/editor targets in the legacy localization configuration scripts.

diff --git a/Engine/Config/Localization/Editor.ini b/Engine/Config/Localization/Editor.ini
--- a/Engine/Config/Localization/Editor.ini
+++ b/Engine/Config/Localization/Editor.ini
@@ -70,4 +70,13 @@ ResourceName=Editor.locres
 ;Export to PO files
 [GatherTextStep7]
 CommandletClass=InternationalizationExport
-bExportLoc=true
\ No newline at end of file
+bExportLoc=true
+
+;Write Text Localization Report
+[GatherTextStep8]
+CommandletClass=GenerateTextLocalizationReport
+DestinationPath=\\epicgames.net\root\UE3\Localization\WordCounts
+bWordCountReport=true
+WordCountReportName=Editor.csv
+bConflictReport=true
+ConflictReportName=EditorConflicts.txt
\ No newline at end of file
diff --git a/Engine/Config/Localization/EditorTutorials.ini b/Engine/Config/Localization/EditorTutorials.ini
--- a/Engine/Config/Localization/EditorTutorials.ini
+++ b/Engine/Config/Localization/EditorTutorials.ini
@@ -44,4 +44,13 @@ ResourceName=EditorTutorials.locres
 ;Export to PO files
 [GatherTextStep5]
 CommandletClass=InternationalizationExport
-bExportLoc=true
\ No newline at end of file
+bExportLoc=true
+
+;Write Text Localization Report
+[GatherTextStep6]
+CommandletClass=GenerateTextLocalizationReport
+DestinationPath=\\epicgames.net\root\UE3\Localization\WordCounts
+bWordCountReport=true
+WordCountReportName=EditorTutorials.csv
+bConflictReport=true
+ConflictReportName=EditorTutorialsConflicts.txt
\ No newline at end of file
diff --git a/Engine/Config/Localization/Engine.ini b/Engine/Config/Localization/Engine.ini
--- a/Engine/Config/Localization/Engine.ini
+++ b/Engine/Config/Localization/Engine.ini
@@ -70,4 +70,13 @@ ResourceName=Engine.locres
 ;Export  PO files
 [GatherTextStep6]
 CommandletClass=InternationalizationExport
-bExportLoc=true
\ No newline at end of file
+bExportLoc=true
+
+;Write Text Localization Report
+[GatherTextStep7]
+CommandletClass=GenerateTextLocalizationReport
+DestinationPath=\\epicgames.net\root\UE3\Localization\WordCounts
+bWordCountReport=true
+WordCountReportName=Engine.csv
+bConflictReport=true
+ConflictReportName=EngineConflicts.txt
\ No newline at end of file
diff --git a/Engine/Config/Localization/Keywords.ini b/Engine/Config/Localization/Keywords.ini
--- a/Engine/Config/Localization/Keywords.ini
+++ b/Engine/Config/Localization/Keywords.ini
@@ -46,4 +46,13 @@ ResourceName=Keywords.locres
 ;Export to PO files
 [GatherTextStep5]
 CommandletClass=InternationalizationExport
-bExportLoc=true
\ No newline at end of file
+bExportLoc=true
+
+;Write Text Localization Report
+[GatherTextStep6]
+CommandletClass=GenerateTextLocalizationReport
+DestinationPath=\\epicgames.net\root\UE3\Localization\WordCounts
+bWordCountReport=true
+WordCountReportName=Keywords.csv
+bConflictReport=true
+ConflictReportName=KeywordsConflicts.txt
\ No newline at end of file
diff --git a/Engine/Config/Localization/PropertyNames.ini b/Engine/Config/Localization/PropertyNames.ini
--- a/Engine/Config/Localization/PropertyNames.ini
+++ b/Engine/Config/Localization/PropertyNames.ini
@@ -46,4 +46,13 @@ ResourceName=PropertyNames.locres
 ;Export to PO files
 [GatherTextStep5]
 CommandletClass=InternationalizationExport
-bExportLoc=true
\ No newline at end of file
+bExportLoc=true
+
+;Write Text Localization Report
+[GatherTextStep6]
+CommandletClass=GenerateTextLocalizationReport
+DestinationPath=\\epicgames.net\root\UE3\Localization\WordCounts
+bWordCountReport=true
+WordCountReportName=PropertyNames.csv
+bConflictReport=true
+ConflictReportName=PropertyNamesConflicts.txt
\ No newline at end of file
diff --git a/Engine/Config/Localization/ToolTips.ini b/Engine/Config/Localization/ToolTips.ini
--- a/Engine/Config/Localization/ToolTips.ini
+++ b/Engine/Config/Localization/ToolTips.ini
@@ -49,4 +49,13 @@ ResourceName=ToolTips.locres
 ;Export to PO files
 [GatherTextStep5]
 CommandletClass=InternationalizationExport
-bExportLoc=true
\ No newline at end of file
+bExportLoc=true
+
+;Write Text Localization Report
+[GatherTextStep6]
+CommandletClass=GenerateTextLocalizationReport
+DestinationPath=\\epicgames.net\root\UE3\Localization\WordCounts
+bWordCountReport=true
+WordCountReportName=ToolTips.csv
+bConflictReport=true
+ConflictReportName=ToolTipsConflicts.txt
\ No newline at end of file
diff --git a/Engine/Source/Programs/AutomationTool/Scripts/Localisation.Automation.cs b/Engine/Source/Programs/AutomationTool/Scripts/Localisation.Automation.cs
--- a/Engine/Source/Programs/AutomationTool/Scripts/Localisation.Automation.cs
+++ b/Engine/Source/Programs/AutomationTool/Scripts/Localisation.Automation.cs
@@ -87,7 +87,6 @@ class Localise : BuildCommand
 				String.Format("-config={0}", @"./Config/Localization/EditorTutorials.ini") + (string.IsNullOrEmpty(CommandletSCCArguments) ? "" : " " + CommandletSCCArguments),
 				String.Format("-config={0}", @"./Config/Localization/PropertyNames.ini") + (string.IsNullOrEmpty(CommandletSCCArguments) ? "" : " " + CommandletSCCArguments),
 				String.Format("-config={0}", @"./Config/Localization/ToolTips.ini") + (string.IsNullOrEmpty(CommandletSCCArguments) ? "" : " " + CommandletSCCArguments),
-				String.Format("-config={0}", @"./Config/Localization/WordCount.ini"),
 			};
 
 		// Execute commandlet for each set of arguments.
