Change the BuildPlugin command to build for all rocket target platforms by default, since those are the only users that benefit from source. Use -NoHostPlatform -NoTargetPlatforms to package as source-only.

diff --git a/Engine/Source/Programs/AutomationTool/Scripts/BuildPluginCommand.Automation.cs b/Engine/Source/Programs/AutomationTool/Scripts/BuildPluginCommand.Automation.cs
--- a/Engine/Source/Programs/AutomationTool/Scripts/BuildPluginCommand.Automation.cs
+++ b/Engine/Source/Programs/AutomationTool/Scripts/BuildPluginCommand.Automation.cs
@@ -9,6 +9,9 @@ using AutomationTool;
 using UnrealBuildTool;
 
 [Help("Builds a plugin, and packages it for distribution")]
+[Help("Plugin", "Specify the path to the plugin that should be packaged")]
+[Help("NoEditor", "Prevent compiling for the editor platform")]
+[Help("TargetPlatforms", "Specify a list of target platforms to build, separated by '+' characters (eg. -TargetPlatforms=Win32+Win64). Default is all the Rocket target platforms.")]
 class BuildPlugin : BuildCommand
 {
 	public override void ExecuteBuild()
@@ -19,11 +22,7 @@ class BuildPlugin : BuildCommand
 		{
 			throw new AutomationException("Plugin file name was not specified via the -plugin argument");
 		}
-
-		// Get the platforms to build for
-		List<UnrealTargetPlatform> EditorPlatforms = ParseParamPlatforms("EditorPlatforms");
-		List<UnrealTargetPlatform> GamePlatforms = ParseParamPlatforms("GamePlatforms");
-
+			
 		// Read the plugin
 		PluginDescriptor Plugin = PluginDescriptor.FromFile(PluginFileName);
 
@@ -34,10 +33,27 @@ class BuildPlugin : BuildCommand
 			CommandUtils.DeleteDirectory(IntermediateBuildDirectory);
 		}
 
-		// Build everything
+		// Build the host platforms
 		List<string> ReceiptFileNames = new List<string>();
 		UE4Build.BuildAgenda Agenda = new UE4Build.BuildAgenda();
-		AddPluginToAgenda(Agenda, PluginFileName, Plugin, EditorPlatforms, GamePlatforms, ReceiptFileNames);
+		UnrealTargetPlatform HostPlatform = BuildHostPlatform.Current.Platform;
+		if(!ParseParam("NoHostPlatform"))
+		{
+			AddPluginToAgenda(Agenda, PluginFileName, Plugin, "UE4Editor", TargetRules.TargetType.Editor, HostPlatform, UnrealTargetConfiguration.Development, ReceiptFileNames);
+		}
+
+		// Add the game targets
+		List<UnrealTargetPlatform> TargetPlatforms = Rocket.RocketBuild.GetTargetPlatforms(this, HostPlatform);
+		foreach(UnrealTargetPlatform TargetPlatform in TargetPlatforms)
+		{
+			if(Rocket.RocketBuild.IsCodeTargetPlatform(HostPlatform, TargetPlatform))
+			{
+				AddPluginToAgenda(Agenda, PluginFileName, Plugin, "UE4Game", TargetRules.TargetType.Game, TargetPlatform, UnrealTargetConfiguration.Development, ReceiptFileNames);
+				AddPluginToAgenda(Agenda, PluginFileName, Plugin, "UE4Game", TargetRules.TargetType.Game, TargetPlatform, UnrealTargetConfiguration.Shipping, ReceiptFileNames);
+			}
+		}
+
+		// Build it
 		UE4Build Build = new UE4Build(this);
 		Build.Build(Agenda, InDeleteBuildProducts: true, InUpdateVersionFiles: false);
 
@@ -72,22 +88,6 @@ class BuildPlugin : BuildCommand
 		return Platforms;
 	}
 
-	static void AddPluginToAgenda(UE4Build.BuildAgenda Agenda, string PluginFileName, PluginDescriptor Plugin, List<UnrealTargetPlatform> EditorPlatforms, List<UnrealTargetPlatform> GamePlatforms, List<string> ReceiptFileNames)
-	{
-		// Add the editor targets
-		foreach(UnrealTargetPlatform EditorPlatform in EditorPlatforms)
-		{
-			AddPluginToAgenda(Agenda, PluginFileName, Plugin, "UE4Editor", TargetRules.TargetType.Editor, EditorPlatform, UnrealTargetConfiguration.Development, ReceiptFileNames);
-		}
-
-		// Add the game targets
-		foreach(UnrealTargetPlatform GamePlatform in GamePlatforms)
-		{
-			AddPluginToAgenda(Agenda, PluginFileName, Plugin, "UE4Game", TargetRules.TargetType.Game, GamePlatform, UnrealTargetConfiguration.Development, ReceiptFileNames);
-			AddPluginToAgenda(Agenda, PluginFileName, Plugin, "UE4Game", TargetRules.TargetType.Game, GamePlatform, UnrealTargetConfiguration.Shipping, ReceiptFileNames);
-		}
-	}
-
 	static void AddPluginToAgenda(UE4Build.BuildAgenda Agenda, string PluginFileName, PluginDescriptor Plugin, string TargetName, TargetRules.TargetType TargetType, UnrealTargetPlatform Platform, UnrealTargetConfiguration Configuration, List<string> ReceiptFileNames)
 	{
 		// Find a list of modules that need to be built for this plugin
diff --git a/Engine/Source/Programs/AutomationTool/Scripts/RocketBuild.Automation.cs b/Engine/Source/Programs/AutomationTool/Scripts/RocketBuild.Automation.cs
--- a/Engine/Source/Programs/AutomationTool/Scripts/RocketBuild.Automation.cs
+++ b/Engine/Source/Programs/AutomationTool/Scripts/RocketBuild.Automation.cs
@@ -77,6 +77,9 @@ namespace Rocket
 				// Find all the target platforms for this host platform.
 				List<UnrealTargetPlatform> TargetPlatforms = GetTargetPlatforms(bp, HostPlatform);
 
+				// Remove any platforms that aren't available on this machine
+				TargetPlatforms.RemoveAll(x => !bp.ActivePlatforms.Contains(x));
+
 				// Get the temp directory for stripped files for this host
 				string StrippedDir = Path.GetFullPath(CommandUtils.CombinePaths(CommandUtils.CmdEnv.LocalRoot, "Engine", "Saved", "Rocket", HostPlatform.ToString()));
 
@@ -156,10 +159,10 @@ namespace Rocket
 			return FEngineVersionSupport.FromVersionFile(CommandUtils.CombinePaths(CommandUtils.CmdEnv.LocalRoot, @"Engine\Source\Runtime\Launch\Resources\Version.h")).ToString();
 		}
 
-		public static List<UnrealTargetPlatform> GetTargetPlatforms(GUBP bp, UnrealTargetPlatform HostPlatform)
+		public static List<UnrealTargetPlatform> GetTargetPlatforms(BuildCommand Command, UnrealTargetPlatform HostPlatform)
 		{
 			List<UnrealTargetPlatform> TargetPlatforms = new List<UnrealTargetPlatform>();
-			if(!bp.ParseParam("NoTargetPlatforms"))
+			if(!Command.ParseParam("NoTargetPlatforms"))
 			{
 				// Always support the host platform
 				TargetPlatforms.Add(HostPlatform);
@@ -186,11 +189,8 @@ namespace Rocket
 					TargetPlatforms.Add(UnrealTargetPlatform.HTML5);
 				}
 
-				// Remove any platforms that aren't available on this machine
-				TargetPlatforms.RemoveAll(x => !bp.ActivePlatforms.Contains(x));
-
 				// Remove any platforms that aren't enabled on the command line
-				string TargetPlatformFilter = bp.ParseParamValue("TargetPlatforms", null);
+				string TargetPlatformFilter = Command.ParseParamValue("TargetPlatforms", null);
 				if(TargetPlatformFilter != null)
 				{
 					List<UnrealTargetPlatform> NewTargetPlatforms = new List<UnrealTargetPlatform>();
