Fix building Rocket samples, and create/publish zip files for each build instead of copying loose files (much faster with network).

diff --git a/Engine/Build/InstalledEngineFilters.ini b/Engine/Build/InstalledEngineFilters.ini
--- a/Engine/Build/InstalledEngineFilters.ini
+++ b/Engine/Build/InstalledEngineFilters.ini
@@ -219,24 +219,3 @@
 
 [StripSymbols.Linux]
 /Engine/Binaries/Linux/*.
-
-[CopyEditorForSamples]
-; Include only binaries for the host platform
--/Engine/Binaries/...
-/Engine/Binaries/DotNET/...
-/Engine/Binaries/ThirdParty/...
-{Win64} /Engine/Binaries/Win64/...
-{Mac} /Engine/Binaries/Mac/...
-
-; Exclude other files
--/Engine/Documentation/...
--/Engine/Extras/...
--/Engine/Source/.../Private/...
-
-; Exclude other content
--/FeaturePacks/...
--/Samples/...
--/Templates/...
-
-; No PDBs
--*.pdb
diff --git a/Engine/Source/Programs/AutomationTool/CommandUtils.cs b/Engine/Source/Programs/AutomationTool/CommandUtils.cs
--- a/Engine/Source/Programs/AutomationTool/CommandUtils.cs
+++ b/Engine/Source/Programs/AutomationTool/CommandUtils.cs
@@ -2107,6 +2107,7 @@ namespace AutomationTool
 		public static void ZipFiles(string ZipFileName, string BaseDirectory, FileFilter Filter)
 		{
 			Ionic.Zip.ZipFile Zip = new Ionic.Zip.ZipFile();
+			Zip.UseZip64WhenSaving = Ionic.Zip.Zip64Option.Always;
 			foreach(string FilteredFile in Filter.ApplyToDirectory(BaseDirectory, true))
 			{
 				Zip.AddFile(Path.Combine(BaseDirectory, FilteredFile), Path.GetDirectoryName(FilteredFile));
@@ -2123,21 +2124,22 @@ namespace AutomationTool
 		/// <returns>List of files written</returns>
 		public static IEnumerable<string> UnzipFiles(string ZipFileName, string BaseDirectory)
 		{
-			Ionic.Zip.ZipFile Zip = new Ionic.Zip.ZipFile(ZipFileName);
-
-			// For some reason, calling ExtractAll() under Mono throws an exception from a call to Directory.CreateDirectory(). Manually extract all the files instead.
-			List<string> OutputFileNames = new List<string>();
-			foreach(Ionic.Zip.ZipEntry Entry in Zip.Entries)
+			using(Ionic.Zip.ZipFile Zip = new Ionic.Zip.ZipFile(ZipFileName))
 			{
-				string OutputFileName = Path.Combine(BaseDirectory, Entry.FileName);
-				Directory.CreateDirectory(Path.GetDirectoryName(OutputFileName));
-				using(FileStream OutputStream = new FileStream(OutputFileName, FileMode.Create, FileAccess.Write))
+				// For some reason, calling ExtractAll() under Mono throws an exception from a call to Directory.CreateDirectory(). Manually extract all the files instead.
+				List<string> OutputFileNames = new List<string>();
+				foreach(Ionic.Zip.ZipEntry Entry in Zip.Entries)
 				{
-					Entry.Extract(OutputStream);
+					string OutputFileName = Path.Combine(BaseDirectory, Entry.FileName);
+					Directory.CreateDirectory(Path.GetDirectoryName(OutputFileName));
+					using(FileStream OutputStream = new FileStream(OutputFileName, FileMode.Create, FileAccess.Write))
+					{
+						Entry.Extract(OutputStream);
+					}
+					OutputFileNames.Add(OutputFileName);
 				}
-				OutputFileNames.Add(OutputFileName);
+				return OutputFileNames;
 			}
-			return OutputFileNames;
 		}
 	}
 
diff --git a/Engine/Source/Programs/AutomationTool/Scripts/RocketBuild.Automation.cs b/Engine/Source/Programs/AutomationTool/Scripts/RocketBuild.Automation.cs
--- a/Engine/Source/Programs/AutomationTool/Scripts/RocketBuild.Automation.cs
+++ b/Engine/Source/Programs/AutomationTool/Scripts/RocketBuild.Automation.cs
@@ -125,10 +125,20 @@ namespace Rocket
 					}
 				}
 
+				// Get the output directory for the build zips
+				string PublishedEngineDir;
+				if(ShouldDoSeriousThingsLikeP4CheckinAndPostToMCP())
+				{
+					PublishedEngineDir = CommandUtils.CombinePaths(CommandUtils.RootSharedTempStorageDirectory(), "Rocket", "Automated", GetBuildLabel(), CommandUtils.GetGenericPlatformName(HostPlatform));
+				}
+				else
+				{
+					PublishedEngineDir = CommandUtils.CombinePaths(CommandUtils.CmdEnv.LocalRoot, "LocalBuilds", "RocketPublish", CommandUtils.GetGenericPlatformName(HostPlatform));
+				}
+
 				// Publish the install to the network
-				string RemoteOutputDir = CommandUtils.CombinePaths(CommandUtils.RootSharedTempStorageDirectory(), "Rocket", "Automated", GetBuildLabel(), CommandUtils.GetGenericPlatformName(HostPlatform));
-				bp.AddNode(new PublishRocketNode(HostPlatform, LocalOutputDir, RemoteOutputDir));
-				bp.AddNode(new PublishRocketSymbolsNode(bp, HostPlatform, TargetPlatforms, RemoteOutputDir + "Symbols"));
+				bp.AddNode(new PublishRocketNode(HostPlatform, LocalOutputDir, PublishedEngineDir));
+				bp.AddNode(new PublishRocketSymbolsNode(bp, HostPlatform, TargetPlatforms, PublishedEngineDir + "Symbols"));
 
 				// Add a dependency on this being published as part of the shared promotable being labeled
 				GUBP.SharedLabelPromotableSuccessNode LabelPromotableNode = (GUBP.SharedLabelPromotableSuccessNode)bp.FindNode(GUBP.SharedLabelPromotableSuccessNode.StaticGetFullName());
@@ -814,11 +824,39 @@ namespace Rocket
 
 		public override void DoBuild(GUBP bp)
 		{
-			if(RocketBuild.ShouldDoSeriousThingsLikeP4CheckinAndPostToMCP())
-			{
-				CommandUtils.ThreadedCopyFiles(LocalDir, PublishDir);
-			}
-
+			// Create a zip file containing the install
+			string FullZipFileName = Path.Combine(CommandUtils.CmdEnv.LocalRoot, "FullInstall" + StaticGetHostPlatformSuffix(HostPlatform) + ".zip");
+			CommandUtils.Log("Creating {0}...", FullZipFileName);
+			CommandUtils.ZipFiles(FullZipFileName, LocalDir, new FileFilter(FileFilterType.Include));
+
+			// Create a filter for the files we need just to run the editor
+			FileFilter EditorFilter = new FileFilter(FileFilterType.Include);
+			EditorFilter.Exclude("/Engine/Binaries/...");
+			EditorFilter.Include("/Engine/Binaries/DotNET/...");
+			EditorFilter.Include("/Engine/Binaries/ThirdParty/...");
+			EditorFilter.Include("/Engine/Binaries/" + HostPlatform.ToString() + "/...");
+			EditorFilter.Exclude("/Engine/Binaries/.../*.lib");
+			EditorFilter.Exclude("/Engine/Binaries/.../*.a");
+			EditorFilter.Exclude("/Engine/Extras/...");
+			EditorFilter.Exclude("/Engine/Source/.../Private/...");
+			EditorFilter.Exclude("/FeaturePacks/...");
+			EditorFilter.Exclude("/Samples/...");
+			EditorFilter.Exclude("/Templates/...");
+			EditorFilter.Exclude("*.pdb");
+
+			// Create a zip file containing the editor install
+			string EditorZipFileName = Path.Combine(CommandUtils.CmdEnv.LocalRoot, "EditorInstall" + StaticGetHostPlatformSuffix(HostPlatform) + ".zip");
+			CommandUtils.Log("Creating {0}...", EditorZipFileName);
+			CommandUtils.ZipFiles(EditorZipFileName, LocalDir, EditorFilter);
+
+			// Copy the files to their final location
+			CommandUtils.Log("Copying files to {0}", PublishDir);
+			CommandUtils.CopyFile(FullZipFileName, Path.Combine(PublishDir, Path.GetFileName(FullZipFileName)));
+			CommandUtils.CopyFile(EditorZipFileName, Path.Combine(PublishDir, Path.GetFileName(EditorZipFileName)));
+			CommandUtils.DeleteFile(FullZipFileName);
+			CommandUtils.DeleteFile(EditorZipFileName);
+
+			// Save a record of success
 			BuildProducts = new List<string>();
 			SaveRecordOfSuccessAndAddToBuildProducts();
 		}
