log all major gubp events

diff --git a/Engine/Source/Programs/AutomationTool/Automation.cs b/Engine/Source/Programs/AutomationTool/Automation.cs
--- a/Engine/Source/Programs/AutomationTool/Automation.cs
+++ b/Engine/Source/Programs/AutomationTool/Automation.cs
@@ -319,7 +319,10 @@ AutomationTool.exe [-verbose] [-compileonly] [-p4] Command0 [-Arg0 -Arg1 -Arg2 
 			// Compile scripts.
 			Log.TraceInformation("Compiling scripts.");
 			ScriptCompiler Compiler = new ScriptCompiler();
+            var StartScriptCompile = DateTime.UtcNow.ToString();
 			Compiler.FindAndCompileAllScripts(AdditionalScriptsFolders: AdditionalScriptsFolders);
+            var FinishScriptCompile = DateTime.UtcNow.ToString();
+            CommandUtils.StepDurations.Add("ScriptCompile", String.Format("{0},{1}", StartScriptCompile, FinishScriptCompile));
 
 			if (GlobalCommandLine.CompileOnly)
 			{
diff --git a/Engine/Source/Programs/AutomationTool/BuildEnvironment.cs b/Engine/Source/Programs/AutomationTool/BuildEnvironment.cs
--- a/Engine/Source/Programs/AutomationTool/BuildEnvironment.cs
+++ b/Engine/Source/Programs/AutomationTool/BuildEnvironment.cs
@@ -18,6 +18,7 @@ namespace AutomationTool
 		// Command Environment
 		static public readonly string LocalRoot = "uebp_LOCAL_ROOT";
 		static public readonly string LogFolder = "uebp_LogFolder";
+        static public readonly string CSVFile = "uebp_CSVFile";
 		static public readonly string EngineSavedFolder = "uebp_EngineSavedFolder";
 		static public readonly string NETFrameworkDir = "FrameworkDir";
 		static public readonly string NETFrameworkVersion = "FrameworkVersion";
@@ -49,7 +50,8 @@ namespace AutomationTool
 
 		public string LocalRoot { get; protected set; }
 		public string EngineSavedFolder { get; protected set; }
-		public string LogFolder { get; protected set; }		
+		public string LogFolder { get; protected set; }
+        public string CSVFile { get; protected set; }
 		public string RobocopyExe { get; protected set; }
 		public string MountExe { get; protected set; }
 		public string CmdExe { get; protected set; }
@@ -104,6 +106,7 @@ namespace AutomationTool
 				SetUATSavedPath();
 			}
 			EngineSavedFolder = CommandUtils.GetEnvVar(EnvVarNames.EngineSavedFolder);
+            CSVFile = CommandUtils.GetEnvVar(EnvVarNames.CSVFile);            
 			LogFolder = CommandUtils.GetEnvVar(EnvVarNames.LogFolder);
 			RobocopyExe = CommandUtils.CombinePaths(Environment.SystemDirectory, "robocopy.exe");
 			MountExe = CommandUtils.CombinePaths(Environment.SystemDirectory, "mount.exe");
diff --git a/Engine/Source/Programs/AutomationTool/CommandUtils.cs b/Engine/Source/Programs/AutomationTool/CommandUtils.cs
--- a/Engine/Source/Programs/AutomationTool/CommandUtils.cs
+++ b/Engine/Source/Programs/AutomationTool/CommandUtils.cs
@@ -1878,6 +1878,8 @@ namespace AutomationTool
 		}
 		private static UnrealBuildTool.UnrealTargetPlatform[] UBTTargetPlatforms;
 
+        public static Dictionary<string, string> StepDurations = new Dictionary<string,string>();
+
 		#endregion
 
 		#region Properties
@@ -2143,8 +2145,6 @@ namespace AutomationTool
 		}
 	}
 
-
-
 	/// <summary>
 	/// Use with "using" syntax to push and pop directories in a convenient, exception-safe way
 	/// </summary>
diff --git a/Engine/Source/Programs/AutomationTool/ProjectUtils.cs b/Engine/Source/Programs/AutomationTool/ProjectUtils.cs
--- a/Engine/Source/Programs/AutomationTool/ProjectUtils.cs
+++ b/Engine/Source/Programs/AutomationTool/ProjectUtils.cs
@@ -707,7 +707,7 @@ namespace AutomationTool
             BaseEngineProject = new BranchUProject();
 
             var AllProjects = UnrealBuildTool.UProjectInfo.FilterGameProjects(false, null);
-
+            var StartTime = DateTime.UtcNow.ToString();
             foreach (var InfoEntry in AllProjects)
             {
                 var UProject = new BranchUProject(InfoEntry);
@@ -725,12 +725,15 @@ namespace AutomationTool
                     }
                 }
             }
+            var FinishTime = DateTime.UtcNow.ToString();
+            CommandUtils.StepDurations.Add("SortProjects", String.Format("{0},{1}", StartTime, FinishTime));
             if (String.IsNullOrEmpty(BaseEngineProject.FilePath))
             {
                 throw new AutomationException("All branches must have the blank project /Samples/Sandbox/BlankProject");
             }
 
             CommandUtils.Log("  Base Engine:");
+            var StartBranchDump = DateTime.UtcNow.ToString();
             BaseEngineProject.Dump(InHostPlatforms);
 
             CommandUtils.Log("  {0} Code projects:", CodeProjects.Count);
@@ -743,6 +746,8 @@ namespace AutomationTool
             {
                 Proj.Dump(InHostPlatforms);
             }
+            var BuildDumpFinish = DateTime.UtcNow.ToString();
+            CommandUtils.StepDurations.Add("Project Dump", String.Format("{0},{1}", StartBranchDump, BuildDumpFinish));
         }
 
         public BranchUProject FindGame(string GameName)
diff --git a/Engine/Source/Programs/AutomationTool/Scripts/GUBP.Automation.cs b/Engine/Source/Programs/AutomationTool/Scripts/GUBP.Automation.cs
--- a/Engine/Source/Programs/AutomationTool/Scripts/GUBP.Automation.cs
+++ b/Engine/Source/Programs/AutomationTool/Scripts/GUBP.Automation.cs
@@ -6444,6 +6444,7 @@ public class GUBP : BuildCommand
         {
             Log("******* {0} GUBP Nodes", GUBPNodes.Count);
             var SortedNodes = TopologicalSort(new HashSet<string>(GUBPNodes.Keys), LocalOnly: true, DoNotConsiderCompletion: true);
+            var DependencyStart = DateTime.UtcNow.ToString();
             foreach (var Node in SortedNodes)
             {
                 string Note = GetControllingTriggerDotName(Node);
@@ -6476,6 +6477,8 @@ public class GUBP : BuildCommand
                     Log("  {0}: {1} [Aggregate]", Node, Note);
                 }
             }
+            var DependencyFinish = DateTime.UtcNow.ToString();
+            CommandUtils.StepDurations.Add("GetDependencies", String.Format("{0},{1}", DependencyStart, DependencyFinish));
         }
 		Dictionary<string, int> FullNodeListSortKey = GetDisplayOrder(FullNodeList.Keys.ToList(), FullNodeDirectDependencies, GUBPNodes);
 
@@ -6955,6 +6958,7 @@ public class GUBP : BuildCommand
             }
 
             var FilteredOrdereredToDo = new List<string>();
+            var StartFilterTime = DateTime.Now.ToString();
             // remove nodes that have unfinished triggers
             foreach (var NodeToDo in OrdereredToDo)
             {
@@ -6979,9 +6983,14 @@ public class GUBP : BuildCommand
                     FilteredOrdereredToDo.Add(NodeToDo);
                 }
             }
+            
             OrdereredToDo = FilteredOrdereredToDo;
+            var FinishFilterTime = DateTime.Now.ToString();
+            StepDurations.Add("FilterNodes", String.Format("{0},{1}", StartFilterTime, FinishFilterTime));
             Log("*********** EC Nodes, in order.");
             PrintNodes(this, OrdereredToDo, LocalOnly, UnfinishedTriggers);
+            var FinishNodePrint = DateTime.Now.ToString();
+            StepDurations.Add("SetupCommanderPrint", String.Format("{0},{1}", FinishFilterTime, FinishNodePrint));
 
             // here we are just making sure everything before the explicit trigger is completed.
             if (!String.IsNullOrEmpty(ExplicitTrigger))
@@ -7027,6 +7036,7 @@ public class GUBP : BuildCommand
             }
 			
             string ParentPath = ParseParamValue("ParentPath");
+            var StartPerlOutput = DateTime.Now.ToString();
             string BaseArgs = String.Format("$batch->createJobStep({{parentPath => '{0}'", ParentPath);
 
             bool bHasNoop = false;
@@ -7347,6 +7357,8 @@ public class GUBP : BuildCommand
                 }
             }
             WriteECPerl(StepList);
+            var FinishPerlOutput = DateTime.Now.ToString();
+            StepDurations.Add("PerlOutput", String.Format("{0},{1}", StartPerlOutput, FinishPerlOutput));
             RunECTool(String.Format("setProperty \"/myWorkflow/HasTests\" \"{0}\"", bHasTests));
             {
                 ECProps.Add("GUBP_LoadedProps=1");
@@ -7368,6 +7380,7 @@ public class GUBP : BuildCommand
                 }
             }
             Log("Commander setup only, done.");
+            PrintCSVFile();
             PrintRunTime();
             return;
 
@@ -7504,7 +7517,10 @@ public class GUBP : BuildCommand
                     {                        
 						Log("***** Building GUBP Node {0} for {1}", NodeToDo, NodeStoreName);
 						var StartTime = DateTime.UtcNow;
+                        var StartString = StartTime.ToString();
                         GUBPNodes[NodeToDo].DoBuild(this);
+                        var FinishBuild = DateTime.UtcNow.ToString();
+                        StepDurations.Add(String.Format("DoBuild{0}", NodeToDo), String.Format("{0},{1}", StartString, FinishBuild));
 						BuildDuration = (DateTime.UtcNow - StartTime).TotalMilliseconds / 1000;
 						
                     }
@@ -7512,8 +7528,12 @@ public class GUBP : BuildCommand
                     {
 						var StoreDuration = 0.0;
 						var StartTime = DateTime.UtcNow;
+                        var StoreTime = StartTime.ToString();
                         StoreToTempStorage(CmdEnv, NodeStoreName, GUBPNodes[NodeToDo].BuildProducts, !bSaveSharedTempStorage, GameNameIfAny, StorageRootIfAny);
 						StoreDuration = (DateTime.UtcNow - StartTime).TotalMilliseconds / 1000;
+                        var FinishStore = DateTime.UtcNow.ToString();
+                        StepDurations.Add(String.Format("StoreBuildProducts{0}", NodeToDo), String.Format("{0},{1}", StoreTime, FinishStore));
+                        //StepDurations.Add("StoreBuildProducts", StoreTime + "," + FinishStore);
 						Log("Took {0} seconds to store build products", StoreDuration);
 						if (IsBuildMachine)
 						{
@@ -7526,7 +7546,10 @@ public class GUBP : BuildCommand
                                 try
                                 {
                                     bool WasLocal;
+                                    var StartRetrieve = DateTime.UtcNow.ToString();
 									RetrieveFromTempStorage(CmdEnv, NodeStoreName, out WasLocal, GameNameIfAny, StorageRootIfAny);
+                                    var FinishRetrieve = DateTime.UtcNow.ToString();
+                                    StepDurations.Add(String.Format("RetrieveBuildProducts{0}", NodeToDo), String.Format("{0},{1}", StartRetrieve, FinishRetrieve));
 									if (!WasLocal)
 									{
 										throw new AutomationException("Retrieve was not local?");
@@ -7546,12 +7569,21 @@ public class GUBP : BuildCommand
                 {
                     if (SaveSuccessRecords)
                     {
-                        UpdateNodeHistory(NodeToDo, CLString);						
+                        var StartNodeHistory = DateTime.Now.ToString();
+                        UpdateNodeHistory(NodeToDo, CLString);
+                        var FinishNodeHistory = DateTime.Now.ToString();
+                        StepDurations.Add(String.Format("UpdateNodeHistory{0}", NodeToDo), String.Format("{0},{1}", StartNodeHistory, FinishNodeHistory));                        
                         SaveStatus(NodeToDo, FailedTempStorageSuffix, NodeStoreName, bSaveSharedTempStorage, GameNameIfAny, ParseParamValue("MyJobStepId"));
+                        var FinishSaveStatus = DateTime.Now.ToString();
+                        StepDurations.Add(String.Format("SaveNodeStatus{0}", NodeToDo), String.Format("{0},{1}", FinishNodeHistory, FinishSaveStatus));
                         UpdateECProps(NodeToDo, CLString);
+                        var FinishECPropUpdate = DateTime.Now.ToString();
+                        StepDurations.Add(String.Format("UpdateECProps{0}", NodeToDo), String.Format("{0},{1}", FinishSaveStatus, FinishECPropUpdate));
 						if (IsBuildMachine)
 						{
 							GetFailureEmails(NodeToDo, CLString);
+                            var FinishFailEmails = DateTime.Now.ToString();
+                            StepDurations.Add(String.Format("GetFailEmails{0}", NodeToDo), String.Format("{0},{1}", FinishECPropUpdate, FinishFailEmails));
 						}
 						UpdateECBuildTime(NodeToDo, BuildDuration);
                     }
@@ -7603,12 +7635,20 @@ public class GUBP : BuildCommand
                 }
                 if (SaveSuccessRecords) 
                 {
-                    UpdateNodeHistory(NodeToDo, CLString);					
+                    var StartNodeHistory = DateTime.Now.ToString();
+                    UpdateNodeHistory(NodeToDo, CLString);
+                    var FinishNodeHistory = DateTime.Now.ToString();					
                     SaveStatus(NodeToDo, SucceededTempStorageSuffix, NodeStoreName, bSaveSharedTempStorage, GameNameIfAny);
+                    var FinishSaveStatus = DateTime.Now.ToString();
+                    StepDurations.Add(String.Format("SaveNodeStatus{0}", NodeToDo), String.Format("{0},{1}", FinishNodeHistory, FinishSaveStatus));
                     UpdateECProps(NodeToDo, CLString);
+                    var FinishECPropUpdate = DateTime.Now.ToString();
+                    StepDurations.Add(String.Format("UpdateECProps{0}", NodeToDo), String.Format("{0},{1}", FinishSaveStatus, FinishECPropUpdate));
 					if (IsBuildMachine)
 					{
 						GetFailureEmails(NodeToDo, CLString);
+                        var FinishFailEmails = DateTime.Now.ToString();
+                        StepDurations.Add(String.Format("GetFailEmails{0}", NodeToDo), String.Format("{0},{1}", FinishECPropUpdate, FinishFailEmails));
 					}
 					UpdateECBuildTime(NodeToDo, BuildDuration);
                 }
@@ -7623,9 +7663,23 @@ public class GUBP : BuildCommand
             }
         }
 
+        PrintCSVFile();
         PrintRunTime();
     }
 
+    static void PrintCSVFile()
+    {
+        if (IsBuildMachine && CmdEnv.CSVFile != "")
+        {
+            var CSVBuilder = new StringBuilder();
+            foreach (KeyValuePair<string, string> Step in StepDurations)
+            {
+                var CSVLineToAppend = String.Format("{0},{1}{2}", Step.Key, Step.Value, Environment.NewLine);
+                CSVBuilder.Append(CSVLineToAppend);
+            }
+            File.AppendAllText(CmdEnv.CSVFile, CSVBuilder.ToString());
+        }
+    }
 	/// <summary>
 	/// Sorts a list of nodes to display in EC. The default order is based on execution order and agent groups, whereas this function arranges nodes by
 	/// frequency then execution order, while trying to group nodes on parallel paths (eg. Mac/Windows editor nodes) together.
