mac cis fix

diff --git a/Engine/Source/Programs/AutomationTool/Scripts/GUBP.Automation.cs b/Engine/Source/Programs/AutomationTool/Scripts/GUBP.Automation.cs
--- a/Engine/Source/Programs/AutomationTool/Scripts/GUBP.Automation.cs
+++ b/Engine/Source/Programs/AutomationTool/Scripts/GUBP.Automation.cs
@@ -7637,7 +7637,8 @@ public class GUBP : BuildCommand
                 {
                     var StartNodeHistory = DateTime.Now.ToString();
                     UpdateNodeHistory(NodeToDo, CLString);
-                    var FinishNodeHistory = DateTime.Now.ToString();					
+                    var FinishNodeHistory = DateTime.Now.ToString();
+                    StepDurations.Add(String.Format("UpdateNodeHistory{0}", NodeToDo), String.Format("{0},{1}", StartNodeHistory, FinishNodeHistory)); 
                     SaveStatus(NodeToDo, SucceededTempStorageSuffix, NodeStoreName, bSaveSharedTempStorage, GameNameIfAny);
                     var FinishSaveStatus = DateTime.Now.ToString();
                     StepDurations.Add(String.Format("SaveNodeStatus{0}", NodeToDo), String.Format("{0},{1}", FinishNodeHistory, FinishSaveStatus));
