Prioritize the PublishRocket nodes over the ChunkRocket nodes. Now that we can copy the install to the network quickly by zipping it, it's more useful for it to happen first.

diff --git a/Engine/Source/Programs/AutomationTool/Scripts/RocketBuild.Automation.cs b/Engine/Source/Programs/AutomationTool/Scripts/RocketBuild.Automation.cs
--- a/Engine/Source/Programs/AutomationTool/Scripts/RocketBuild.Automation.cs
+++ b/Engine/Source/Programs/AutomationTool/Scripts/RocketBuild.Automation.cs
@@ -817,11 +817,6 @@ namespace Rocket
 			return StaticGetFullName(HostPlatform);
 		}
 
-		public override float Priority()
-		{
-			return -10000.0f; // Do this and PublishRocketSymbolsNode at the end
-		}
-
 		public override void DoBuild(GUBP bp)
 		{
 			// Create a zip file containing the install
@@ -897,11 +892,6 @@ namespace Rocket
 			return StaticGetFullName(HostPlatform);
 		}
 
-		public override float Priority()
-		{
-			return -100000.0f; // Do this last; it's not mission critical
-		}
-
 		public override void DoBuild(GUBP bp)
 		{
 			if(RocketBuild.ShouldDoSeriousThingsLikeP4CheckinAndPostToMCP())
