[INTEGRATE] Change 2534926 by [REDACTED_PERSON][REDACTED_HANDLE] on 2015/05/02 14:04:22

	Fix search path for target platform headers.

diff --git a/Engine/Source/Programs/AutomationTool/Scripts/GUBP.Automation.cs b/Engine/Source/Programs/AutomationTool/Scripts/GUBP.Automation.cs
--- a/Engine/Source/Programs/AutomationTool/Scripts/GUBP.Automation.cs
+++ b/Engine/Source/Programs/AutomationTool/Scripts/GUBP.Automation.cs
@@ -1746,8 +1746,8 @@ public class GUBP : BuildCommand
 
 				// Archive all the headers
 				FileFilter Filter = new FileFilter();
-				Filter.Include("/Engine/Intermediate/Build/" + HostPlatform.ToString() + "/UE4/Inc/...");
-				Filter.Include("/Engine/Plugins/.../Intermediate/Build/" + HostPlatform.ToString() + "/UE4/Inc/...");
+				Filter.Include("/Engine/Intermediate/Build/" + TargetPlatform.ToString() + "/UE4/Inc/...");
+				Filter.Include("/Engine/Plugins/.../Intermediate/Build/" + TargetPlatform.ToString() + "/UE4/Inc/...");
 
 				string ZipFileName = StaticGetArchivedHeadersPath(HostPlatform, GameProj, TargetPlatform);
 				CommandUtils.ZipFiles(ZipFileName, CommandUtils.CmdEnv.LocalRoot, Filter);
