[INTEGRATE] Change 2534940 by [REDACTED_PERSON][REDACTED_HANDLE] on 2015/05/02 14:40:19

	Fix name of bootstrap executable when packaging content-only projects. Should be ProjectName.exe, not UE4Game.exe.

diff --git a/Engine/Source/Programs/AutomationTool/Win/WinPlatform.Automation.cs b/Engine/Source/Programs/AutomationTool/Win/WinPlatform.Automation.cs
--- a/Engine/Source/Programs/AutomationTool/Win/WinPlatform.Automation.cs
+++ b/Engine/Source/Programs/AutomationTool/Win/WinPlatform.Automation.cs
@@ -61,7 +61,20 @@ public abstract class BaseWinPlatform : Platform
 							BootstrapArguments = String.Format("..\\..\\..\\{0}\\{0}.uproject", SC.ShortProjectName);
 						}
 
-						string BootstrapExeName = (SC.StageTargetConfigurations.Count == 1) ? (Receipt.GetProperty("TargetName", SC.ShortProjectName) + ".exe") : Path.GetFileName(Executable.Path);
+						string BootstrapExeName;
+						if(SC.StageTargetConfigurations.Count > 1)
+						{
+							BootstrapExeName = Path.GetFileName(Executable.Path);
+						}
+						else if(Params.IsCodeBasedProject)
+						{
+							BootstrapExeName = Receipt.GetProperty("TargetName", SC.ShortProjectName) + ".exe";
+						}
+						else
+						{
+							BootstrapExeName = SC.ShortProjectName + ".exe";
+						}
+
 						StageBootstrapExecutable(SC, BootstrapExeName, Executable.Path, SC.NonUFSStagingFiles[Executable.Path], BootstrapArguments);
 					}
 				}
