Fix executable permissions on files extracted from zips.

diff --git a/Engine/Source/Programs/AutomationTool/CommandUtils.cs b/Engine/Source/Programs/AutomationTool/CommandUtils.cs
--- a/Engine/Source/Programs/AutomationTool/CommandUtils.cs
+++ b/Engine/Source/Programs/AutomationTool/CommandUtils.cs
@@ -2138,6 +2138,10 @@ namespace AutomationTool
 					{
 						Entry.Extract(OutputStream);
 					}
+					if (UnrealBuildTool.Utils.IsRunningOnMono && CommandUtils.IsProbablyAMacOrIOSExe(OutputFileName))
+					{
+						FixUnixFilePermissions(OutputFileName);
+					}
 					OutputFileNames.Add(OutputFileName);
 				}
 				return OutputFileNames;
