Write a CSV file along with each compressed PAK file, showing what's inside it and how big it is.

diff --git a/Engine/Source/Developer/DerivedDataCache/Private/PakFileDerivedDataBackend.h b/Engine/Source/Developer/DerivedDataCache/Private/PakFileDerivedDataBackend.h
--- a/Engine/Source/Developer/DerivedDataCache/Private/PakFileDerivedDataBackend.h
+++ b/Engine/Source/Developer/DerivedDataCache/Private/PakFileDerivedDataBackend.h
@@ -404,20 +404,23 @@ public:
 
 		// Copy all the DDC to the new cache
 		TArray<uint8> Buffer;
+		TArray<uint32> KeySizes;
 		for (int KeyIndex = 0; KeyIndex < KeyNames.Num(); KeyIndex++)
 		{
 			Buffer.Reset();
 			InputPak.GetCachedData(*KeyNames[KeyIndex], Buffer);
 			OutputPak.PutCachedData(*KeyNames[KeyIndex], Buffer, false);
+			KeySizes.Add(Buffer.Num());
 		}
 
 		// Write out a TOC listing for debugging
 		FStringOutputDevice Output;
+		Output.Logf(TEXT("Asset,Size") LINE_TERMINATOR);
 		for(int KeyIndex = 0; KeyIndex < KeyNames.Num(); KeyIndex++)
 		{
-			Output.Logf(TEXT("%s") LINE_TERMINATOR, *KeyNames[KeyIndex]);
+			Output.Logf(TEXT("%s,%d") LINE_TERMINATOR, *KeyNames[KeyIndex], KeySizes[KeyIndex]);
 		}
-		FFileHelper::SaveStringToFile(Output, *FPaths::Combine(*FPaths::GetPath(OutputFilename), *(FPaths::GetBaseFilename(OutputFilename) + TEXT(".txt"))));
+		FFileHelper::SaveStringToFile(Output, *FPaths::Combine(*FPaths::GetPath(OutputFilename), *(FPaths::GetBaseFilename(OutputFilename) + TEXT(".csv"))));
 		return true;
 	}
 
diff --git a/Engine/Source/Programs/AutomationTool/Scripts/RocketBuild.Automation.cs b/Engine/Source/Programs/AutomationTool/Scripts/RocketBuild.Automation.cs
--- a/Engine/Source/Programs/AutomationTool/Scripts/RocketBuild.Automation.cs
+++ b/Engine/Source/Programs/AutomationTool/Scripts/RocketBuild.Automation.cs
@@ -962,7 +962,7 @@ namespace Rocket
 
 				// Delete the output file
 				string OutputPakFile = CommandUtils.CombinePaths(CommandUtils.CmdEnv.LocalRoot, RelativePakPath);
-				string OutputTxtFile = Path.ChangeExtension(OutputPakFile, ".txt");
+				string OutputCsvFile = Path.ChangeExtension(OutputPakFile, ".csv");
 
 				// Generate DDC for all the non-code projects. We don't necessarily have editor DLLs for the code projects, but they should be the same as their blueprint counterparts.
 				List<string> ProjectPakFiles = new List<string>();
@@ -978,9 +978,9 @@ namespace Rocket
 						CommandUtils.DeleteFile(ProjectPakFile);
 						CommandUtils.RenameFile(OutputPakFile, ProjectPakFile);
 
-						string ProjectTxtFile = Path.ChangeExtension(ProjectPakFile, ".txt");
-						CommandUtils.DeleteFile(ProjectTxtFile);
-						CommandUtils.RenameFile(OutputTxtFile, ProjectTxtFile);
+						string ProjectCsvFile = Path.ChangeExtension(ProjectPakFile, ".csv");
+						CommandUtils.DeleteFile(ProjectCsvFile);
+						CommandUtils.RenameFile(OutputCsvFile, ProjectCsvFile);
 
 						ProjectPakFiles.Add(Path.GetFileName(ProjectPakFile));
 					}
