Re-enabled separate word count generation via localization automation script.

diff --git a/Engine/Source/Programs/AutomationTool/Scripts/Localisation.Automation.cs b/Engine/Source/Programs/AutomationTool/Scripts/Localisation.Automation.cs
--- a/Engine/Source/Programs/AutomationTool/Scripts/Localisation.Automation.cs
+++ b/Engine/Source/Programs/AutomationTool/Scripts/Localisation.Automation.cs
@@ -87,6 +87,7 @@ class Localise : BuildCommand
 				String.Format("-config={0}", @"./Config/Localization/EditorTutorials.ini") + (string.IsNullOrEmpty(CommandletSCCArguments) ? "" : " " + CommandletSCCArguments),
 				String.Format("-config={0}", @"./Config/Localization/PropertyNames.ini") + (string.IsNullOrEmpty(CommandletSCCArguments) ? "" : " " + CommandletSCCArguments),
 				String.Format("-config={0}", @"./Config/Localization/ToolTips.ini") + (string.IsNullOrEmpty(CommandletSCCArguments) ? "" : " " + CommandletSCCArguments),
+				String.Format("-config={0}", @"./Config/Localization/WordCount.ini"),
 			};
 
 		// Execute commandlet for each set of arguments.
