build script, keep sticky agent for editor builds

diff --git a/Engine/Source/Programs/AutomationTool/Scripts/GUBP.Automation.cs b/Engine/Source/Programs/AutomationTool/Scripts/GUBP.Automation.cs
--- a/Engine/Source/Programs/AutomationTool/Scripts/GUBP.Automation.cs
+++ b/Engine/Source/Programs/AutomationTool/Scripts/GUBP.Automation.cs
@@ -769,7 +769,10 @@ public class GUBP : BuildCommand
         public ToolsForCompileNode(UnrealTargetPlatform InHostPlatform)
             : base(InHostPlatform, false)
         {
-            AgentSharingGroup = "Editor" + StaticGetHostPlatformSuffix(InHostPlatform);
+            if (InHostPlatform != UnrealTargetPlatform.Win64)
+            {
+                AgentSharingGroup = "Editor" + StaticGetHostPlatformSuffix(InHostPlatform);
+            }
         }
         public static string StaticGetFullName(UnrealTargetPlatform InHostPlatform)
         {
@@ -779,6 +782,22 @@ public class GUBP : BuildCommand
         {
             return StaticGetFullName(HostPlatform);
         }
+        public override bool IsSticky()
+        {
+            if (HostPlatform == UnrealTargetPlatform.Win64)
+            {
+                return true;
+            }
+            return false;
+        }
+        public override string ECProcedure()
+        {
+            if (HostPlatform == UnrealTargetPlatform.Win64)
+            {
+                return String.Format("GUBP_UAT_Node_Parallel_AgentShare_Editor");
+            }
+            return base.ECProcedure();
+        }
         public override bool DeleteBuildProducts()
         {
             return true;
@@ -824,7 +843,10 @@ public class GUBP : BuildCommand
         public RootEditorNode(UnrealTargetPlatform InHostPlatform)
             : base(InHostPlatform)
         {
-            AgentSharingGroup = "Editor" + StaticGetHostPlatformSuffix(InHostPlatform);
+            if (InHostPlatform != UnrealTargetPlatform.Win64)
+            {
+                AgentSharingGroup = "Editor" + StaticGetHostPlatformSuffix(InHostPlatform);
+            }
         }
         public static string StaticGetFullName(UnrealTargetPlatform InHostPlatform)
         {
@@ -857,6 +879,22 @@ public class GUBP : BuildCommand
 		{
 			return CommandUtils.CombinePaths(CommandUtils.CmdEnv.LocalRoot, "Engine", "Saved", "Precompiled", "Headers-RootEditor" + StaticGetHostPlatformSuffix(HostPlatform) + ".zip");
 		}
+        public override bool IsSticky()
+        {
+            if (HostPlatform == UnrealTargetPlatform.Win64)
+            {
+                return true;
+            }
+            return false;
+        }
+        public override string ECProcedure()
+        {
+            if (HostPlatform == UnrealTargetPlatform.Win64)
+            {
+                return String.Format("GUBP_UAT_Node_Parallel_AgentShare_Editor");
+            }
+            return base.ECProcedure();
+        }
         public override UE4Build.BuildAgenda GetAgenda(GUBP bp)
         {
             var Agenda = new UE4Build.BuildAgenda();
@@ -1392,7 +1430,10 @@ public class GUBP : BuildCommand
         public EditorPlatformNode(UnrealTargetPlatform InHostPlatform, UnrealTargetPlatform Plat)
             : base(InHostPlatform)
         {
-            AgentSharingGroup = "Editor" + StaticGetHostPlatformSuffix(InHostPlatform);
+            if (InHostPlatform != UnrealTargetPlatform.Win64)
+            {
+                AgentSharingGroup = "Editor" + StaticGetHostPlatformSuffix(InHostPlatform);
+            }
             EditorPlatform = Plat;
             AddDependency(RootEditorNode.StaticGetFullName(InHostPlatform));
         }
@@ -1409,6 +1450,22 @@ public class GUBP : BuildCommand
         {
             return base.Priority() + 1;
         }
+        public override bool IsSticky()
+        {
+            if (HostPlatform == UnrealTargetPlatform.Win64)
+            {
+                return true;
+            }
+            return false;
+        }
+        public override string ECProcedure()
+        {
+            if (HostPlatform == UnrealTargetPlatform.Win64)
+            {
+                return String.Format("GUBP_UAT_Node_Parallel_AgentShare_Editor");
+            }
+            return base.ECProcedure();
+        }
         public override int CISFrequencyQuantumShift(GUBP bp)
         {
             int Result = base.CISFrequencyQuantumShift(bp);            
@@ -1445,7 +1502,10 @@ public class GUBP : BuildCommand
         public EditorGameNode(GUBP bp, UnrealTargetPlatform InHostPlatform, BranchInfo.BranchUProject InGameProj)
             : base(InHostPlatform)
         {
-            AgentSharingGroup = "Editor"  + StaticGetHostPlatformSuffix(InHostPlatform);
+            if (InHostPlatform != UnrealTargetPlatform.Win64)
+            {
+                AgentSharingGroup = "Editor" + StaticGetHostPlatformSuffix(InHostPlatform);
+            }
             GameProjects.Add(InGameProj);
 			AddDependency(RootEditorNode.StaticGetFullName(InHostPlatform));						
 		}
@@ -1465,6 +1525,22 @@ public class GUBP : BuildCommand
         {
             return StaticGetFullName(HostPlatform, GameProjects[0]);
         }
+        public override bool IsSticky()
+        {
+            if (HostPlatform == UnrealTargetPlatform.Win64)
+            {
+                return true;
+            }
+            return false;
+        }
+        public override string ECProcedure()
+        {
+            if (HostPlatform == UnrealTargetPlatform.Win64)
+            {
+                return String.Format("GUBP_UAT_Node_Parallel_AgentShare_Editor");
+            }
+            return base.ECProcedure();
+        }
         public override string GameNameIfAnyForTempStorage()
         {
             return GameProjects[0].Options(HostPlatform).GroupName ?? GameProjects[0].GameName;
