build script, remove unrealversionselector from fort branch

diff --git a/Engine/Source/Programs/AutomationTool/Scripts/GUBP.Automation.cs b/Engine/Source/Programs/AutomationTool/Scripts/GUBP.Automation.cs
--- a/Engine/Source/Programs/AutomationTool/Scripts/GUBP.Automation.cs
+++ b/Engine/Source/Programs/AutomationTool/Scripts/GUBP.Automation.cs
@@ -1169,11 +1169,14 @@ public class GUBP : BuildCommand
 				bool CrossCompile;
                 if (ProgramTarget.Rules.GUBP_AlwaysBuildWithTools(HostPlatform, out bInternalOnly, out SeparateNode, out CrossCompile) && ProgramTarget.Rules.SupportsPlatform(HostPlatform) && !bInternalOnly && !SeparateNode)
                 {
-                    foreach (var Plat in ProgramTarget.Rules.GUBP_ToolPlatforms(HostPlatform))
+                    if (!bp.BranchOptions.ExcludeNodes.Contains(ProgramTarget.TargetName))
                     {
-                        foreach (var Config in ProgramTarget.Rules.GUBP_ToolConfigs(HostPlatform))
+                        foreach (var Plat in ProgramTarget.Rules.GUBP_ToolPlatforms(HostPlatform))
                         {
-                            Agenda.AddTargets(new string[] { ProgramTarget.TargetName }, Plat, Config, InAddArgs: AddArgs);
+                            foreach (var Config in ProgramTarget.Rules.GUBP_ToolConfigs(HostPlatform))
+                            {
+                                Agenda.AddTargets(new string[] { ProgramTarget.TargetName }, Plat, Config, InAddArgs: AddArgs);
+                            }
                         }
                     }
                 }
