Merging using UE4-To-UE4-LauncherDev

diff --git a/Engine/Source/Programs/AutomationTool/Scripts/LauncherLocalization.Automation.cs b/Engine/Source/Programs/AutomationTool/Scripts/LauncherLocalization.Automation.cs
--- a/Engine/Source/Programs/AutomationTool/Scripts/LauncherLocalization.Automation.cs
+++ b/Engine/Source/Programs/AutomationTool/Scripts/LauncherLocalization.Automation.cs
@@ -66,8 +66,7 @@ class LauncherLocalization : BuildCommand
 		// Setup commandlet arguments with configurations.
 		var CommandletArgumentSets = new string[] 
 			{
-				String.Format("-config={0}", @"../Portal/Config/Localization/App.ini") + (string.IsNullOrEmpty(CommandletSCCArguments) ? "" : " " + CommandletSCCArguments),
-				String.Format("-config={0}", @"../Portal/Config/Localization/WordCount.ini"),
+				String.Format("-config={0}", @"../Portal/Config/Localization/App.ini") + (string.IsNullOrEmpty(CommandletSCCArguments) ? "" : " " + CommandletSCCArguments)
 			};
 
 		// Execute commandlet for each set of arguments.
