Add missing template content into Rocket build.

diff --git a/Engine/Source/Programs/AutomationTool/Scripts/RocketBuild.Automation.cs b/Engine/Source/Programs/AutomationTool/Scripts/RocketBuild.Automation.cs
--- a/Engine/Source/Programs/AutomationTool/Scripts/RocketBuild.Automation.cs
+++ b/Engine/Source/Programs/AutomationTool/Scripts/RocketBuild.Automation.cs
@@ -14,6 +14,8 @@ namespace Rocket
 	{
 		static readonly string[] CurrentTemplates = 
 		{
+			"FP_FirstPerson",
+			"FP_FirstPersonBP",
 			"TP_FirstPerson",
 			"TP_FirstPersonBP",
 			"TP_Flying",
