Fix various issues being able to package a plugin using the installed engine build.

diff --git a/Engine/Build/InstalledEngineFilters.ini b/Engine/Build/InstalledEngineFilters.ini
--- a/Engine/Build/InstalledEngineFilters.ini
+++ b/Engine/Build/InstalledEngineFilters.ini
@@ -37,6 +37,8 @@
 {Win64} -/Engine/Plugins/.../Binaries/Mac/...
 
 ; Source code
+/Engine/Source/UE4Game.Target.cs
+/Engine/Source/UE4Editor.Target.cs
 /Engine/Source/Runtime/...
 /Engine/Source/Developer/...
 /Engine/Source/Editor/...
@@ -123,8 +125,6 @@
 
 {Mac} /Engine/Extras/LLDBDataFormatters/UE4DataFormatters.py
 
-{Mac} /Engine/Source/UE4Game*.cs
-
 ; Don't want these folders, even if they're part of ToolsForCompile
 {Mac} -/Engine/Binaries/Win32/...
 {Mac} -/Engine/Binaries/Win64/...
diff --git a/Engine/Source/Programs/AutomationTool/Scripts/BuildPluginCommand.Automation.cs b/Engine/Source/Programs/AutomationTool/Scripts/BuildPluginCommand.Automation.cs
--- a/Engine/Source/Programs/AutomationTool/Scripts/BuildPluginCommand.Automation.cs
+++ b/Engine/Source/Programs/AutomationTool/Scripts/BuildPluginCommand.Automation.cs
@@ -33,13 +33,20 @@ class BuildPlugin : BuildCommand
 			CommandUtils.DeleteDirectory(IntermediateBuildDirectory);
 		}
 
+		// Get any additional arguments from the commandline
+		string AdditionalArgs = "";
+		if(ParseParam("Rocket"))
+		{
+			AdditionalArgs += " -Rocket";
+		}
+
 		// Build the host platforms
 		List<string> ReceiptFileNames = new List<string>();
 		UE4Build.BuildAgenda Agenda = new UE4Build.BuildAgenda();
 		UnrealTargetPlatform HostPlatform = BuildHostPlatform.Current.Platform;
 		if(!ParseParam("NoHostPlatform"))
 		{
-			AddPluginToAgenda(Agenda, PluginFileName, Plugin, "UE4Editor", TargetRules.TargetType.Editor, HostPlatform, UnrealTargetConfiguration.Development, ReceiptFileNames);
+			AddPluginToAgenda(Agenda, PluginFileName, Plugin, "UE4Editor", TargetRules.TargetType.Editor, HostPlatform, UnrealTargetConfiguration.Development, ReceiptFileNames, AdditionalArgs);
 		}
 
 		// Add the game targets
@@ -48,8 +55,8 @@ class BuildPlugin : BuildCommand
 		{
 			if(Rocket.RocketBuild.IsCodeTargetPlatform(HostPlatform, TargetPlatform))
 			{
-				AddPluginToAgenda(Agenda, PluginFileName, Plugin, "UE4Game", TargetRules.TargetType.Game, TargetPlatform, UnrealTargetConfiguration.Development, ReceiptFileNames);
-				AddPluginToAgenda(Agenda, PluginFileName, Plugin, "UE4Game", TargetRules.TargetType.Game, TargetPlatform, UnrealTargetConfiguration.Shipping, ReceiptFileNames);
+				AddPluginToAgenda(Agenda, PluginFileName, Plugin, "UE4Game", TargetRules.TargetType.Game, TargetPlatform, UnrealTargetConfiguration.Development, ReceiptFileNames, AdditionalArgs);
+				AddPluginToAgenda(Agenda, PluginFileName, Plugin, "UE4Game", TargetRules.TargetType.Game, TargetPlatform, UnrealTargetConfiguration.Shipping, ReceiptFileNames, AdditionalArgs);
 			}
 		}
 
@@ -88,7 +95,7 @@ class BuildPlugin : BuildCommand
 		return Platforms;
 	}
 
-	static void AddPluginToAgenda(UE4Build.BuildAgenda Agenda, string PluginFileName, PluginDescriptor Plugin, string TargetName, TargetRules.TargetType TargetType, UnrealTargetPlatform Platform, UnrealTargetConfiguration Configuration, List<string> ReceiptFileNames)
+	static void AddPluginToAgenda(UE4Build.BuildAgenda Agenda, string PluginFileName, PluginDescriptor Plugin, string TargetName, TargetRules.TargetType TargetType, UnrealTargetPlatform Platform, UnrealTargetConfiguration Configuration, List<string> ReceiptFileNames, string InAdditionalArgs)
 	{
 		// Find a list of modules that need to be built for this plugin
 		List<string> ModuleNames = new List<string>();
@@ -112,6 +119,11 @@ class BuildPlugin : BuildCommand
 			string ReceiptFileName = BuildReceipt.GetDefaultPath(Path.GetDirectoryName(PluginFileName), TargetName, Platform, Configuration, "");
 			Arguments += String.Format(" -receipt {0}", CommandUtils.MakePathSafeToUseWithCommandLine(ReceiptFileName));
 			ReceiptFileNames.Add(ReceiptFileName);
+			
+			if(!String.IsNullOrEmpty(InAdditionalArgs))
+			{
+				Arguments += InAdditionalArgs;
+			}
 
 			Agenda.AddTarget(TargetName, Platform, Configuration, InAddArgs: Arguments);
 		}
diff --git a/Engine/Source/Programs/UnrealBuildTool/Configuration/UEBuildTarget.cs b/Engine/Source/Programs/UnrealBuildTool/Configuration/UEBuildTarget.cs
--- a/Engine/Source/Programs/UnrealBuildTool/Configuration/UEBuildTarget.cs
+++ b/Engine/Source/Programs/UnrealBuildTool/Configuration/UEBuildTarget.cs
@@ -1098,7 +1098,7 @@ namespace UnrealBuildTool
 				}
 
 				// Clean the intermediate directory
-				if( !String.IsNullOrEmpty( ProjectIntermediateDirectory ) )
+				if( !String.IsNullOrEmpty( ProjectIntermediateDirectory ) && (!UnrealBuildTool.IsEngineInstalled() || !Utils.IsFileUnderDirectory(ProjectIntermediateDirectory, BuildConfiguration.RelativeEnginePath)))
 				{
 					if (Directory.Exists(ProjectIntermediateDirectory))
 					{
@@ -1799,12 +1799,12 @@ namespace UnrealBuildTool
 				var FilteredBinaries = new List<UEBuildBinary>();
 
 				// Have to do absolute here as this could be a project that is under the root
-				string FullUProjectPath = Path.GetFullPath(UnrealBuildTool.GetUProjectPath());
+				string FullEnginePath = Path.GetFullPath(BuildConfiguration.RelativeEnginePath);
 
 				// We only want to build rocket projects...
 				foreach (var DLLBinary in AppBinaries)
 				{
-					if (Utils.IsFileUnderDirectory(DLLBinary.Config.OutputFilePath, FullUProjectPath))
+					if (!Utils.IsFileUnderDirectory(DLLBinary.Config.OutputFilePath, FullEnginePath))
 					{
 						FilteredBinaries.Add(DLLBinary);
 					}
@@ -3011,8 +3011,7 @@ namespace UnrealBuildTool
 					}
 
 					// So all we care about are the game module and/or plugins.
-					//@todo Rocket: This assumes plugins that have source will be under the game folder...
-					if (bDiscoverFiles && (!UnrealBuildTool.RunningRocket() || Utils.IsFileUnderDirectory(ModuleFileName, UnrealBuildTool.GetUProjectPath())))
+					if (bDiscoverFiles && (!UnrealBuildTool.IsEngineInstalled() || !Utils.IsFileUnderDirectory(ModuleFileName, BuildConfiguration.RelativeEnginePath)))
 					{
 						var SourceFilePaths = new List<string>();
 
diff --git a/Engine/Source/Programs/UnrealBuildTool/System/ActionHistory.cs b/Engine/Source/Programs/UnrealBuildTool/System/ActionHistory.cs
--- a/Engine/Source/Programs/UnrealBuildTool/System/ActionHistory.cs
+++ b/Engine/Source/Programs/UnrealBuildTool/System/ActionHistory.cs
@@ -109,7 +109,7 @@ namespace UnrealBuildTool
 			else
 			{
 				// Shared action history (unless this is a rocket target)
-				Folder = UnrealBuildTool.RunningRocket() ?
+				Folder = (UnrealBuildTool.RunningRocket() && UnrealBuildTool.HasUProjectFile()) ?
 					Path.Combine(UnrealBuildTool.GetUProjectPath(), BuildConfiguration.BaseIntermediateFolder) :
 					BuildConfiguration.BaseIntermediatePath;
 			}
diff --git a/Engine/Source/Programs/UnrealBuildTool/System/ExternalExecution.cs b/Engine/Source/Programs/UnrealBuildTool/System/ExternalExecution.cs
--- a/Engine/Source/Programs/UnrealBuildTool/System/ExternalExecution.cs
+++ b/Engine/Source/Programs/UnrealBuildTool/System/ExternalExecution.cs
@@ -359,15 +359,10 @@ namespace UnrealBuildTool
 
 			foreach( var Module in UObjectModules )
 			{
-				// In Rocket, we skip checking timestamps for modules that don't exist within the project's directory
-				if (UnrealBuildTool.RunningRocket())
+				// If the engine is installed, skip skip checking timestamps for modules that are under the engine directory
+				if (UnrealBuildTool.IsEngineInstalled() && Utils.IsFileUnderDirectory( Module.ModuleDirectory, BuildConfiguration.RelativeEnginePath))
 				{
-					// @todo Rocket: This could be done in a better way I'm sure
-					if (!Utils.IsFileUnderDirectory( Module.ModuleDirectory, UnrealBuildTool.GetUProjectPath() ))
-					{
-						// Engine or engine plugin module - Rocket does not regenerate them so don't compare their timestamps
-						continue;
-					}
+					continue;
 				}
 
 				// Make sure we have an existing folder for generated code.  If not, then we definitely need to generate code!
diff --git a/Engine/Source/Programs/UnrealBuildTool/System/UnrealBuildTool.cs b/Engine/Source/Programs/UnrealBuildTool/System/UnrealBuildTool.cs
--- a/Engine/Source/Programs/UnrealBuildTool/System/UnrealBuildTool.cs
+++ b/Engine/Source/Programs/UnrealBuildTool/System/UnrealBuildTool.cs
@@ -1507,7 +1507,7 @@ namespace UnrealBuildTool
 				}
 
 				UEBuildConfiguration.bHotReloadFromIDE = UEBuildConfiguration.bAllowHotReloadFromIDE && TargetDescs.Count == 1 && !TargetDescs[0].bIsEditorRecompile && ShouldDoHotReloadFromIDE(TargetDescs[0]);
-				bool bIsHotReload = UEBuildConfiguration.bHotReloadFromIDE || ( TargetDescs.Count == 1 && TargetDescs[0].OnlyModules.Count > 0 );
+				bool bIsHotReload = UEBuildConfiguration.bHotReloadFromIDE || ( TargetDescs.Count == 1 && TargetDescs[0].OnlyModules.Count > 0 && TargetDescs[0].ForeignPlugins.Count == 0 );
 				TargetDescriptor HotReloadTargetDesc = bIsHotReload ? TargetDescs[0] : null;
 
 				if (ProjectFileGenerator.bGenerateProjectFiles)
@@ -2475,7 +2475,6 @@ namespace UnrealBuildTool
 					{
 						// Remove "-####-Platform-Configuration" suffix in Debug configuration
 						var IndexOfSecondLastHyphen = OriginalFileNameWithoutExtension.LastIndexOf('-', IndexOfLastHyphen - 1, IndexOfLastHyphen);
-						var SuffixLenght = OriginalFileNameWithoutExtension.Length - IndexOfSecondLastHyphen;
 						var IndexOfThirdLastHyphen = OriginalFileNameWithoutExtension.LastIndexOf('-', IndexOfSecondLastHyphen - 1, IndexOfSecondLastHyphen);
 						OriginalFileNameWithoutNumberSuffix = OriginalFileNameWithoutExtension.Substring(0, IndexOfThirdLastHyphen);
 						PlatformConfigSuffix = OriginalFileNameWithoutExtension.Substring(IndexOfSecondLastHyphen);
