UAT: fix packaging into directories with spaces (UE-12940).

- Unlike Windows, quotes from paths weren't stripped out on Linux. Mac might have been affected as well.

diff --git a/Engine/Source/Programs/AutomationTool/ProjectParams.cs b/Engine/Source/Programs/AutomationTool/ProjectParams.cs
--- a/Engine/Source/Programs/AutomationTool/ProjectParams.cs
+++ b/Engine/Source/Programs/AutomationTool/ProjectParams.cs
@@ -88,21 +88,22 @@ namespace AutomationTool
 		/// <param name="SpecifiedValue">Value specified in the constructor (or not)</param>
 		/// <param name="ParamName">Command line parameter name to parse.</param>
 		/// <param name="Default">Default value</param>
+		/// <param name="bTrimQuotes">If set, the leading and trailing quotes will be removed, e.g. instead of "[REDACTED_USER_HOME] Name" it will return [REDACTED_USER_HOME] Name</param>
 		/// <returns>Parameter value.</returns>
-		string ParseParamValueIfNotSpecified(CommandUtils Command, string SpecifiedValue, string ParamName, string Default = "")
+		string ParseParamValueIfNotSpecified(CommandUtils Command, string SpecifiedValue, string ParamName, string Default = "", bool bTrimQuotes = false)
 		{
+			string Result = Default;
+
 			if (SpecifiedValue != null)
 			{
-				return SpecifiedValue;
+				Result = SpecifiedValue;
 			}
 			else if (Command != null)
 			{
-				return Command.ParseParamValue(ParamName, Default);
-			}
-			else
-			{
-				return Default;
+				Result = Command.ParseParamValue(ParamName, Default);
 			}
+
+			return bTrimQuotes ? Result.Trim( new char[]{'\"'} ) : Result;
 		}
 
 		/// <summary>
@@ -494,14 +495,14 @@ namespace AutomationTool
 			{
 				this.Stage = true;
 			}
-			this.StageDirectoryParam = ParseParamValueIfNotSpecified(Command, StageDirectoryParam, "stagingdirectory", String.Empty).Trim(new char[]{'\"'});
+			this.StageDirectoryParam = ParseParamValueIfNotSpecified(Command, StageDirectoryParam, "stagingdirectory", String.Empty, true);
             this.StageNonMonolithic = GetParamValueIfNotSpecified(Command, StageNonMonolithic, this.StageNonMonolithic, "StageNonMonolithic");
 			this.Manifests = GetParamValueIfNotSpecified(Command, Manifests, this.Manifests, "manifests");
             this.CreateChunkInstall = GetParamValueIfNotSpecified(Command, CreateChunkInstall, this.CreateChunkInstall, "createchunkinstall");
-			this.ChunkInstallDirectory = ParseParamValueIfNotSpecified(Command, ChunkInstallDirectory, "chunkinstalldirectory", String.Empty).Trim(new char[]{'\"'});
-			this.ChunkInstallVersionString = ParseParamValueIfNotSpecified(Command, ChunkInstallVersionString, "chunkinstallversion", String.Empty).Trim(new char[] { '\"' });
+			this.ChunkInstallDirectory = ParseParamValueIfNotSpecified(Command, ChunkInstallDirectory, "chunkinstalldirectory", String.Empty, true);
+			this.ChunkInstallVersionString = ParseParamValueIfNotSpecified(Command, ChunkInstallVersionString, "chunkinstallversion", String.Empty, true);
 			this.Archive = GetParamValueIfNotSpecified(Command, Archive, this.Archive, "archive");
-			this.ArchiveDirectoryParam = ParseParamValueIfNotSpecified(Command, ArchiveDirectoryParam, "archivedirectory", String.Empty);
+			this.ArchiveDirectoryParam = ParseParamValueIfNotSpecified(Command, ArchiveDirectoryParam, "archivedirectory", String.Empty, true);
 			this.ArchiveMetaData = GetParamValueIfNotSpecified(Command, ArchiveMetaData, this.ArchiveMetaData, "archivemetadata");
 			this.Distribution = GetParamValueIfNotSpecified(Command, Distribution, this.Distribution, "distribution");
 			this.Prereqs = GetParamValueIfNotSpecified(Command, Prereqs, this.Prereqs, "prereqs");
