UAT: Fix LaunchOn under Linux (UE-5491).

- Use the binary name that matches the staged binary.
- Do not require IP if running under Linux and not starting remotely.

diff --git a/Engine/Source/Programs/AutomationTool/Linux/LinuxPlatform.Automation.cs b/Engine/Source/Programs/AutomationTool/Linux/LinuxPlatform.Automation.cs
--- a/Engine/Source/Programs/AutomationTool/Linux/LinuxPlatform.Automation.cs
+++ b/Engine/Source/Programs/AutomationTool/Linux/LinuxPlatform.Automation.cs
@@ -23,6 +23,17 @@ public abstract class BaseLinuxPlatform : Platform
 	{
 	}
 
+	public override List<string> GetExecutableNames(DeploymentContext SC, bool bIsRun = false)
+	{
+		List<string> Exes = base.GetExecutableNames(SC, bIsRun);
+		// replace the binary name to match what was staged
+		if (bIsRun && !SC.IsCodeBasedProject)
+		{
+			Exes[0] = CommandUtils.CombinePaths(SC.StageProjectRoot, "Binaries", SC.PlatformDir, SC.ShortProjectName);
+		}
+		return Exes;
+	}
+
 	public override void GetFilesToDeployOrStage(ProjectParams Params, DeploymentContext SC)
 	{
 		// FIXME: use build architecture
@@ -348,9 +359,9 @@ chmod 700 $HOME/Desktop/{1}.desktop", DesiredGLVersion, SC.ShortProjectName, SC.
 			// Use key as password
 			ProjParams.DevicePassword = linuxKey;
 		}
-		else if (ProjParams.Deploy || ProjParams.Run)
+		else if ((ProjParams.Deploy || ProjParams.Run) && BuildHostPlatform.Current.Platform != UnrealTargetPlatform.Linux)
 		{
-			throw new AutomationException("must specify device for Linux target (-serverdevice=<ip>)");
+			throw new AutomationException("must specify device IP for remote Linux target (-serverdevice=<ip>)");
 		}
 	}
 	public override List<string> GetDebugFileExtentions()
