UAT: Minor fixes to RocketSync and BuildPlugin scripts, mainly to improve help attributes.

diff --git a/Engine/Source/Programs/AutomationTool/Scripts/BuildPluginCommand.Automation.cs b/Engine/Source/Programs/AutomationTool/Scripts/BuildPluginCommand.Automation.cs
--- a/Engine/Source/Programs/AutomationTool/Scripts/BuildPluginCommand.Automation.cs
+++ b/Engine/Source/Programs/AutomationTool/Scripts/BuildPluginCommand.Automation.cs
@@ -9,9 +9,10 @@ using AutomationTool;
 using UnrealBuildTool;
 
 [Help("Builds a plugin, and packages it for distribution")]
-[Help("Plugin", "Specify the path to the plugin that should be packaged")]
-[Help("NoEditor", "Prevent compiling for the editor platform")]
+[Help("Plugin", "Specify the path to the descriptor file for the plugin that should be packaged")]
+[Help("NoHostPlatform", "Prevent compiling for the editor platform on the host")]
 [Help("TargetPlatforms", "Specify a list of target platforms to build, separated by '+' characters (eg. -TargetPlatforms=Win32+Win64). Default is all the Rocket target platforms.")]
+[Help("Package", "The path which the build artifcats should be packaged to, ready for distribution.")]
 class BuildPlugin : BuildCommand
 {
 	public override void ExecuteBuild()
@@ -22,7 +23,7 @@ class BuildPlugin : BuildCommand
 		{
 			throw new AutomationException("Plugin file name was not specified via the -plugin argument");
 		}
-			
+
 		// Read the plugin
 		PluginDescriptor Plugin = PluginDescriptor.FromFile(PluginFileName);
 
@@ -73,28 +74,6 @@ class BuildPlugin : BuildCommand
 		}
 	}
 
-	List<UnrealTargetPlatform> ParseParamPlatforms(string ParamName)
-	{
-		// Parse the raw parameter
-		string ParamValue = ParseParamValue(ParamName, null);
-
-		// Convert it to a list of target platform enums
-		List<UnrealTargetPlatform> Platforms = new List<UnrealTargetPlatform>();
-		if(ParamValue != null)
-		{
-			foreach(string PlatformName in ParamValue.Split(new char[]{ '+' }, StringSplitOptions.RemoveEmptyEntries))
-			{
-				UnrealTargetPlatform Platform;
-				if(!Enum.TryParse<UnrealTargetPlatform>(PlatformName, true, out Platform))
-				{
-					throw new AutomationException("'{0}' is not a valid platform name; valid platforms are: {1}.", PlatformName, String.Join(", ", Enum.GetNames(typeof(UnrealTargetPlatform))));
-				}
-				Platforms.Add(Platform);
-			}
-		}
-		return Platforms;
-	}
-
 	static void AddPluginToAgenda(UE4Build.BuildAgenda Agenda, string PluginFileName, PluginDescriptor Plugin, string TargetName, TargetRules.TargetType TargetType, UnrealTargetPlatform Platform, UnrealTargetConfiguration Configuration, List<string> ReceiptFileNames, string InAdditionalArgs)
 	{
 		// Find a list of modules that need to be built for this plugin
