UAT: Fix type in HelpAttribute of BuildPlugin script.

diff --git a/Engine/Source/Programs/AutomationTool/Scripts/BuildPluginCommand.Automation.cs b/Engine/Source/Programs/AutomationTool/Scripts/BuildPluginCommand.Automation.cs
--- a/Engine/Source/Programs/AutomationTool/Scripts/BuildPluginCommand.Automation.cs
+++ b/Engine/Source/Programs/AutomationTool/Scripts/BuildPluginCommand.Automation.cs
@@ -12,7 +12,7 @@ using UnrealBuildTool;
 [Help("Plugin", "Specify the path to the descriptor file for the plugin that should be packaged")]
 [Help("NoHostPlatform", "Prevent compiling for the editor platform on the host")]
 [Help("TargetPlatforms", "Specify a list of target platforms to build, separated by '+' characters (eg. -TargetPlatforms=Win32+Win64). Default is all the Rocket target platforms.")]
-[Help("Package", "The path which the build artifcats should be packaged to, ready for distribution.")]
+[Help("Package", "The path which the build artifacts should be packaged to, ready for distribution.")]
 class BuildPlugin : BuildCommand
 {
 	public override void ExecuteBuild()
