build script, better logging

diff --git a/Engine/Source/Programs/AutomationTool/Automation.cs b/Engine/Source/Programs/AutomationTool/Automation.cs
--- a/Engine/Source/Programs/AutomationTool/Automation.cs
+++ b/Engine/Source/Programs/AutomationTool/Automation.cs
@@ -319,10 +319,10 @@ AutomationTool.exe [-verbose] [-compileonly] [-p4] Command0 [-Arg0 -Arg1 -Arg2 
 			// Compile scripts.
 			Log.TraceInformation("Compiling scripts.");
 			ScriptCompiler Compiler = new ScriptCompiler();
-            var StartScriptCompile = DateTime.UtcNow.ToString();
+            var StartScriptCompile = DateTime.Now.ToString();
 			Compiler.FindAndCompileAllScripts(AdditionalScriptsFolders: AdditionalScriptsFolders);
-            var FinishScriptCompile = DateTime.UtcNow.ToString();
-            CommandUtils.StepDurations.Add("ScriptCompile", String.Format("{0},{1}", StartScriptCompile, FinishScriptCompile));
+            var FinishScriptCompile = DateTime.Now.ToString();
+            CommandUtils.StepDurations.Add("UAT,ScriptCompile", String.Format("{0},{1}", StartScriptCompile, FinishScriptCompile));
 
 			if (GlobalCommandLine.CompileOnly)
 			{
diff --git a/Engine/Source/Programs/AutomationTool/ProjectUtils.cs b/Engine/Source/Programs/AutomationTool/ProjectUtils.cs
--- a/Engine/Source/Programs/AutomationTool/ProjectUtils.cs
+++ b/Engine/Source/Programs/AutomationTool/ProjectUtils.cs
@@ -707,7 +707,7 @@ namespace AutomationTool
             BaseEngineProject = new BranchUProject();
 
             var AllProjects = UnrealBuildTool.UProjectInfo.FilterGameProjects(false, null);
-            var StartTime = DateTime.UtcNow.ToString();
+            var StartTime = DateTime.Now.ToString();
             foreach (var InfoEntry in AllProjects)
             {
                 var UProject = new BranchUProject(InfoEntry);
@@ -725,15 +725,15 @@ namespace AutomationTool
                     }
                 }
             }
-            var FinishTime = DateTime.UtcNow.ToString();
-            CommandUtils.StepDurations.Add("SortProjects", String.Format("{0},{1}", StartTime, FinishTime));
+            var FinishTime = DateTime.Now.ToString();
+            CommandUtils.StepDurations.Add("UAT,SortProjects", String.Format("{0},{1}", StartTime, FinishTime));
             if (String.IsNullOrEmpty(BaseEngineProject.FilePath))
             {
                 throw new AutomationException("All branches must have the blank project /Samples/Sandbox/BlankProject");
             }
 
             CommandUtils.Log("  Base Engine:");
-            var StartBranchDump = DateTime.UtcNow.ToString();
+            var StartBranchDump = DateTime.Now.ToString();
             BaseEngineProject.Dump(InHostPlatforms);
 
             CommandUtils.Log("  {0} Code projects:", CodeProjects.Count);
@@ -746,8 +746,8 @@ namespace AutomationTool
             {
                 Proj.Dump(InHostPlatforms);
             }
-            var BuildDumpFinish = DateTime.UtcNow.ToString();
-            CommandUtils.StepDurations.Add("Project Dump", String.Format("{0},{1}", StartBranchDump, BuildDumpFinish));
+            var BuildDumpFinish = DateTime.Now.ToString();
+            CommandUtils.StepDurations.Add("UAT,Project Dump", String.Format("{0},{1}", StartBranchDump, BuildDumpFinish));
         }
 
         public BranchUProject FindGame(string GameName)
diff --git a/Engine/Source/Programs/AutomationTool/Scripts/GUBP.Automation.cs b/Engine/Source/Programs/AutomationTool/Scripts/GUBP.Automation.cs
--- a/Engine/Source/Programs/AutomationTool/Scripts/GUBP.Automation.cs
+++ b/Engine/Source/Programs/AutomationTool/Scripts/GUBP.Automation.cs
@@ -6526,7 +6526,7 @@ public class GUBP : BuildCommand
         {
             Log("******* {0} GUBP Nodes", GUBPNodes.Count);
             var SortedNodes = TopologicalSort(new HashSet<string>(GUBPNodes.Keys), LocalOnly: true, DoNotConsiderCompletion: true);
-            var DependencyStart = DateTime.UtcNow.ToString();
+            var DependencyStart = DateTime.Now.ToString();
             foreach (var Node in SortedNodes)
             {
                 string Note = GetControllingTriggerDotName(Node);
@@ -6559,8 +6559,8 @@ public class GUBP : BuildCommand
                     Log("  {0}: {1} [Aggregate]", Node, Note);
                 }
             }
-            var DependencyFinish = DateTime.UtcNow.ToString();
-            CommandUtils.StepDurations.Add("GetDependencies", String.Format("{0},{1}", DependencyStart, DependencyFinish));
+            var DependencyFinish = DateTime.Now.ToString();
+            CommandUtils.StepDurations.Add("UAT,GetDependencies", String.Format("{0},{1}", DependencyStart, DependencyFinish));
         }
 		Dictionary<string, int> FullNodeListSortKey = GetDisplayOrder(FullNodeList.Keys.ToList(), FullNodeDirectDependencies, GUBPNodes);
 
@@ -7068,11 +7068,11 @@ public class GUBP : BuildCommand
             
             OrdereredToDo = FilteredOrdereredToDo;
             var FinishFilterTime = DateTime.Now.ToString();
-            StepDurations.Add("FilterNodes", String.Format("{0},{1}", StartFilterTime, FinishFilterTime));
+            StepDurations.Add("UAT,FilterNodes", String.Format("{0},{1}", StartFilterTime, FinishFilterTime));
             Log("*********** EC Nodes, in order.");
             PrintNodes(this, OrdereredToDo, LocalOnly, UnfinishedTriggers);
             var FinishNodePrint = DateTime.Now.ToString();
-            StepDurations.Add("SetupCommanderPrint", String.Format("{0},{1}", FinishFilterTime, FinishNodePrint));
+            StepDurations.Add("UAT,SetupCommanderPrint", String.Format("{0},{1}", FinishFilterTime, FinishNodePrint));
 
             // here we are just making sure everything before the explicit trigger is completed.
             if (!String.IsNullOrEmpty(ExplicitTrigger))
@@ -7440,7 +7440,7 @@ public class GUBP : BuildCommand
             }
             WriteECPerl(StepList);
             var FinishPerlOutput = DateTime.Now.ToString();
-            StepDurations.Add("PerlOutput", String.Format("{0},{1}", StartPerlOutput, FinishPerlOutput));
+            StepDurations.Add("UAT,PerlOutput", String.Format("{0},{1}", StartPerlOutput, FinishPerlOutput));
             RunECTool(String.Format("setProperty \"/myWorkflow/HasTests\" \"{0}\"", bHasTests));
             {
                 ECProps.Add("GUBP_LoadedProps=1");
@@ -7599,10 +7599,13 @@ public class GUBP : BuildCommand
                     {                        
 						Log("***** Building GUBP Node {0} for {1}", NodeToDo, NodeStoreName);
 						var StartTime = DateTime.UtcNow;
-                        var StartString = StartTime.ToString();
+                        var StartString = DateTime.Now.ToString();
                         GUBPNodes[NodeToDo].DoBuild(this);
-                        var FinishBuild = DateTime.UtcNow.ToString();
-                        StepDurations.Add(String.Format("DoBuild{0}", NodeToDo), String.Format("{0},{1}", StartString, FinishBuild));
+                        var FinishBuild = DateTime.Now.ToString();
+                        if (IsBuildMachine)
+                        {
+                            StepDurations.Add("UAT,DoBuild", String.Format("{0},{1}", StartString, FinishBuild));
+                        }
 						BuildDuration = (DateTime.UtcNow - StartTime).TotalMilliseconds / 1000;
 						
                     }
@@ -7610,15 +7613,16 @@ public class GUBP : BuildCommand
                     {
 						var StoreDuration = 0.0;
 						var StartTime = DateTime.UtcNow;
-                        var StoreTime = StartTime.ToString();
+                        var StoreTime = DateTime.Now.ToString();
                         StoreToTempStorage(CmdEnv, NodeStoreName, GUBPNodes[NodeToDo].BuildProducts, !bSaveSharedTempStorage, GameNameIfAny, StorageRootIfAny);
 						StoreDuration = (DateTime.UtcNow - StartTime).TotalMilliseconds / 1000;
-                        var FinishStore = DateTime.UtcNow.ToString();
-                        StepDurations.Add(String.Format("StoreBuildProducts{0}", NodeToDo), String.Format("{0},{1}", StoreTime, FinishStore));
+                        var FinishStore = DateTime.Now.ToString();
+                        
                         //StepDurations.Add("StoreBuildProducts", StoreTime + "," + FinishStore);
 						Log("Took {0} seconds to store build products", StoreDuration);
 						if (IsBuildMachine)
 						{
+                            StepDurations.Add("UAT,StoreBuildProducts", String.Format("{0},{1}", StoreTime, FinishStore));
 							RunECTool(String.Format("setProperty \"/myJobStep/StoreDuration\" \"{0}\"", StoreDuration.ToString()));
 						}
                         if (ParseParam("StompCheck"))
@@ -7628,10 +7632,13 @@ public class GUBP : BuildCommand
                                 try
                                 {
                                     bool WasLocal;
-                                    var StartRetrieve = DateTime.UtcNow.ToString();
+                                    var StartRetrieve = DateTime.Now.ToString();
 									RetrieveFromTempStorage(CmdEnv, NodeStoreName, out WasLocal, GameNameIfAny, StorageRootIfAny);
-                                    var FinishRetrieve = DateTime.UtcNow.ToString();
-                                    StepDurations.Add(String.Format("RetrieveBuildProducts{0}", NodeToDo), String.Format("{0},{1}", StartRetrieve, FinishRetrieve));
+                                    var FinishRetrieve = DateTime.Now.ToString();
+                                    if (IsBuildMachine)
+                                    {
+                                        StepDurations.Add("UAT,RetrieveBuildProducts", String.Format("{0},{1}", StartRetrieve, FinishRetrieve));
+                                    }
 									if (!WasLocal)
 									{
 										throw new AutomationException("Retrieve was not local?");
@@ -7653,19 +7660,20 @@ public class GUBP : BuildCommand
                     {
                         var StartNodeHistory = DateTime.Now.ToString();
                         UpdateNodeHistory(NodeToDo, CLString);
-                        var FinishNodeHistory = DateTime.Now.ToString();
-                        StepDurations.Add(String.Format("UpdateNodeHistory{0}", NodeToDo), String.Format("{0},{1}", StartNodeHistory, FinishNodeHistory));                        
+                        var FinishNodeHistory = DateTime.Now.ToString();                        
                         SaveStatus(NodeToDo, FailedTempStorageSuffix, NodeStoreName, bSaveSharedTempStorage, GameNameIfAny, ParseParamValue("MyJobStepId"));
-                        var FinishSaveStatus = DateTime.Now.ToString();
-                        StepDurations.Add(String.Format("SaveNodeStatus{0}", NodeToDo), String.Format("{0},{1}", FinishNodeHistory, FinishSaveStatus));
+                        var FinishSaveStatus = DateTime.Now.ToString();                        
                         UpdateECProps(NodeToDo, CLString);
                         var FinishECPropUpdate = DateTime.Now.ToString();
-                        StepDurations.Add(String.Format("UpdateECProps{0}", NodeToDo), String.Format("{0},{1}", FinishSaveStatus, FinishECPropUpdate));
+                        
 						if (IsBuildMachine)
 						{
 							GetFailureEmails(NodeToDo, CLString);
                             var FinishFailEmails = DateTime.Now.ToString();
-                            StepDurations.Add(String.Format("GetFailEmails{0}", NodeToDo), String.Format("{0},{1}", FinishECPropUpdate, FinishFailEmails));
+                            StepDurations.Add("UAT,UpdateNodeHistory", String.Format("{0},{1}", StartNodeHistory, FinishNodeHistory));
+                            StepDurations.Add("UAT,SaveNodeStatus", String.Format("{0},{1}", FinishNodeHistory, FinishSaveStatus));
+                            StepDurations.Add("UAT,UpdateECProps", String.Format("{0},{1}", FinishSaveStatus, FinishECPropUpdate));
+                            StepDurations.Add("UAT,GetFailEmails", String.Format("{0},{1}", FinishECPropUpdate, FinishFailEmails));
 						}
 						UpdateECBuildTime(NodeToDo, BuildDuration);
                     }
@@ -7720,18 +7728,19 @@ public class GUBP : BuildCommand
                     var StartNodeHistory = DateTime.Now.ToString();
                     UpdateNodeHistory(NodeToDo, CLString);
                     var FinishNodeHistory = DateTime.Now.ToString();
-                    StepDurations.Add(String.Format("UpdateNodeHistory{0}", NodeToDo), String.Format("{0},{1}", StartNodeHistory, FinishNodeHistory)); 
                     SaveStatus(NodeToDo, SucceededTempStorageSuffix, NodeStoreName, bSaveSharedTempStorage, GameNameIfAny);
-                    var FinishSaveStatus = DateTime.Now.ToString();
-                    StepDurations.Add(String.Format("SaveNodeStatus{0}", NodeToDo), String.Format("{0},{1}", FinishNodeHistory, FinishSaveStatus));
+                    var FinishSaveStatus = DateTime.Now.ToString();                    
                     UpdateECProps(NodeToDo, CLString);
                     var FinishECPropUpdate = DateTime.Now.ToString();
-                    StepDurations.Add(String.Format("UpdateECProps{0}", NodeToDo), String.Format("{0},{1}", FinishSaveStatus, FinishECPropUpdate));
+                    
 					if (IsBuildMachine)
 					{
 						GetFailureEmails(NodeToDo, CLString);
                         var FinishFailEmails = DateTime.Now.ToString();
-                        StepDurations.Add(String.Format("GetFailEmails{0}", NodeToDo), String.Format("{0},{1}", FinishECPropUpdate, FinishFailEmails));
+                        StepDurations.Add("UAT,UpdateNodeHistory", String.Format("{0},{1}", StartNodeHistory, FinishNodeHistory));
+                        StepDurations.Add("UAT,SaveNodeStatus", String.Format("{0},{1}", FinishNodeHistory, FinishSaveStatus));
+                        StepDurations.Add("UAT,UpdateECProps", String.Format("{0},{1}", FinishSaveStatus, FinishECPropUpdate));
+                        StepDurations.Add("UAT,GetFailEmails", String.Format("{0},{1}", FinishECPropUpdate, FinishFailEmails));
 					}
 					UpdateECBuildTime(NodeToDo, BuildDuration);
                 }
