branch hacker, optimistically compile win32 and linux monolithics without waiting for the editor to compile successfully

diff --git a/Engine/Source/Programs/AutomationTool/Scripts/GUBP.Automation.cs b/Engine/Source/Programs/AutomationTool/Scripts/GUBP.Automation.cs
--- a/Engine/Source/Programs/AutomationTool/Scripts/GUBP.Automation.cs
+++ b/Engine/Source/Programs/AutomationTool/Scripts/GUBP.Automation.cs
@@ -202,6 +202,7 @@ public class GUBP : BuildCommand
 			public List<UnrealTargetPlatform> ExcludePlatformsForEditor = new List<UnrealTargetPlatform>();
             public List<UnrealTargetPlatform> RemovePlatformFromPromotable = new List<UnrealTargetPlatform>();
             public List<string> PromotablesWithoutTools = new List<string>();
+            public List<string> NodesToRemovePseudoDependencies = new List<string>();
 			public bool bNoAutomatedTesting = false;
 			public bool bNoDocumentation = false;
 			public bool bNoInstalledEngine = false;
@@ -6761,6 +6762,10 @@ public class GUBP : BuildCommand
                 var Fringe = new HashSet<string>(); 
                 foreach (var NodeToDo in NodesToDo)
                 {
+                    if (BranchOptions.NodesToRemovePseudoDependencies.Contains(NodeToDo))
+                    {
+                        RemoveAllPseudodependenciesFromNode(NodeToDo);
+                    }
                     foreach (var Dep in GUBPNodes[NodeToDo].FullNamesOfDependencies)
                     {
                         if (!GUBPNodes.ContainsKey(Dep))
