Stage missing PS4 x64 dlls (pad/morpheus) and stage all .uplugins

diff --git a/Engine/Source/Programs/AutomationTool/Scripts/CopyBuildToStagingDirectory.Automation.cs b/Engine/Source/Programs/AutomationTool/Scripts/CopyBuildToStagingDirectory.Automation.cs
--- a/Engine/Source/Programs/AutomationTool/Scripts/CopyBuildToStagingDirectory.Automation.cs
+++ b/Engine/Source/Programs/AutomationTool/Scripts/CopyBuildToStagingDirectory.Automation.cs
@@ -362,13 +362,22 @@ public partial class Project : CommandUtils
             {
                 StageLocalizationDataForCulture(SC, Culture, CombinePaths(SC.LocalRoot, "Engine/Content/Localization/Engine"), null, !Params.UsePak(SC.StageTargetPlatform));
             }
-            SC.StageFiles(StagedFileType.UFS, CombinePaths(SC.LocalRoot, "Engine/Plugins"), "*.uplugin", true, null, null, true, !Params.UsePak(SC.StageTargetPlatform));
+
+			// Engine & Game Plugins
+			ProjectDescriptor Project = ProjectDescriptor.FromFile(SC.RawProjectPath);
+			List<PluginInfo> AvailablePlugins = Plugins.ReadAvailablePlugins(SC.RawProjectPath);
+			foreach (PluginInfo Plugin in AvailablePlugins)
+			{
+				if (UProjectInfo.IsPluginEnabledForProject(Plugin, Project, SC.StageTargetPlatform.PlatformType))
+				{
+					SC.StageFiles(StagedFileType.UFS, Plugin.Directory, "*.uplugin", false, null, null, true, !Params.UsePak(SC.StageTargetPlatform), null, true, false);
+				}
+			}
 
             // Game ufs (content)
 
             SC.StageFiles(StagedFileType.UFS, CombinePaths(SC.ProjectRoot), "*.uproject", false, null, CombinePaths(SC.RelativeProjectRootForStage), true, !Params.UsePak(SC.StageTargetPlatform));
             SC.StageFiles(StagedFileType.UFS, CombinePaths(SC.ProjectRoot, "Config"), "*", true, null, CombinePaths(SC.RelativeProjectRootForStage, "Config"), true, !Params.UsePak(SC.StageTargetPlatform)); // TODO: Exclude localization data generation config files.
-            SC.StageFiles(StagedFileType.UFS, CombinePaths(SC.ProjectRoot, "Plugins"), "*.uplugin", true, null, null, true, !Params.UsePak(SC.StageTargetPlatform));
             foreach (string Culture in CulturesToStage)
             {
                 StageLocalizationDataForCulture(SC, Culture, CombinePaths(SC.ProjectRoot, "Content/Localization/Game"), CombinePaths(SC.RelativeProjectRootForStage, "Content/Localization/Game"), !Params.UsePak(SC.StageTargetPlatform));
