build script, fix logging in cases of aggregation

diff --git a/Engine/Source/Programs/AutomationTool/Scripts/GUBP.Automation.cs b/Engine/Source/Programs/AutomationTool/Scripts/GUBP.Automation.cs
--- a/Engine/Source/Programs/AutomationTool/Scripts/GUBP.Automation.cs
+++ b/Engine/Source/Programs/AutomationTool/Scripts/GUBP.Automation.cs
@@ -6561,7 +6561,16 @@ public class GUBP : BuildCommand
                 }
             }
             var DependencyFinish = DateTime.Now.ToString();
-            CommandUtils.StepDurations.Add("UAT,GetDependencies", String.Format("{0},{1}", DependencyStart, DependencyFinish));
+            int NumberofDependencyRuns = 1;
+            if(!StepDurations.ContainsKey("UAT,GetDependencies"))
+            {
+                CommandUtils.StepDurations.Add("UAT,GetDependencies", String.Format("{0},{1}", DependencyStart, DependencyFinish));
+            }
+            else
+            {
+                CommandUtils.StepDurations.Add(String.Format("UAT,GetDependencies{0}", NumberofDependencyRuns.ToString()), String.Format("{0},{1}", DependencyStart, DependencyFinish));
+                NumberofDependencyRuns++;
+            }            
         }
 		Dictionary<string, int> FullNodeListSortKey = GetDisplayOrder(FullNodeList.Keys.ToList(), FullNodeDirectDependencies, GUBPNodes);
 
@@ -7073,12 +7082,29 @@ public class GUBP : BuildCommand
             
             OrdereredToDo = FilteredOrdereredToDo;
             var FinishFilterTime = DateTime.Now.ToString();
-            StepDurations.Add("UAT,FilterNodes", String.Format("{0},{1}", StartFilterTime, FinishFilterTime));
+            int NumberofFilterRuns = 1;
+            if (!StepDurations.ContainsKey("UAT,FilterNodes"))
+            {
+                StepDurations.Add("UAT,FilterNodes", String.Format("{0},{1}", StartFilterTime, FinishFilterTime));
+            }
+            else
+            {
+                CommandUtils.StepDurations.Add(String.Format("UAT,FilterNodes{0}", NumberofFilterRuns.ToString()), String.Format("{0},{1}", StartFilterTime, FinishFilterTime));
+                NumberofFilterRuns++;
+            }      
             Log("*********** EC Nodes, in order.");
             PrintNodes(this, OrdereredToDo, LocalOnly, UnfinishedTriggers);
             var FinishNodePrint = DateTime.Now.ToString();
-            StepDurations.Add("UAT,SetupCommanderPrint", String.Format("{0},{1}", FinishFilterTime, FinishNodePrint));
-
+            int NumberofSetupRuns = 1;
+            if (!StepDurations.ContainsKey("UAT,SetupCommanderPrint"))
+            {
+                StepDurations.Add("UAT,SetupCommanderPrint", String.Format("{0},{1}", FinishFilterTime, FinishNodePrint));
+            }
+            else
+            {
+                CommandUtils.StepDurations.Add(String.Format("UAT,SetupCommanderPrint{0}", NumberofSetupRuns.ToString()), String.Format("{0},{1}", FinishFilterTime, FinishNodePrint));
+                NumberofSetupRuns++;
+            }
             // here we are just making sure everything before the explicit trigger is completed.
             if (!String.IsNullOrEmpty(ExplicitTrigger))
             {
@@ -7609,7 +7635,16 @@ public class GUBP : BuildCommand
                         var FinishBuild = DateTime.Now.ToString();
                         if (IsBuildMachine)
                         {
-                            StepDurations.Add("UAT,DoBuild", String.Format("{0},{1}", StartString, FinishBuild));
+                            int NumberofBuildRuns = 1;
+                            if (!StepDurations.ContainsKey("UAT,DoBuild"))
+                            {
+                                StepDurations.Add("UAT,DoBuild", String.Format("{0},{1}", StartString, FinishBuild));
+                            }
+                            else
+                            {
+                                CommandUtils.StepDurations.Add(String.Format("UAT,DoBuild{0}", NumberofBuildRuns.ToString()), String.Format("{0},{1}", StartString, FinishBuild));
+                                NumberofBuildRuns++;
+                            }                  
                         }
 						BuildDuration = (DateTime.UtcNow - StartTime).TotalMilliseconds / 1000;
 						
@@ -7627,7 +7662,16 @@ public class GUBP : BuildCommand
 						Log("Took {0} seconds to store build products", StoreDuration);
 						if (IsBuildMachine)
 						{
-                            StepDurations.Add("UAT,StoreBuildProducts", String.Format("{0},{1}", StoreTime, FinishStore));
+                            int NumberofStoreRuns = 1;
+                            if (!StepDurations.ContainsKey("UAT,StoreBuildProducts"))
+                            {
+                                StepDurations.Add("UAT,StoreBuildProducts", String.Format("{0},{1}", StoreTime, FinishStore));
+                            }
+                            else
+                            {
+                                CommandUtils.StepDurations.Add(String.Format("UAT,StoreBuildProducts{0}", NumberofStoreRuns.ToString()), String.Format("{0},{1}", StoreTime, FinishStore));
+                                NumberofStoreRuns++;
+                            }
 							RunECTool(String.Format("setProperty \"/myJobStep/StoreDuration\" \"{0}\"", StoreDuration.ToString()));
 						}
                         if (ParseParam("StompCheck"))
@@ -7642,7 +7686,16 @@ public class GUBP : BuildCommand
                                     var FinishRetrieve = DateTime.Now.ToString();
                                     if (IsBuildMachine)
                                     {
-                                        StepDurations.Add("UAT,RetrieveBuildProducts", String.Format("{0},{1}", StartRetrieve, FinishRetrieve));
+                                        int NumberofRetrieveRuns = 1;
+                                        if (!StepDurations.ContainsKey("UAT,RetrieveBuildProducts"))
+                                        {
+                                            StepDurations.Add("UAT,RetrieveBuildProducts", String.Format("{0},{1}", StartRetrieve, FinishRetrieve));
+                                        }
+                                        else
+                                        {
+                                            CommandUtils.StepDurations.Add(String.Format("UAT,RetrieveBuildProducts{0}", NumberofRetrieveRuns.ToString()), String.Format("{0},{1}", StartRetrieve, FinishRetrieve));
+                                            NumberofRetrieveRuns++;
+                                        }                                        
                                     }
 									if (!WasLocal)
 									{
@@ -7675,10 +7728,46 @@ public class GUBP : BuildCommand
 						{
 							GetFailureEmails(NodeToDo, CLString);
                             var FinishFailEmails = DateTime.Now.ToString();
-                            StepDurations.Add("UAT,UpdateNodeHistory", String.Format("{0},{1}", StartNodeHistory, FinishNodeHistory));
-                            StepDurations.Add("UAT,SaveNodeStatus", String.Format("{0},{1}", FinishNodeHistory, FinishSaveStatus));
-                            StepDurations.Add("UAT,UpdateECProps", String.Format("{0},{1}", FinishSaveStatus, FinishECPropUpdate));
-                            StepDurations.Add("UAT,GetFailEmails", String.Format("{0},{1}", FinishECPropUpdate, FinishFailEmails));
+                            int NumberofUpdateRuns = 1;
+                            if (!StepDurations.ContainsKey("UAT,UpdateNodeHistory"))
+                            {
+                                StepDurations.Add("UAT,UpdateNodeHistory", String.Format("{0},{1}", StartNodeHistory, FinishNodeHistory));
+                            }
+                            else
+                            {
+                                CommandUtils.StepDurations.Add(String.Format("UAT,UpdateNodeHistory{0}", NumberofUpdateRuns.ToString()), String.Format("{0},{1}", StartNodeHistory, FinishNodeHistory));
+                                NumberofUpdateRuns++;
+                            }
+                            int NumberofSaveRuns = 1;
+                            if (!StepDurations.ContainsKey("UAT,SaveNodeStatus"))
+                            {
+                                StepDurations.Add("UAT,SaveNodeStatus", String.Format("{0},{1}", FinishNodeHistory, FinishSaveStatus));
+                            }
+                            else
+                            {
+                                CommandUtils.StepDurations.Add(String.Format("UAT,SaveNodeStatus{0}", NumberofSaveRuns.ToString()), String.Format("{0},{1}", FinishNodeHistory, FinishSaveStatus));
+                                NumberofSaveRuns++;
+                            }
+                            int NumberofECSaveRuns = 1;
+                            if (!StepDurations.ContainsKey("UAT,UpdateECProps"))
+                            {
+                                StepDurations.Add("UAT,UpdateECProps", String.Format("{0},{1}", FinishSaveStatus, FinishECPropUpdate));
+                            }
+                            else
+                            {
+                                CommandUtils.StepDurations.Add(String.Format("UAT,UpdateECProps{0}", NumberofECSaveRuns.ToString()), String.Format("{0},{1}", FinishSaveStatus, FinishECPropUpdate));
+                                NumberofECSaveRuns++;
+                            }
+                            int NumberofEmailRuns = 1;
+                            if (!StepDurations.ContainsKey("UAT,GetFailEmails"))
+                            {
+                                StepDurations.Add("UAT,GetFailEmails", String.Format("{0},{1}", FinishECPropUpdate, FinishFailEmails));
+                            }
+                            else
+                            {
+                                CommandUtils.StepDurations.Add(String.Format("UAT,GetFailEmails{0}", NumberofEmailRuns.ToString()), String.Format("{0},{1}", FinishECPropUpdate, FinishFailEmails));
+                                NumberofEmailRuns++;
+                            }                            
 						}
 						UpdateECBuildTime(NodeToDo, BuildDuration);
                     }
@@ -7742,10 +7831,46 @@ public class GUBP : BuildCommand
 					{
 						GetFailureEmails(NodeToDo, CLString);
                         var FinishFailEmails = DateTime.Now.ToString();
-                        StepDurations.Add("UAT,UpdateNodeHistory", String.Format("{0},{1}", StartNodeHistory, FinishNodeHistory));
-                        StepDurations.Add("UAT,SaveNodeStatus", String.Format("{0},{1}", FinishNodeHistory, FinishSaveStatus));
-                        StepDurations.Add("UAT,UpdateECProps", String.Format("{0},{1}", FinishSaveStatus, FinishECPropUpdate));
-                        StepDurations.Add("UAT,GetFailEmails", String.Format("{0},{1}", FinishECPropUpdate, FinishFailEmails));
+                        int NumberofUpdateRuns = 1;
+                        if (!StepDurations.ContainsKey("UAT,UpdateNodeHistory"))
+                        {
+                            StepDurations.Add("UAT,UpdateNodeHistory", String.Format("{0},{1}", StartNodeHistory, FinishNodeHistory));
+                        }
+                        else
+                        {
+                            CommandUtils.StepDurations.Add(String.Format("UAT,UpdateNodeHistory{0}", NumberofUpdateRuns.ToString()), String.Format("{0},{1}", StartNodeHistory, FinishNodeHistory));
+                            NumberofUpdateRuns++;
+                        }
+                        int NumberofSaveRuns = 1;
+                        if (!StepDurations.ContainsKey("UAT,SaveNodeStatus"))
+                        {
+                            StepDurations.Add("UAT,SaveNodeStatus", String.Format("{0},{1}", FinishNodeHistory, FinishSaveStatus));
+                        }
+                        else
+                        {
+                            CommandUtils.StepDurations.Add(String.Format("UAT,SaveNodeStatus{0}", NumberofSaveRuns.ToString()), String.Format("{0},{1}", FinishNodeHistory, FinishSaveStatus));
+                            NumberofSaveRuns++;
+                        }
+                        int NumberofECSaveRuns = 1;
+                        if (!StepDurations.ContainsKey("UAT,UpdateECProps"))
+                        {
+                            StepDurations.Add("UAT,UpdateECProps", String.Format("{0},{1}", FinishSaveStatus, FinishECPropUpdate));
+                        }
+                        else
+                        {
+                            CommandUtils.StepDurations.Add(String.Format("UAT,UpdateECProps{0}", NumberofECSaveRuns.ToString()), String.Format("{0},{1}", FinishSaveStatus, FinishECPropUpdate));
+                            NumberofECSaveRuns++;
+                        }
+                        int NumberofEmailRuns = 1;
+                        if (!StepDurations.ContainsKey("UAT,GetFailEmails"))
+                        {
+                            StepDurations.Add("UAT,GetFailEmails", String.Format("{0},{1}", FinishECPropUpdate, FinishFailEmails));
+                        }
+                        else
+                        {
+                            CommandUtils.StepDurations.Add(String.Format("UAT,GetFailEmails{0}", NumberofEmailRuns.ToString()), String.Format("{0},{1}", FinishECPropUpdate, FinishFailEmails));
+                            NumberofEmailRuns++;
+                        }  
 					}
 					UpdateECBuildTime(NodeToDo, BuildDuration);
                 }
