Fixes for missing Unicode characters during packaging:

* Passing -UTF8Output during commandlet execution.
* Console encoding set before log initialization.
* wprintf used instead of printf.

Missing specifiers added to some UAT build logs.

diff --git a/Engine/Source/Programs/AutomationTool/CommandletUtils.cs b/Engine/Source/Programs/AutomationTool/CommandletUtils.cs
--- a/Engine/Source/Programs/AutomationTool/CommandletUtils.cs
+++ b/Engine/Source/Programs/AutomationTool/CommandletUtils.cs
@@ -165,13 +165,21 @@ namespace AutomationTool
 
 			string LocalLogFile = LogUtils.GetUniqueLogName(CombinePaths(CmdEnv.EngineSavedFolder, Commandlet));
 			Log("Commandlet log file is {0}", LocalLogFile);
-			var RunResult = Run(EditorExe, String.Format("{0} -run={1} {2} -abslog={3} -stdout -FORCELOGFLUSH -CrashForUAT -unattended -AllowStdOutLogVerbosity {4}", 
-                (ProjectName == null)? "" : CommandUtils.MakePathSafeToUseWithCommandLine(ProjectName), 
-                Commandlet, 
-				String.IsNullOrEmpty(Parameters) ? "" : Parameters, 
+			string Args = String.Format(
+				"{0} -run={1} {2} -abslog={3} -stdout -FORCELOGFLUSH -CrashForUAT -unattended -AllowStdOutLogVerbosity {4}",
+				(ProjectName == null) ? "" : CommandUtils.MakePathSafeToUseWithCommandLine(ProjectName),
+				Commandlet,
+				String.IsNullOrEmpty(Parameters) ? "" : Parameters,
 				CommandUtils.MakePathSafeToUseWithCommandLine(LocalLogFile),
-                IsBuildMachine ? "-buildmachine" : ""
-                ));
+				IsBuildMachine ? "-buildmachine" : ""
+			);
+			ERunOptions Opts = ERunOptions.Default;
+			if (GlobalCommandLine.UTF8Output)
+			{
+				Args += " -UTF8Output";
+				Opts |= ERunOptions.UTF8Output;
+			}
+			var RunResult = Run(EditorExe, Args, Options: Opts);
 			PopDir();
 
 			// Copy the local commandlet log to the destination folder.
diff --git a/Engine/Source/Programs/AutomationTool/Log.cs b/Engine/Source/Programs/AutomationTool/Log.cs
--- a/Engine/Source/Programs/AutomationTool/Log.cs
+++ b/Engine/Source/Programs/AutomationTool/Log.cs
@@ -22,7 +22,13 @@ namespace AutomationTool
 		/// <param name="CommandLine">Command line.</param>
 		public static void InitLogging(string[] CommandLine)
 		{
-            UnrealBuildTool.Log.InitLogging(
+            // ensure UTF8Output flag is respected, since we are initializing logging early in the program.
+            if (CommandLine.Any(Arg => Arg.Equals("-utf8output", StringComparison.InvariantCultureIgnoreCase)))
+            {
+                Console.OutputEncoding = new System.Text.UTF8Encoding(false, false);
+            }
+
+			UnrealBuildTool.Log.InitLogging(
                 bLogTimestamps: CommandUtils.ParseParam(CommandLine, "-Timestamps"),
                 bLogVerbose: CommandUtils.ParseParam(CommandLine, "-Verbose"),
                 bLogSeverity: true,
@@ -36,12 +42,6 @@ namespace AutomationTool
                     //@todo - this is only used by GUBP nodes. Ideally we don't waste this 20MB if we are not running GUBP.
                     new AutomationMemoryLogListener(),
                 });
-
-            // ensure UTF8Output flag is respected, since we are initializing logging early in the program.
-            if (CommandLine.Any(Arg => Arg.Equals("-utf8output", StringComparison.InvariantCultureIgnoreCase)))
-            {
-                Console.OutputEncoding = new System.Text.UTF8Encoding(false, false);
-            }
         }
 
         public static void ShutdownLogging()
diff --git a/Engine/Source/Programs/AutomationTool/UE4Build.cs b/Engine/Source/Programs/AutomationTool/UE4Build.cs
--- a/Engine/Source/Programs/AutomationTool/UE4Build.cs
+++ b/Engine/Source/Programs/AutomationTool/UE4Build.cs
@@ -1267,10 +1267,10 @@ namespace AutomationTool
 			}
 
 			Log("************************* UE4Build:");
-			Log("************************* ForceMonolithic: ", bForceMonolithic);
-			Log("************************* ForceNonUnity: ", bForceNonUnity);
-			Log("************************* ForceDebugInfo: ", bForceDebugInfo);
-			Log("************************* UseXGE: ", bCanUseXGE);
+			Log("************************* ForceMonolithic: {0}", bForceMonolithic);
+			Log("************************* ForceNonUnity:{0} ", bForceNonUnity);
+			Log("************************* ForceDebugInfo: {0}", bForceDebugInfo);
+			Log("************************* UseXGE: {0}", bCanUseXGE);
 
 			if (bCanUseXGE)
 			{
diff --git a/Engine/Source/Runtime/Launch/Private/LaunchEngineLoop.cpp b/Engine/Source/Runtime/Launch/Private/LaunchEngineLoop.cpp
--- a/Engine/Source/Runtime/Launch/Private/LaunchEngineLoop.cpp
+++ b/Engine/Source/Runtime/Launch/Private/LaunchEngineLoop.cpp
@@ -118,7 +118,7 @@ public:
 		if ( (bAllowLogVerbosity && Verbosity <= ELogVerbosity::Log) || (Verbosity <= ELogVerbosity::Display) )
 		{
 #if PLATFORM_USE_LS_SPEC_FOR_WIDECHAR
-			printf("\n%ls", *FOutputDevice::FormatLogLine(Verbosity, Category, V, GPrintLogTimes));
+			wprintf(TEXT("\n%ls"), *FOutputDevice::FormatLogLine(Verbosity, Category, V, GPrintLogTimes));
 #else
 			wprintf(TEXT("\n%s"), *FOutputDevice::FormatLogLine(Verbosity, Category, V, GPrintLogTimes));
 #endif
