editors run in parallel, allowing for preconditions to be met before they complete

diff --git a/Engine/Source/Programs/AutomationTool/Scripts/GUBP.Automation.cs b/Engine/Source/Programs/AutomationTool/Scripts/GUBP.Automation.cs
--- a/Engine/Source/Programs/AutomationTool/Scripts/GUBP.Automation.cs
+++ b/Engine/Source/Programs/AutomationTool/Scripts/GUBP.Automation.cs
@@ -7192,6 +7192,10 @@ public class GUBP : BuildCommand
 
                     bool Sticky = GUBPNodes[NodeToDo].IsSticky();
                     bool DoParallel = !Sticky;
+                    if(GUBPNodes[NodeToDo].ECProcedure() == "GUBP_UAT_Node_Parallel_AgentShare_Editor")
+                    {
+                        DoParallel = true;
+                    }
                     if (Sticky && GUBPNodes[NodeToDo].ECAgentString() != "")
                     {
                         throw new AutomationException("Node {1} is sticky but has agent requirements.", NodeToDo);
