build script, logging in a better way to slurp the data clarifying variables

diff --git a/Engine/Source/Programs/AutomationTool/Automation.cs b/Engine/Source/Programs/AutomationTool/Automation.cs
--- a/Engine/Source/Programs/AutomationTool/Automation.cs
+++ b/Engine/Source/Programs/AutomationTool/Automation.cs
@@ -322,7 +322,7 @@ AutomationTool.exe [-verbose] [-compileonly] [-p4] Command0 [-Arg0 -Arg1 -Arg2 
             var StartScriptCompile = DateTime.Now.ToString();
 			Compiler.FindAndCompileAllScripts(AdditionalScriptsFolders: AdditionalScriptsFolders);
             var FinishScriptCompile = DateTime.Now.ToString();
-            CommandUtils.StepDurations.Add("UAT,ScriptCompile", String.Format("{0},{1}", StartScriptCompile, FinishScriptCompile));
+            CommandUtils.PrintCSVFile(String.Format("UAT,ScriptCompile,{0},{1}", StartScriptCompile, FinishScriptCompile));
 
 			if (GlobalCommandLine.CompileOnly)
 			{
diff --git a/Engine/Source/Programs/AutomationTool/CommandUtils.cs b/Engine/Source/Programs/AutomationTool/CommandUtils.cs
--- a/Engine/Source/Programs/AutomationTool/CommandUtils.cs
+++ b/Engine/Source/Programs/AutomationTool/CommandUtils.cs
@@ -1876,9 +1876,17 @@ namespace AutomationTool
 				return UBTTargetPlatforms;
 			}
 		}
-		private static UnrealBuildTool.UnrealTargetPlatform[] UBTTargetPlatforms;
-
-        public static Dictionary<string, string> StepDurations = new Dictionary<string,string>();
+		private static UnrealBuildTool.UnrealTargetPlatform[] UBTTargetPlatforms; 
+        public static void PrintCSVFile(string Input)
+        {
+            if (IsBuildMachine && CmdEnv.CSVFile != "")
+            {
+                var CSVBuilder = new StringBuilder();                
+                var CSVLineToAppend = String.Format("{0}{2}", Input, Environment.NewLine);
+                CSVBuilder.Append(CSVLineToAppend);                
+                File.AppendAllText(CmdEnv.CSVFile, CSVBuilder.ToString());
+            }
+        }
 
 		#endregion
 
diff --git a/Engine/Source/Programs/AutomationTool/ProjectUtils.cs b/Engine/Source/Programs/AutomationTool/ProjectUtils.cs
--- a/Engine/Source/Programs/AutomationTool/ProjectUtils.cs
+++ b/Engine/Source/Programs/AutomationTool/ProjectUtils.cs
@@ -726,7 +726,7 @@ namespace AutomationTool
                 }
             }
             var FinishTime = DateTime.Now.ToString();
-            CommandUtils.StepDurations.Add("UAT,SortProjects", String.Format("{0},{1}", StartTime, FinishTime));
+            CommandUtils.PrintCSVFile(String.Format("UAT,SortProjects,{0},{1}", StartTime, FinishTime));            
             if (String.IsNullOrEmpty(BaseEngineProject.FilePath))
             {
                 throw new AutomationException("All branches must have the blank project /Samples/Sandbox/BlankProject");
@@ -746,8 +746,8 @@ namespace AutomationTool
             {
                 Proj.Dump(InHostPlatforms);
             }
-            var BuildDumpFinish = DateTime.Now.ToString();
-            CommandUtils.StepDurations.Add("UAT,Project Dump", String.Format("{0},{1}", StartBranchDump, BuildDumpFinish));
+            var FinishBranchDump = DateTime.Now.ToString();
+            CommandUtils.PrintCSVFile(String.Format("UAT,Project Dump,{0},{1}", StartBranchDump, FinishBranchDump));
         }
 
         public BranchUProject FindGame(string GameName)
diff --git a/Engine/Source/Programs/AutomationTool/Scripts/GUBP.Automation.cs b/Engine/Source/Programs/AutomationTool/Scripts/GUBP.Automation.cs
--- a/Engine/Source/Programs/AutomationTool/Scripts/GUBP.Automation.cs
+++ b/Engine/Source/Programs/AutomationTool/Scripts/GUBP.Automation.cs
@@ -6561,16 +6561,7 @@ public class GUBP : BuildCommand
                 }
             }
             var DependencyFinish = DateTime.Now.ToString();
-            int NumberofDependencyRuns = 1;
-            if(!StepDurations.ContainsKey("UAT,GetDependencies"))
-            {
-                CommandUtils.StepDurations.Add("UAT,GetDependencies", String.Format("{0},{1}", DependencyStart, DependencyFinish));
-            }
-            else
-            {
-                CommandUtils.StepDurations.Add(String.Format("UAT,GetDependencies{0}", NumberofDependencyRuns.ToString()), String.Format("{0},{1}", DependencyStart, DependencyFinish));
-                NumberofDependencyRuns++;
-            }            
+            PrintCSVFile(String.Format("UAT,GetDependencies,{0},{1}", DependencyStart, DependencyFinish));
         }
 		Dictionary<string, int> FullNodeListSortKey = GetDisplayOrder(FullNodeList.Keys.ToList(), FullNodeDirectDependencies, GUBPNodes);
 
@@ -7082,29 +7073,11 @@ public class GUBP : BuildCommand
             
             OrdereredToDo = FilteredOrdereredToDo;
             var FinishFilterTime = DateTime.Now.ToString();
-            int NumberofFilterRuns = 1;
-            if (!StepDurations.ContainsKey("UAT,FilterNodes"))
-            {
-                StepDurations.Add("UAT,FilterNodes", String.Format("{0},{1}", StartFilterTime, FinishFilterTime));
-            }
-            else
-            {
-                CommandUtils.StepDurations.Add(String.Format("UAT,FilterNodes{0}", NumberofFilterRuns.ToString()), String.Format("{0},{1}", StartFilterTime, FinishFilterTime));
-                NumberofFilterRuns++;
-            }      
+            PrintCSVFile(String.Format("UAT,FilterNodes,{0},{1}", StartFilterTime, FinishFilterTime));            
             Log("*********** EC Nodes, in order.");
             PrintNodes(this, OrdereredToDo, LocalOnly, UnfinishedTriggers);
             var FinishNodePrint = DateTime.Now.ToString();
-            int NumberofSetupRuns = 1;
-            if (!StepDurations.ContainsKey("UAT,SetupCommanderPrint"))
-            {
-                StepDurations.Add("UAT,SetupCommanderPrint", String.Format("{0},{1}", FinishFilterTime, FinishNodePrint));
-            }
-            else
-            {
-                CommandUtils.StepDurations.Add(String.Format("UAT,SetupCommanderPrint{0}", NumberofSetupRuns.ToString()), String.Format("{0},{1}", FinishFilterTime, FinishNodePrint));
-                NumberofSetupRuns++;
-            }
+            PrintCSVFile(String.Format("UAT,SetupCommanderPrint,{0},{1}", FinishFilterTime, FinishNodePrint));
             // here we are just making sure everything before the explicit trigger is completed.
             if (!String.IsNullOrEmpty(ExplicitTrigger))
             {
@@ -7475,7 +7448,7 @@ public class GUBP : BuildCommand
             }
             WriteECPerl(StepList);
             var FinishPerlOutput = DateTime.Now.ToString();
-            StepDurations.Add("UAT,PerlOutput", String.Format("{0},{1}", StartPerlOutput, FinishPerlOutput));
+            PrintCSVFile(String.Format("UAT,PerlOutput,{0},{1}", StartPerlOutput, FinishPerlOutput));
             RunECTool(String.Format("setProperty \"/myWorkflow/HasTests\" \"{0}\"", bHasTests));
             {
                 ECProps.Add("GUBP_LoadedProps=1");
@@ -7496,8 +7469,7 @@ public class GUBP : BuildCommand
                     RunUAT(CmdEnv, Args);
                 }
             }
-            Log("Commander setup only, done.");
-            PrintCSVFile();
+            Log("Commander setup only, done.");            
             PrintRunTime();
             return;
 
@@ -7634,21 +7606,12 @@ public class GUBP : BuildCommand
                     {                        
 						Log("***** Building GUBP Node {0} for {1}", NodeToDo, NodeStoreName);
 						var StartTime = DateTime.UtcNow;
-                        var StartString = DateTime.Now.ToString();
+                        var StartBuild = DateTime.Now.ToString();
                         GUBPNodes[NodeToDo].DoBuild(this);
                         var FinishBuild = DateTime.Now.ToString();
                         if (IsBuildMachine)
                         {
-                            int NumberofBuildRuns = 1;
-                            if (!StepDurations.ContainsKey("UAT,DoBuild"))
-                            {
-                                StepDurations.Add("UAT,DoBuild", String.Format("{0},{1}", StartString, FinishBuild));
-                            }
-                            else
-                            {
-                                CommandUtils.StepDurations.Add(String.Format("UAT,DoBuild{0}", NumberofBuildRuns.ToString()), String.Format("{0},{1}", StartString, FinishBuild));
-                                NumberofBuildRuns++;
-                            }                  
+                            PrintCSVFile(String.Format("UAT,DoBuild,{0},{1}", StartBuild, FinishBuild));
                         }
 						BuildDuration = (DateTime.UtcNow - StartTime).TotalMilliseconds / 1000;
 						
@@ -7657,25 +7620,14 @@ public class GUBP : BuildCommand
                     {
 						var StoreDuration = 0.0;
 						var StartTime = DateTime.UtcNow;
-                        var StoreTime = DateTime.Now.ToString();
+                        var StartStore = DateTime.Now.ToString();
                         StoreToTempStorage(CmdEnv, NodeStoreName, GUBPNodes[NodeToDo].BuildProducts, !bSaveSharedTempStorage, GameNameIfAny, StorageRootIfAny);
 						StoreDuration = (DateTime.UtcNow - StartTime).TotalMilliseconds / 1000;
-                        var FinishStore = DateTime.Now.ToString();
-                        
-                        //StepDurations.Add("StoreBuildProducts", StoreTime + "," + FinishStore);
+                        var FinishStore = DateTime.Now.ToString();                        
 						Log("Took {0} seconds to store build products", StoreDuration);
 						if (IsBuildMachine)
 						{
-                            int NumberofStoreRuns = 1;
-                            if (!StepDurations.ContainsKey("UAT,StoreBuildProducts"))
-                            {
-                                StepDurations.Add("UAT,StoreBuildProducts", String.Format("{0},{1}", StoreTime, FinishStore));
-                            }
-                            else
-                            {
-                                CommandUtils.StepDurations.Add(String.Format("UAT,StoreBuildProducts{0}", NumberofStoreRuns.ToString()), String.Format("{0},{1}", StoreTime, FinishStore));
-                                NumberofStoreRuns++;
-                            }
+                            PrintCSVFile(String.Format("UAT,StoreBuildProducts,{0},{1}", StartStore, FinishStore));
 							RunECTool(String.Format("setProperty \"/myJobStep/StoreDuration\" \"{0}\"", StoreDuration.ToString()));
 						}
                         if (ParseParam("StompCheck"))
@@ -7690,16 +7642,7 @@ public class GUBP : BuildCommand
                                     var FinishRetrieve = DateTime.Now.ToString();
                                     if (IsBuildMachine)
                                     {
-                                        int NumberofRetrieveRuns = 1;
-                                        if (!StepDurations.ContainsKey("UAT,RetrieveBuildProducts"))
-                                        {
-                                            StepDurations.Add("UAT,RetrieveBuildProducts", String.Format("{0},{1}", StartRetrieve, FinishRetrieve));
-                                        }
-                                        else
-                                        {
-                                            CommandUtils.StepDurations.Add(String.Format("UAT,RetrieveBuildProducts{0}", NumberofRetrieveRuns.ToString()), String.Format("{0},{1}", StartRetrieve, FinishRetrieve));
-                                            NumberofRetrieveRuns++;
-                                        }                                        
+                                        PrintCSVFile(String.Format("UAT,RetrieveBuildProducts,{0},{1}", StartRetrieve, FinishRetrieve));
                                     }
 									if (!WasLocal)
 									{
@@ -7732,46 +7675,10 @@ public class GUBP : BuildCommand
 						{
 							GetFailureEmails(NodeToDo, CLString);
                             var FinishFailEmails = DateTime.Now.ToString();
-                            int NumberofUpdateRuns = 1;
-                            if (!StepDurations.ContainsKey("UAT,UpdateNodeHistory"))
-                            {
-                                StepDurations.Add("UAT,UpdateNodeHistory", String.Format("{0},{1}", StartNodeHistory, FinishNodeHistory));
-                            }
-                            else
-                            {
-                                CommandUtils.StepDurations.Add(String.Format("UAT,UpdateNodeHistory{0}", NumberofUpdateRuns.ToString()), String.Format("{0},{1}", StartNodeHistory, FinishNodeHistory));
-                                NumberofUpdateRuns++;
-                            }
-                            int NumberofSaveRuns = 1;
-                            if (!StepDurations.ContainsKey("UAT,SaveNodeStatus"))
-                            {
-                                StepDurations.Add("UAT,SaveNodeStatus", String.Format("{0},{1}", FinishNodeHistory, FinishSaveStatus));
-                            }
-                            else
-                            {
-                                CommandUtils.StepDurations.Add(String.Format("UAT,SaveNodeStatus{0}", NumberofSaveRuns.ToString()), String.Format("{0},{1}", FinishNodeHistory, FinishSaveStatus));
-                                NumberofSaveRuns++;
-                            }
-                            int NumberofECSaveRuns = 1;
-                            if (!StepDurations.ContainsKey("UAT,UpdateECProps"))
-                            {
-                                StepDurations.Add("UAT,UpdateECProps", String.Format("{0},{1}", FinishSaveStatus, FinishECPropUpdate));
-                            }
-                            else
-                            {
-                                CommandUtils.StepDurations.Add(String.Format("UAT,UpdateECProps{0}", NumberofECSaveRuns.ToString()), String.Format("{0},{1}", FinishSaveStatus, FinishECPropUpdate));
-                                NumberofECSaveRuns++;
-                            }
-                            int NumberofEmailRuns = 1;
-                            if (!StepDurations.ContainsKey("UAT,GetFailEmails"))
-                            {
-                                StepDurations.Add("UAT,GetFailEmails", String.Format("{0},{1}", FinishECPropUpdate, FinishFailEmails));
-                            }
-                            else
-                            {
-                                CommandUtils.StepDurations.Add(String.Format("UAT,GetFailEmails{0}", NumberofEmailRuns.ToString()), String.Format("{0},{1}", FinishECPropUpdate, FinishFailEmails));
-                                NumberofEmailRuns++;
-                            }                            
+                            PrintCSVFile(String.Format("UAT,UpdateNodeHistory,{0},{1}", StartNodeHistory, FinishNodeHistory));
+                            PrintCSVFile(String.Format("UAT,SaveNodeStatus,{0},{1}", FinishNodeHistory, FinishSaveStatus));
+                            PrintCSVFile(String.Format("UAT,UpdateECProps,{0},{1}", FinishSaveStatus, FinishECPropUpdate));
+                            PrintCSVFile(String.Format("UAT,GetFailEmails,{0},{1}", FinishECPropUpdate, FinishFailEmails));
 						}
 						UpdateECBuildTime(NodeToDo, BuildDuration);
                     }
@@ -7835,46 +7742,10 @@ public class GUBP : BuildCommand
 					{
 						GetFailureEmails(NodeToDo, CLString);
                         var FinishFailEmails = DateTime.Now.ToString();
-                        int NumberofUpdateRuns = 1;
-                        if (!StepDurations.ContainsKey("UAT,UpdateNodeHistory"))
-                        {
-                            StepDurations.Add("UAT,UpdateNodeHistory", String.Format("{0},{1}", StartNodeHistory, FinishNodeHistory));
-                        }
-                        else
-                        {
-                            CommandUtils.StepDurations.Add(String.Format("UAT,UpdateNodeHistory{0}", NumberofUpdateRuns.ToString()), String.Format("{0},{1}", StartNodeHistory, FinishNodeHistory));
-                            NumberofUpdateRuns++;
-                        }
-                        int NumberofSaveRuns = 1;
-                        if (!StepDurations.ContainsKey("UAT,SaveNodeStatus"))
-                        {
-                            StepDurations.Add("UAT,SaveNodeStatus", String.Format("{0},{1}", FinishNodeHistory, FinishSaveStatus));
-                        }
-                        else
-                        {
-                            CommandUtils.StepDurations.Add(String.Format("UAT,SaveNodeStatus{0}", NumberofSaveRuns.ToString()), String.Format("{0},{1}", FinishNodeHistory, FinishSaveStatus));
-                            NumberofSaveRuns++;
-                        }
-                        int NumberofECSaveRuns = 1;
-                        if (!StepDurations.ContainsKey("UAT,UpdateECProps"))
-                        {
-                            StepDurations.Add("UAT,UpdateECProps", String.Format("{0},{1}", FinishSaveStatus, FinishECPropUpdate));
-                        }
-                        else
-                        {
-                            CommandUtils.StepDurations.Add(String.Format("UAT,UpdateECProps{0}", NumberofECSaveRuns.ToString()), String.Format("{0},{1}", FinishSaveStatus, FinishECPropUpdate));
-                            NumberofECSaveRuns++;
-                        }
-                        int NumberofEmailRuns = 1;
-                        if (!StepDurations.ContainsKey("UAT,GetFailEmails"))
-                        {
-                            StepDurations.Add("UAT,GetFailEmails", String.Format("{0},{1}", FinishECPropUpdate, FinishFailEmails));
-                        }
-                        else
-                        {
-                            CommandUtils.StepDurations.Add(String.Format("UAT,GetFailEmails{0}", NumberofEmailRuns.ToString()), String.Format("{0},{1}", FinishECPropUpdate, FinishFailEmails));
-                            NumberofEmailRuns++;
-                        }  
+                        PrintCSVFile(String.Format("UAT,UpdateNodeHistory,{0},{1}", StartNodeHistory, FinishNodeHistory));
+                        PrintCSVFile(String.Format("UAT,SaveNodeStatus,{0},{1}", FinishNodeHistory, FinishSaveStatus));
+                        PrintCSVFile(String.Format("UAT,UpdateECProps,{0},{1}", FinishSaveStatus, FinishECPropUpdate));
+                        PrintCSVFile(String.Format("UAT,GetFailEmails,{0},{1}", FinishECPropUpdate, FinishFailEmails));
 					}
 					UpdateECBuildTime(NodeToDo, BuildDuration);
                 }
@@ -7887,25 +7758,10 @@ public class GUBP : BuildCommand
                 }
                 BuildProductToNodeMap.Add(Product, NodeToDo);
             }
-        }
-
-        PrintCSVFile();
+        }        
         PrintRunTime();
     }
-
-    static void PrintCSVFile()
-    {
-        if (IsBuildMachine && CmdEnv.CSVFile != "")
-        {
-            var CSVBuilder = new StringBuilder();
-            foreach (KeyValuePair<string, string> Step in StepDurations)
-            {
-                var CSVLineToAppend = String.Format("{0},{1}{2}", Step.Key, Step.Value, Environment.NewLine);
-                CSVBuilder.Append(CSVLineToAppend);
-            }
-            File.AppendAllText(CmdEnv.CSVFile, CSVBuilder.ToString());
-        }
-    }
+    
 	/// <summary>
 	/// Sorts a list of nodes to display in EC. The default order is based on execution order and agent groups, whereas this function arranges nodes by
 	/// frequency then execution order, while trying to group nodes on parallel paths (eg. Mac/Windows editor nodes) together.
