build script, cis fix

diff --git a/Engine/Source/Programs/AutomationTool/CommandUtils.cs b/Engine/Source/Programs/AutomationTool/CommandUtils.cs
--- a/Engine/Source/Programs/AutomationTool/CommandUtils.cs
+++ b/Engine/Source/Programs/AutomationTool/CommandUtils.cs
@@ -1882,7 +1882,7 @@ namespace AutomationTool
             if (IsBuildMachine && CmdEnv.CSVFile != "")
             {
                 var CSVBuilder = new StringBuilder();                
-                var CSVLineToAppend = String.Format("{0}{2}", Input, Environment.NewLine);
+                var CSVLineToAppend = String.Format("{0}{1}", Input, Environment.NewLine);
                 CSVBuilder.Append(CSVLineToAppend);                
                 File.AppendAllText(CmdEnv.CSVFile, CSVBuilder.ToString());
             }
