build script, add concept of preconditions to sticky nodes that are not necessarily dependent on the stick node above it

diff --git a/Engine/Source/Programs/AutomationTool/Scripts/GUBP.Automation.cs b/Engine/Source/Programs/AutomationTool/Scripts/GUBP.Automation.cs
--- a/Engine/Source/Programs/AutomationTool/Scripts/GUBP.Automation.cs
+++ b/Engine/Source/Programs/AutomationTool/Scripts/GUBP.Automation.cs
@@ -7136,6 +7136,7 @@ public class GUBP : BuildCommand
 
             var FakeECArgs = new List<string>();
             var AgentGroupChains = new Dictionary<string, List<string>>();
+            var StickyChain = new List<string>();
             foreach (var NodeToDo in OrdereredToDo)
             {
                 if (GUBPNodes[NodeToDo].RunInEC() && !NodeIsAlreadyComplete(NodeToDo, LocalOnly)) // if something is already finished, we don't put it into EC  
@@ -7153,8 +7154,14 @@ public class GUBP : BuildCommand
                         }
                     }
                 }
+                if(GUBPNodes[NodeToDo].IsSticky())
+                {
+                    if(!StickyChain.Contains(NodeToDo))
+                    {
+                        StickyChain.Add(NodeToDo);
+                    }
+                }
             }
-
             foreach (var NodeToDo in OrdereredToDo)
             {
                 if (GUBPNodes[NodeToDo].RunInEC() && !NodeIsAlreadyComplete(NodeToDo, LocalOnly)) // if something is already finished, we don't put it into EC  
@@ -7275,6 +7282,33 @@ public class GUBP : BuildCommand
                             }
                         }
                     }
+                    if(GUBPNodes[NodeToDo].IsSticky())
+                    {
+                        var MyChain = StickyChain;
+                        int MyIndex = MyChain.IndexOf(NodeToDo);
+                        if (MyIndex > 0)
+                        {
+                            PreConditionUncompletedEcDeps.Add(MyChain[MyIndex - 1]);
+                        }
+                        else
+                        {
+                            var EcDeps = GetECDependencies(NodeToDo);
+                            foreach (var Dep in EcDeps)
+                            {
+                                if (GUBPNodes[Dep].RunInEC() && !NodeIsAlreadyComplete(Dep, LocalOnly) && OrdereredToDo.Contains(Dep)) // if something is already finished, we don't put it into EC
+                                {
+                                    if (OrdereredToDo.IndexOf(Dep) > OrdereredToDo.IndexOf(NodeToDo))
+                                    {
+                                        throw new AutomationException("Topological sort error, node {0} has a dependency of {1} which sorted after it.", NodeToDo, Dep);
+                                    }
+                                    if (!MyChain.Contains(Dep) && !PreConditionUncompletedEcDeps.Contains(Dep))
+                                    {
+                                        PreConditionUncompletedEcDeps.Add(Dep);
+                                    }
+                                }
+                            }
+                        }
+                    }
                     if (bHasNoop && PreConditionUncompletedEcDeps.Count == 0)
                     {
                         PreConditionUncompletedEcDeps.Add("Noop");
