Linux: fix LND deployment issues in Rocket (UE-14996).

- Also fix LND build script.

diff --git a/Engine/Build/BatchFiles/Linux/BuildThirdParty.sh b/Engine/Build/BatchFiles/Linux/BuildThirdParty.sh
--- a/Engine/Build/BatchFiles/Linux/BuildThirdParty.sh
+++ b/Engine/Build/BatchFiles/Linux/BuildThirdParty.sh
@@ -227,7 +227,7 @@ BuildLND()
   make $MAKE_ARGS
   local LIB_DIR=../lib/Linux/$TARGET_ARCH/
   mv *.so $LIB_DIR
-  cp -P --remove-destination $LIB_DIR/*.so ${TOP_DIR}/Binaries/ThirdParty/Linux/$TARGET_ARCH/
+  cp -P --remove-destination $LIB_DIR/*.so ${TOP_DIR}/Binaries/ThirdParty/LinuxNativeDialogs/Linux/$TARGET_ARCH/
   set +x
 }
 
diff --git a/Engine/Source/Programs/AutomationTool/Linux/LinuxPlatform.Automation.cs b/Engine/Source/Programs/AutomationTool/Linux/LinuxPlatform.Automation.cs
--- a/Engine/Source/Programs/AutomationTool/Linux/LinuxPlatform.Automation.cs
+++ b/Engine/Source/Programs/AutomationTool/Linux/LinuxPlatform.Automation.cs
@@ -48,7 +48,7 @@ public abstract class BaseLinuxPlatform : Platform
 			SC.StageFiles(StagedFileType.NonUFS, CombinePaths(SC.LocalRoot, "Engine/Binaries/ThirdParty/ICU/icu4c-53_1/", SC.PlatformDir, BuildArchitecture), Params.bDebugBuildsActuallyUseDebugCRT ? "*d.so*" : "*.so*", false, new[] { Params.bDebugBuildsActuallyUseDebugCRT ? "*.so*" : "*d.so*" }, CombinePaths("Engine/Binaries", SC.PlatformDir));
 		}
 
-		// assume that we always have to deploy Steam (FIXME: should be automatic)
+		// assume that we always have to deploy Steam (FIXME: should be automatic - UEPLAT-807)
 		{
 			string SteamVersion = "Steamv132";
 
@@ -63,13 +63,13 @@ public abstract class BaseLinuxPlatform : Platform
 			SC.StageFiles(StagedFileType.NonUFS, CombinePaths(SC.ProjectRoot, "Config"), "PerfCounters.json", false, null, CommandUtils.CombinePaths(SC.RelativeProjectRootForStage, "Saved/Config"), true);
 		}
 
-		// stage libLND
-		if (!SC.DedicatedServer)
+		// stage libLND (omit it for dedservers and Rocket - proper resolution is to use build receipts, see UEPLAT-807)
+		if (!SC.DedicatedServer && !GlobalCommandLine.Rocket)
 		{
 			SC.StageFiles(StagedFileType.NonUFS, CombinePaths(SC.LocalRoot, "Engine/Binaries/ThirdParty/LinuxNativeDialogs/", SC.PlatformDir, BuildArchitecture), "libLND*.so");
 		}
 
-		// assume that we always have to deploy OpenAL (FIXME: should be automatic)
+		// assume that we always have to deploy OpenAL (FIXME: should be automatic - UEPLAT-807)
 		{
 			SC.StageFiles(StagedFileType.NonUFS, CombinePaths(SC.LocalRoot, "Engine/Binaries/ThirdParty/OpenAL/", SC.PlatformDir), "libopenal.so.1", false, null, CombinePaths("Engine/Binaries", SC.PlatformDir));
 		}
