more robust build logging

diff --git a/Engine/Source/Programs/AutomationTool/BuildUtils.cs b/Engine/Source/Programs/AutomationTool/BuildUtils.cs
--- a/Engine/Source/Programs/AutomationTool/BuildUtils.cs
+++ b/Engine/Source/Programs/AutomationTool/BuildUtils.cs
@@ -54,7 +54,10 @@ namespace AutomationTool
 				throw new AutomationException(String.Format("Unabled to build Solution {0}. devenv.com not found.", SolutionFile));
 			}
 			string CmdLine = String.Format("\"{0}\" /build \"{1}\"", SolutionFile, BuildConfig);
+            var StartBuildSolution = DateTime.Now.ToString();
 			RunAndLog(Env, Env.MsDevExe, CmdLine, LogName);
+            var FinishBuildSolution = DateTime.Now.ToString();
+            PrintCSVFile(String.Format("UAT,Compile,{0},{1}", StartBuildSolution, FinishBuildSolution));
 		}
 
 		/// <summary>
diff --git a/Engine/Source/Programs/AutomationTool/Scripts/GUBP.Automation.cs b/Engine/Source/Programs/AutomationTool/Scripts/GUBP.Automation.cs
--- a/Engine/Source/Programs/AutomationTool/Scripts/GUBP.Automation.cs
+++ b/Engine/Source/Programs/AutomationTool/Scripts/GUBP.Automation.cs
@@ -739,9 +739,10 @@ public class GUBP : BuildCommand
                 bool ReallyDeleteBuildProducts = DeleteBuildProducts() && !GUBP.bForceIncrementalCompile;
                 Agenda.DoRetries = false; // these would delete build products
                 UE4Build.Build(Agenda, InDeleteBuildProducts: ReallyDeleteBuildProducts, InUpdateVersionFiles: false, InForceUnity: true);
-				
+                var StartPostBuild = DateTime.Now.ToString();
                 PostBuild(bp, UE4Build);
-
+                var FinishPostBuild = DateTime.Now.ToString();
+                PrintCSVFile(String.Format("UAT,PostBuild,{0},{1}", StartPostBuild, FinishPostBuild));
                 UE4Build.CheckBuildProducts(UE4Build.BuildProductFiles);
                 foreach (var Product in UE4Build.BuildProductFiles)
                 {
@@ -751,7 +752,10 @@ public class GUBP : BuildCommand
                 if (bp.bSignBuildProducts)
                 {
                     // Sign everything we built
+                    var StartSign = DateTime.Now.ToString();
                     CodeSign.SignMultipleIfEXEOrDLL(bp, BuildProducts);
+                    var FinishSign = DateTime.Now.ToString();
+                    PrintCSVFile(String.Format("UAT,SignBuildProducts,{0},{1}", StartSign, FinishSign));
                 }
                 PostBuildProducts(bp);
             }
@@ -2783,9 +2787,10 @@ public class GUBP : BuildCommand
             {
                 // not sure if we need something here or if the cook commandlet will automatically convert the exe name
             }
-
+            var StartCook = DateTime.Now.ToString();
 			CommandUtils.CookCommandlet(GameProj.FilePath, "UE4Editor-Cmd.exe", null, null, null, null, CookPlatform);
-		
+            var FinishCook = DateTime.Now.ToString();
+            PrintCSVFile(String.Format("UAT,Cook.{0}.{1},{2},{3}", GameProj.GameName, CookPlatform, StartCook, FinishCook));
             var CookedPath = RootIfAnyForTempStorage();
             var CookedFiles = CommandUtils.FindFiles("*", true, CookedPath);
             if (CookedFiles.GetLength(0) < 1)
diff --git a/Engine/Source/Programs/AutomationTool/UE4Build.cs b/Engine/Source/Programs/AutomationTool/UE4Build.cs
--- a/Engine/Source/Programs/AutomationTool/UE4Build.cs
+++ b/Engine/Source/Programs/AutomationTool/UE4Build.cs
@@ -185,9 +185,10 @@ namespace AutomationTool
 			XGEItem Result = new XGEItem();
 
 			ClearExportedXGEXML();
-
+            var StartXGEManifestExport = DateTime.Now.ToString();
 			RunUBT(CmdEnv, UBTExecutable: UBTExecutable, Project: ProjectName, Target: TargetName, Platform: Platform.ToString(), Config: Config, AdditionalArgs: "-generatemanifest -nobuilduht -xgeexport" + AddArgs, EnvVars: EnvVars);
-
+            var FinishXGEManifestExport = DateTime.Now.ToString();
+            PrintCSVFile(String.Format("UAT,GenerateXGEManifest.{0}.{1}.{2},{3},{4}", TargetName, Platform.ToString(), Config, StartXGEManifestExport, FinishXGEManifestExport));
 			PrepareManifest(UBTManifest, true);
 
 			Result.Platform = Platform;
@@ -279,8 +280,10 @@ namespace AutomationTool
 			}
 
 			PrepareUBT();
-
+            var StartCleanUBT = DateTime.Now.ToString();
 			RunUBT(CmdEnv, UBTExecutable: UBTExecutable, Project: ProjectName, Target: TargetName, Platform: Platform.ToString(), Config: Config, AdditionalArgs: "-clean" + AddArgs, EnvVars: EnvVars);
+            var FinishCleanUBT = DateTime.Now.ToString();
+            PrintCSVFile(String.Format("UAT,CleanWithUBT.{0}.{1}.{2},{3},{4}", TargetName, Platform.ToString(), Config, StartCleanUBT, FinishCleanUBT));
 		}
 
 		void BuildWithUBT(string ProjectName, string TargetName, UnrealBuildTool.UnrealTargetPlatform TargetPlatform, string Config, string UprojectPath, bool ForceMonolithic = false, bool ForceNonUnity = false, bool ForceDebugInfo = false, bool ForceFlushMac = false, bool DisableXGE = false, string InAddArgs = "", bool ForceUnity = false, Dictionary<string, string> EnvVars = null)
@@ -326,14 +329,16 @@ namespace AutomationTool
                 string UBTManifest = GetUBTManifest(UprojectPath, AddArgs);
 
 				DeleteFile(UBTManifest);
-
+                var StartGenerateManifest = DateTime.Now.ToString();
 				RunUBT(CmdEnv, UBTExecutable: UBTExecutable, Project: ProjectName, Target: TargetName, Platform: TargetPlatform.ToString(), Config: Config, AdditionalArgs: AddArgs +  " -generatemanifest" , EnvVars: EnvVars);
-
+                var FinishGenerateManifest = DateTime.Now.ToString();
+                PrintCSVFile(String.Format("UAT,PrepareUBTManifest.{0}.{1}.{2},{3},{4}", TargetName, TargetPlatform.ToString(), Config, StartGenerateManifest, FinishGenerateManifest));
 				PrepareManifest(UBTManifest, false);
 			}
-
+            var StartCompile = DateTime.Now.ToString();
 			RunUBT(CmdEnv, UBTExecutable: UBTExecutable, Project: ProjectName, Target: TargetName, Platform: TargetPlatform.ToString(), Config: Config, AdditionalArgs: AddArgs, EnvVars: EnvVars);
-
+            var FinishCompile = DateTime.Now.ToString();
+            PrintCSVFile(String.Format("UAT,Compile.{0}.{1}.{2},{3},{4}", TargetName, TargetPlatform.ToString(), Config, StartCompile, FinishCompile));
 			// allow the platform to perform any actions after building a target (seems almost like this should be done in UBT)
 			Platform.Platforms[TargetPlatform].PostBuildTarget(this, string.IsNullOrEmpty(ProjectName) ? TargetName : ProjectName, UprojectPath, Config);
 
@@ -1304,14 +1309,20 @@ namespace AutomationTool
 				}
 				if (XGEItems.Count > 0)
 				{
+                    var StartXGERun = DateTime.Now.ToString();
+                    var FinishXGERun = "";
 					if (!RunXGE(XGEItems, TaskFilePath, Agenda.DoRetries, Agenda.SpecialTestFlag))
 					{
+                        FinishXGERun = DateTime.Now.ToString();
+                        PrintCSVFile(String.Format("UAT,CompileWithXGE,{0},{1}", StartXGERun, FinishXGERun));
 						throw new AutomationException("BUILD FAILED: XGE failed, retries not enabled:");
 					}
 					else
 					{
 						bUsedXGE = true;
 					}
+                    FinishXGERun = DateTime.Now.ToString();
+                    PrintCSVFile(String.Format("UAT,CompileWithXGE,{0},{1}", StartXGERun, FinishXGERun));
 				}
 				else
 				{
@@ -1347,6 +1358,7 @@ namespace AutomationTool
 		{
 			// Check build products
 			{
+                var StartCheckBuildProducts = DateTime.Now.ToString();
 				Log("Build products *******");
 				if (BuildProductFiles.Count < 1)
 				{
@@ -1364,6 +1376,8 @@ namespace AutomationTool
 					}
 				}
 				Log("End Build products *******");
+                var FinishCheckBuildProducts = DateTime.Now.ToString();
+                PrintCSVFile(String.Format("UAT,CheckBuildProducts,{0},{1}", StartCheckBuildProducts, FinishCheckBuildProducts));
 			}
 		}
 
