Allow staging runtime dependencies for all platforms, not just windows. Leave staging build products from receipts as opt-in, since most platforms already have custom code to deal with it.

diff --git a/Engine/Source/Programs/AutomationTool/DeploymentContext.cs b/Engine/Source/Programs/AutomationTool/DeploymentContext.cs
--- a/Engine/Source/Programs/AutomationTool/DeploymentContext.cs
+++ b/Engine/Source/Programs/AutomationTool/DeploymentContext.cs
@@ -367,7 +367,7 @@ public class DeploymentContext //: ProjectParams
 		}
 	}
 
-	public void StageFilesInReceipt(BuildReceipt Receipt)
+	public void StageBuildProductsFromReceipt(BuildReceipt Receipt)
 	{
 		// Stage all the build products needed at runtime
 		foreach(BuildProduct BuildProduct in Receipt.BuildProducts)
@@ -387,7 +387,10 @@ public class DeploymentContext //: ProjectParams
 				StageFile(StagedFileType.DebugNonUFS, BuildProduct.Path);
 			}
 		}
+	}
 
+	public void StageRuntimeDependenciesFromReceipt(BuildReceipt Receipt)
+	{
 		// Also stage any additional runtime dependencies, like ThirdParty DLLs
 		foreach(RuntimeDependency RuntimeDependency in Receipt.RuntimeDependencies)
 		{
diff --git a/Engine/Source/Programs/AutomationTool/Scripts/CopyBuildToStagingDirectory.Automation.cs b/Engine/Source/Programs/AutomationTool/Scripts/CopyBuildToStagingDirectory.Automation.cs
--- a/Engine/Source/Programs/AutomationTool/Scripts/CopyBuildToStagingDirectory.Automation.cs
+++ b/Engine/Source/Programs/AutomationTool/Scripts/CopyBuildToStagingDirectory.Automation.cs
@@ -263,6 +263,13 @@ public partial class Project : CommandUtils
 
 		ThisPlatform.GetFilesToDeployOrStage(Params, SC);
 
+
+		// Stage any extra runtime dependencies from the receipts
+		foreach(BuildReceipt Receipt in SC.StageTargetReceipts)
+		{
+			SC.StageRuntimeDependenciesFromReceipt(Receipt);
+		}
+
 		// Get the build.properties file
 		// this file needs to be treated as a UFS file for casing, but NonUFS for being put into the .pak file
 		// @todo: Maybe there should be a new category - UFSNotForPak
diff --git a/Engine/Source/Programs/AutomationTool/Win/WinPlatform.Automation.cs b/Engine/Source/Programs/AutomationTool/Win/WinPlatform.Automation.cs
--- a/Engine/Source/Programs/AutomationTool/Win/WinPlatform.Automation.cs
+++ b/Engine/Source/Programs/AutomationTool/Win/WinPlatform.Automation.cs
@@ -35,10 +35,10 @@ public abstract class BaseWinPlatform : Platform
 			StageExecutable("exe", SC, CommandUtils.CombinePaths(SC.LocalRoot, "Engine/Binaries", SC.PlatformDir), "CrashReportClient.");
 		}
 
-		// Loop through all the things to stage
+		// Stage all the build products
 		foreach(BuildReceipt Receipt in SC.StageTargetReceipts)
 		{
-			SC.StageFilesInReceipt(Receipt);
+			SC.StageBuildProductsFromReceipt(Receipt);
 		}
 
 		// Copy the splash screen, windows specific
