HTML5 launch helper built by GUBP automation scrips.

diff --git a/Engine/Source/Programs/AutomationTool/Scripts/GUBP.Automation.cs b/Engine/Source/Programs/AutomationTool/Scripts/GUBP.Automation.cs
--- a/Engine/Source/Programs/AutomationTool/Scripts/GUBP.Automation.cs
+++ b/Engine/Source/Programs/AutomationTool/Scripts/GUBP.Automation.cs
@@ -1164,6 +1164,17 @@ public class GUBP : BuildCommand
 						}
                     );
 				}
+
+				bool WithHTML5 = !bp.BranchOptions.PlatformsToRemove.Contains(UnrealTargetPlatform.HTML5);
+				if (WithHTML5)
+				{
+					Agenda.HTML5DotNetProjects.AddRange(
+						new string[]
+						{
+							CombinePaths(@"Engine\Source\Programs\HTML5\HTML5LaunchHelper\HTML5LaunchHelper.csproj"),
+						}
+					);
+				}
             }
 
             string AddArgs = "-nobuilduht -skipactionhistory -CopyAppBundleBackToDevice";
diff --git a/Engine/Source/Programs/AutomationTool/UE4Build.cs b/Engine/Source/Programs/AutomationTool/UE4Build.cs
--- a/Engine/Source/Programs/AutomationTool/UE4Build.cs
+++ b/Engine/Source/Programs/AutomationTool/UE4Build.cs
@@ -647,6 +647,11 @@ namespace AutomationTool
 			/// the solution file base name, but with various binary file extensions
 			public List<string> IOSDotNetProjects = new List<string>();
 
+			/// These are .NET project binary files that are specific to HTML5.  We define these build products as existing in the 
+			/// DotNET directory and assume that the output binary file names match
+			/// the solution file base name, but with various binary file extensions
+			public List<string> HTML5DotNetProjects = new List<string>();
+
 			public string SwarmProject = "";
 
 
@@ -1253,6 +1258,17 @@ namespace AutomationTool
 				AddIOSBuildProductsForCSharpProj(IOSCsProj);
 			}
 
+			foreach (var HTML5DotNetProject in Agenda.HTML5DotNetProjects)
+			{
+				string CsProj = Path.Combine(CmdEnv.LocalRoot, HTML5DotNetProject);
+
+				PrepareBuildProductsForCSharpProj(CsProj);
+
+				BuildCSharpProject(CmdEnv, CsProj);
+
+				AddBuildProductsForCSharpProj(CsProj);
+			}
+
 			foreach (var File in Agenda.ExtraDotNetFiles)
 			{
 				AddBuildProductsForCSharpProj(Path.Combine(CmdEnv.LocalRoot, File));
