fix for Tappy Chicken not building out of the box in Binary build on PC
UE-13689
#ios

diff --git a/Engine/Source/Programs/AutomationTool/ProjectUtils.cs b/Engine/Source/Programs/AutomationTool/ProjectUtils.cs
--- a/Engine/Source/Programs/AutomationTool/ProjectUtils.cs
+++ b/Engine/Source/Programs/AutomationTool/ProjectUtils.cs
@@ -256,7 +256,7 @@ namespace AutomationTool
 
 			// detect if the project is content only, but has non-default build settings
 			List<string> ExtraSearchPaths = null;
-			if (!string.IsNullOrEmpty(RawProjectPath))
+			if (!string.IsNullOrEmpty(RawProjectPath) && (!GlobalCommandLine.Rocket || BuildHostPlatform.Current.Platform == UnrealTargetPlatform.Mac))
 			{
 				if (RequiresTempTarget(RawProjectPath, ClientTargetPlatforms))
 				{
@@ -483,7 +483,7 @@ namespace AutomationTool
 						typeof(UnrealBuildTool.UnrealBuildTool).Assembly.Location
 					};
 			var TargetsDLL = DynamicCompilation.CompileAndLoadAssembly(TargetsDllFilename, TargetScripts, ReferencedAssemblies, null, DoNotCompile);
-			var DummyTargetInfo = new TargetInfo(UnrealTargetPlatform.Win64, UnrealTargetConfiguration.Development);
+			var DummyTargetInfo = new TargetInfo(BuildHostPlatform.Current.Platform, UnrealTargetConfiguration.Development);
 			var AllCompiledTypes = TargetsDLL.GetTypes();
 			foreach (Type TargetType in AllCompiledTypes)
 			{
