build script, fix preconditions

diff --git a/Engine/Source/Programs/AutomationTool/Scripts/GUBP.Automation.cs b/Engine/Source/Programs/AutomationTool/Scripts/GUBP.Automation.cs
--- a/Engine/Source/Programs/AutomationTool/Scripts/GUBP.Automation.cs
+++ b/Engine/Source/Programs/AutomationTool/Scripts/GUBP.Automation.cs
@@ -7314,7 +7314,10 @@ public class GUBP : BuildCommand
                         int MyIndex = MyChain.IndexOf(NodeToDo);
                         if (MyIndex > 0)
                         {
-                            PreConditionUncompletedEcDeps.Add(MyChain[MyIndex - 1]);
+                            if (!PreConditionUncompletedEcDeps.Contains(MyChain[MyIndex - 1]) && !NodeIsAlreadyComplete(MyChain[MyIndex - 1], LocalOnly))
+                            {
+                                PreConditionUncompletedEcDeps.Add(MyChain[MyIndex - 1]);
+                            }
                         }
                         else
                         {
