Fix incorrectly discarding rules which may have a higher precedence than the best one currently found.

diff --git a/Engine/Source/Programs/AutomationTool/FileFilter.cs b/Engine/Source/Programs/AutomationTool/FileFilter.cs
--- a/Engine/Source/Programs/AutomationTool/FileFilter.cs
+++ b/Engine/Source/Programs/AutomationTool/FileFilter.cs
@@ -461,19 +461,9 @@ public class FileFilter
 		}
 
 		// If there is no rule under the current node which is better than the current best node, early out
-		if (CurrentBestNode.Type == FileFilterType.Include)
+		if (CurrentNode.MaxIncludeRuleNumber <= CurrentBestNode.RuleNumber && CurrentNode.MaxExcludeRuleNumber <= CurrentBestNode.RuleNumber)
 		{
-			if (CurrentNode.MaxExcludeRuleNumber <= CurrentBestNode.RuleNumber)
-			{
-				return CurrentBestNode;
-			}
-		}
-		else
-		{
-			if (CurrentNode.MaxIncludeRuleNumber <= CurrentBestNode.RuleNumber)
-			{
-				return CurrentBestNode;
-			}
+			return CurrentBestNode;
 		}
 
 		// Test all the branches for one that matches
