Create a mini-version of the Rocket install to generate DDC, since shader and static mesh DDC keys are different without PS4 and Simplygon respectively.

diff --git a/Engine/Source/Programs/AutomationTool/Scripts/RocketBuild.Automation.cs b/Engine/Source/Programs/AutomationTool/Scripts/RocketBuild.Automation.cs
--- a/Engine/Source/Programs/AutomationTool/Scripts/RocketBuild.Automation.cs
+++ b/Engine/Source/Programs/AutomationTool/Scripts/RocketBuild.Automation.cs
@@ -933,6 +933,7 @@ namespace Rocket
 			SavedManifestPath = CommandUtils.CombinePaths(SavedDir, "DerivedDataCacheManifest.txt");
 
 			AddDependency(GUBP.RootEditorNode.StaticGetFullName(HostPlatform));
+			AddDependency(GUBP.ToolsNode.StaticGetFullName(HostPlatform));
 		}
 
 		public static string StaticGetFullName(UnrealTargetPlatform InHostPlatform)
@@ -959,33 +960,60 @@ namespace Rocket
 			List<string> ManifestFiles = new List<string>();
 			if(!bp.ParseParam("NoDDC"))
 			{
-				string EditorExe = CommandUtils.GetEditorCommandletExe(CommandUtils.CmdEnv.LocalRoot, HostPlatform);
-				string RelativePakPath = "Engine/DerivedDataCache/Compressed.ddp";
+				// Find all the projects we're interested in
+				List<BranchInfo.BranchUProject> Projects = new List<BranchInfo.BranchUProject>();
+				foreach(string ProjectName in ProjectNames)
+				{
+					BranchInfo.BranchUProject Project = bp.Branch.FindGameChecked(ProjectName);
+					if(!Project.Properties.bIsCodeBasedProject)
+					{
+						Projects.Add(Project);
+					}
+				}
 
-				// Delete the output file
-				string OutputPakFile = CommandUtils.CombinePaths(CommandUtils.CmdEnv.LocalRoot, RelativePakPath);
+				// Filter out the files we need to build DDC. Removing confidential folders can affect DDC keys, so we want to be sure that we're making DDC with a build that can use it.
+				FileFilter Filter = new FileFilter(FileFilterType.Exclude);
+				Filter.AddRuleForFiles(AllDependencyBuildProducts, CommandUtils.CmdEnv.LocalRoot, FileFilterType.Include);
+				Filter.ReadRulesFromFile(CommandUtils.CombinePaths(CommandUtils.CmdEnv.LocalRoot, "Engine", "Build", "InstalledEngineFilters.ini"), "CopyEditor", HostPlatform.ToString());
+				Filter.Exclude("/Engine/Build/...");
+				Filter.Exclude("/Engine/Extras/...");
+				Filter.Exclude("/Engine/DerivedDataCache/...");
+				Filter.Exclude("/Samples/...");
+				Filter.Exclude("/Templates/...");
+				Filter.Exclude(".../Source/...");
+				Filter.Exclude(".../Intermediate/...");
+				Filter.ExcludeConfidentialPlatforms();
+				Filter.ExcludeConfidentialFolders();
+				Filter.Include("/Engine/Build/NotForLicensees/EpicInternal.txt");
+				Filter.Include("/Engine/Binaries/.../*DDCUtils*"); // Make sure we can use the shared DDC!
+
+				// Copy everything to a temporary directory
+				string TempDir = CommandUtils.CombinePaths(CommandUtils.CmdEnv.LocalRoot, "LocalBuilds", "RocketDDC", CommandUtils.GetGenericPlatformName(HostPlatform));
+				CommandUtils.DeleteDirectoryContents(TempDir);
+				CommandUtils.ThreadedCopyFiles(CommandUtils.CmdEnv.LocalRoot, TempDir, Filter, true);
+
+				// Get paths to everything within the temporary directory
+				string EditorExe = CommandUtils.GetEditorCommandletExe(TempDir, HostPlatform);
+				string RelativePakPath = "Engine/DerivedDataCache/Compressed.ddp";
+				string OutputPakFile = CommandUtils.CombinePaths(TempDir, RelativePakPath);
 				string OutputCsvFile = Path.ChangeExtension(OutputPakFile, ".csv");
 
 				// Generate DDC for all the non-code projects. We don't necessarily have editor DLLs for the code projects, but they should be the same as their blueprint counterparts.
 				List<string> ProjectPakFiles = new List<string>();
-				foreach(string ProjectName in ProjectNames)
+				foreach(BranchInfo.BranchUProject Project in Projects)
 				{
-					BranchInfo.BranchUProject Project = bp.Branch.FindGameChecked(ProjectName);
-					if(!Project.Properties.bIsCodeBasedProject)
-					{
-						CommandUtils.Log("Generating DDC data for {0} on {1}", Project.GameName, TargetPlatforms);
-						CommandUtils.DDCCommandlet(Project.FilePath, EditorExe, null, TargetPlatforms, "-fill -DDC=CreateInstalledEnginePak -ProjectOnly");
+					CommandUtils.Log("Generating DDC data for {0} on {1}", Project.GameName, TargetPlatforms);
+					CommandUtils.DDCCommandlet(Project.FilePath, EditorExe, null, TargetPlatforms, "-fill -DDC=CreateInstalledEnginePak -ProjectOnly");
 
-						string ProjectPakFile = CommandUtils.CombinePaths(Path.GetDirectoryName(OutputPakFile), String.Format("Compressed-{0}.ddp", ProjectName));
-						CommandUtils.DeleteFile(ProjectPakFile);
-						CommandUtils.RenameFile(OutputPakFile, ProjectPakFile);
+					string ProjectPakFile = CommandUtils.CombinePaths(Path.GetDirectoryName(OutputPakFile), String.Format("Compressed-{0}.ddp", Project.GameName));
+					CommandUtils.DeleteFile(ProjectPakFile);
+					CommandUtils.RenameFile(OutputPakFile, ProjectPakFile);
 
-						string ProjectCsvFile = Path.ChangeExtension(ProjectPakFile, ".csv");
-						CommandUtils.DeleteFile(ProjectCsvFile);
-						CommandUtils.RenameFile(OutputCsvFile, ProjectCsvFile);
+					string ProjectCsvFile = Path.ChangeExtension(ProjectPakFile, ".csv");
+					CommandUtils.DeleteFile(ProjectCsvFile);
+					CommandUtils.RenameFile(OutputCsvFile, ProjectCsvFile);
 
-						ProjectPakFiles.Add(Path.GetFileName(ProjectPakFile));
-					}
+					ProjectPakFiles.Add(Path.GetFileName(ProjectPakFile));
 				}
 
 				// Generate DDC for the editor, and merge all the other PAK files in
