Added the architecture to Project Params.
Specify the windows XP arch from the package project option

diff --git a/Engine/Source/Editor/MainFrame/Private/Frame/MainFrameActions.cpp b/Engine/Source/Editor/MainFrame/Private/Frame/MainFrameActions.cpp
--- a/Engine/Source/Editor/MainFrame/Private/Frame/MainFrameActions.cpp
+++ b/Engine/Source/Editor/MainFrame/Private/Frame/MainFrameActions.cpp
@@ -549,7 +549,7 @@ void FMainFrameActionCallbacks::PackageProject( const FName InPlatformInfoName )
 		GConfig->GetString(TEXT("/Script/WindowsTargetPlatform.WindowsTargetSettings"), TEXT("MinimumOSVersion"), MinumumSupportedWindowsOS, GEngineIni);
 		if (MinumumSupportedWindowsOS == TEXT("MSOS_XP"))
 		{
-			OptionalParams += TEXT(" -OverrideMinimumOS=WinXP");
+			OptionalParams += TEXT(" -SpecifiedArchitecture=_xp");
 			bTargetPlatformCanUseCrashReporter = false;
 		}
 	}
diff --git a/Engine/Source/Programs/AutomationTool/ProjectParams.cs b/Engine/Source/Programs/AutomationTool/ProjectParams.cs
--- a/Engine/Source/Programs/AutomationTool/ProjectParams.cs
+++ b/Engine/Source/Programs/AutomationTool/ProjectParams.cs
@@ -390,7 +390,7 @@ namespace AutomationTool
 			bool? Distribution = null,
             bool? Prebuilt = null,
             int? RunTimeoutSeconds = null,
-			string OverrideMinimumOS = null,
+			string SpecifiedArchitecture = null,
             bool? IterativeDeploy = null
 			)
 		{
@@ -559,7 +559,7 @@ namespace AutomationTool
 			this.DeviceUsername = ParseParamValueIfNotSpecified(Command, DeviceUsername, "deviceuser", String.Empty);
 			this.DevicePassword = ParseParamValueIfNotSpecified(Command, DevicePassword, "devicepass", String.Empty);
 			this.CrashReporter = GetParamValueIfNotSpecified(Command, CrashReporter, this.CrashReporter, "crashreporter");
-			this.OverrideMinimumOS = ParseParamValueIfNotSpecified(Command, OverrideMinimumOS, "OverrideMinimumOS", String.Empty);
+			this.SpecifiedArchitecture = ParseParamValueIfNotSpecified(Command, SpecifiedArchitecture, "specifiedarchitecture", String.Empty);
 			if (ClientConfigsToBuild == null)
 			{
 				if (Command != null)
@@ -1339,8 +1339,8 @@ namespace AutomationTool
         [Help("RunTimeoutSeconds", "timeout to wait after we lunch the game")]
         public int RunTimeoutSeconds;
 
-		[Help("OverrideMinimumOS", "Determine a specific Minimum OS")]
-		public string OverrideMinimumOS;
+		[Help("SpecifiedArchitecture", "Determine a specific Minimum OS")]
+		public string SpecifiedArchitecture;
 
 		#endregion
 
diff --git a/Engine/Source/Programs/AutomationTool/Scripts/CopyBuildToStagingDirectory.Automation.cs b/Engine/Source/Programs/AutomationTool/Scripts/CopyBuildToStagingDirectory.Automation.cs
--- a/Engine/Source/Programs/AutomationTool/Scripts/CopyBuildToStagingDirectory.Automation.cs
+++ b/Engine/Source/Programs/AutomationTool/Scripts/CopyBuildToStagingDirectory.Automation.cs
@@ -1360,6 +1360,7 @@ public partial class Project : CommandUtils
 			List<string> ExecutablesToStage = new List<string>();
 
 			string PlatformName = StagePlatform.ToString();
+			string StageArchitecture = !String.IsNullOrEmpty(Params.SpecifiedArchitecture) ? Params.SpecifiedArchitecture : "";
 			foreach (var Target in ListToProcess)
 			{
 				foreach (var Config in ConfigsToProcess)
@@ -1367,7 +1368,7 @@ public partial class Project : CommandUtils
 					string Exe = Target;
 					if (Config != UnrealTargetConfiguration.Development)
 					{
-						Exe = Target + "-" + PlatformName + "-" + Config.ToString();
+						Exe = Target + "-" + PlatformName + "-" + Config.ToString() + StageArchitecture;
 					}
 					ExecutablesToStage.Add(Exe);
 				}
@@ -1404,8 +1405,7 @@ public partial class Project : CommandUtils
 
 					foreach (UnrealTargetPlatform ReceiptPlatform in SubPlatformsToStage)
 					{
-						// Read the receipt (FIXME: store the architecture in project params?)
-						string Architecture = UEBuildPlatform.GetBuildPlatform(ReceiptPlatform).GetActiveArchitecture();
+						string Architecture = !String.IsNullOrEmpty(Params.SpecifiedArchitecture) ? Params.SpecifiedArchitecture : UEBuildPlatform.GetBuildPlatform(ReceiptPlatform).GetActiveArchitecture();
 						string ReceiptFileName = BuildReceipt.GetDefaultPath(ReceiptBaseDir, Target, ReceiptPlatform, Config, Architecture);
 						if(!File.Exists(ReceiptFileName))
 						{
