fix for project on the command line and in UE4CommandLine.txt
UE-15445

diff --git a/Engine/Source/Programs/AutomationTool/Scripts/CopyBuildToStagingDirectory.Automation.cs b/Engine/Source/Programs/AutomationTool/Scripts/CopyBuildToStagingDirectory.Automation.cs
--- a/Engine/Source/Programs/AutomationTool/Scripts/CopyBuildToStagingDirectory.Automation.cs
+++ b/Engine/Source/Programs/AutomationTool/Scripts/CopyBuildToStagingDirectory.Automation.cs
@@ -1174,6 +1174,10 @@ public partial class Project : CommandUtils
 			}
 
 			String ProjectFile = String.Format("{0} ", SC.ProjectArgForCommandLines);
+			if (SC.StageTargetPlatform.PlatformType == UnrealTargetPlatform.Mac || SC.StageTargetPlatform.PlatformType == UnrealTargetPlatform.Win64 || SC.StageTargetPlatform.PlatformType == UnrealTargetPlatform.Win32 || SC.StageTargetPlatform.PlatformType == UnrealTargetPlatform.Linux)
+			{
+				ProjectFile = "";
+			}
 			Directory.CreateDirectory(GetIntermediateCommandlineDir(SC));
 			string CommandLine = String.Format("{0} {1} {2} {3}\n", ProjectFile, Params.StageCommandline.Trim(new char[] { '\"' }), Params.RunCommandline.Trim(new char[] { '\"' }), FileHostParams).Trim();
 			if (Params.IterativeDeploy)
@@ -1185,6 +1189,10 @@ public partial class Project : CommandUtils
 		else if (!Params.IsCodeBasedProject)
 		{
 			String ProjectFile = String.Format("{0} ", SC.ProjectArgForCommandLines);
+			if (SC.StageTargetPlatform.PlatformType == UnrealTargetPlatform.Mac || SC.StageTargetPlatform.PlatformType == UnrealTargetPlatform.Win64 || SC.StageTargetPlatform.PlatformType == UnrealTargetPlatform.Win32 || SC.StageTargetPlatform.PlatformType == UnrealTargetPlatform.Linux)
+			{
+				ProjectFile = "";
+			}
 			Directory.CreateDirectory(GetIntermediateCommandlineDir(SC));
 			File.WriteAllText(IntermediateCmdLineFile, ProjectFile);
 		}
