build script, remove vs express test compilation

diff --git a/Engine/Source/Programs/AutomationTool/Scripts/GUBP.Automation.cs b/Engine/Source/Programs/AutomationTool/Scripts/GUBP.Automation.cs
--- a/Engine/Source/Programs/AutomationTool/Scripts/GUBP.Automation.cs
+++ b/Engine/Source/Programs/AutomationTool/Scripts/GUBP.Automation.cs
@@ -6398,7 +6398,7 @@ public class GUBP : BuildCommand
 			{
 				//AddNode(new IOSOnPCTestNode(this)); - Disable IOSOnPCTest until a1011 crash is fixed
 			}
-			AddNode(new VSExpressTestNode(this));
+			//AddNode(new VSExpressTestNode(this));
 			if (ActivePlatforms.Contains(UnrealTargetPlatform.Linux) && !BranchOptions.ExcludePlatformsForEditor.Contains(UnrealTargetPlatform.Linux))
 			{
 				AddNode(new RootEditorCrossCompileLinuxNode(UnrealTargetPlatform.Win64));
