Added the new Category and Keywords loc targets to the OneSky automation

diff --git a/Engine/Config/Localization/WordCount.ini b/Engine/Config/Localization/WordCount.ini
--- a/Engine/Config/Localization/WordCount.ini
+++ b/Engine/Config/Localization/WordCount.ini
@@ -40,4 +40,18 @@ WordCountReportName=PropertyNames.csv
 CommandletClass=GenerateTextLocalizationReport
 SourcePath=Content/Localization/EditorTutorials
 ManifestName=EditorTutorials.manifest
-WordCountReportName=EditorTutorials.csv
\ No newline at end of file
+WordCountReportName=EditorTutorials.csv
+
+;Write Text Localization Report for Editor Categories
+[GatherTextStep5]
+CommandletClass=GenerateTextLocalizationReport
+SourcePath=Content/Localization/Category
+ManifestName=Category.manifest
+WordCountReportName=Category.csv
+
+;Write Text Localization Report for Editor Keywords
+[GatherTextStep6]
+CommandletClass=GenerateTextLocalizationReport
+SourcePath=Content/Localization/Keywords
+ManifestName=Keywords.manifest
+WordCountReportName=Keywords.csv
\ No newline at end of file
diff --git a/Engine/Source/Programs/AutomationTool/Scripts/Localisation.Automation.cs b/Engine/Source/Programs/AutomationTool/Scripts/Localisation.Automation.cs
--- a/Engine/Source/Programs/AutomationTool/Scripts/Localisation.Automation.cs
+++ b/Engine/Source/Programs/AutomationTool/Scripts/Localisation.Automation.cs
@@ -55,6 +55,8 @@ class Localise : BuildCommand
         ExportProjectToDirectory(oneSkyService, projectGroup, "EditorTutorials");
         ExportProjectToDirectory(oneSkyService, projectGroup, "PropertyNames");
         ExportProjectToDirectory(oneSkyService, projectGroup, "ToolTips");
+        ExportProjectToDirectory(oneSkyService, projectGroup, "Category");
+        ExportProjectToDirectory(oneSkyService, projectGroup, "Keywords");
 
         // Submit changelist for backed up POs from OneSky.
         if (P4Enabled)
@@ -86,7 +88,9 @@ class Localise : BuildCommand
 				String.Format("-config={0}", @"./Config/Localization/Editor.ini") + (string.IsNullOrEmpty(CommandletSCCArguments) ? "" : " " + CommandletSCCArguments),
 				String.Format("-config={0}", @"./Config/Localization/EditorTutorials.ini") + (string.IsNullOrEmpty(CommandletSCCArguments) ? "" : " " + CommandletSCCArguments),
 				String.Format("-config={0}", @"./Config/Localization/PropertyNames.ini") + (string.IsNullOrEmpty(CommandletSCCArguments) ? "" : " " + CommandletSCCArguments),
-				String.Format("-config={0}", @"./Config/Localization/ToolTips.ini") + (string.IsNullOrEmpty(CommandletSCCArguments) ? "" : " " + CommandletSCCArguments),
+                String.Format("-config={0}", @"./Config/Localization/ToolTips.ini") + (string.IsNullOrEmpty(CommandletSCCArguments) ? "" : " " + CommandletSCCArguments),
+                String.Format("-config={0}", @"./Config/Localization/Category.ini") + (string.IsNullOrEmpty(CommandletSCCArguments) ? "" : " " + CommandletSCCArguments),
+				String.Format("-config={0}", @"./Config/Localization/Keywords.ini") + (string.IsNullOrEmpty(CommandletSCCArguments) ? "" : " " + CommandletSCCArguments),
 				String.Format("-config={0}", @"./Config/Localization/WordCount.ini"),
 			};
 
@@ -112,6 +116,8 @@ class Localise : BuildCommand
         UploadDirectoryToProject(GetProject(oneSkyService, "Unreal Engine", "EditorTutorials"), new DirectoryInfo(CmdEnv.LocalRoot + "/Engine/Content/Localization/EditorTutorials"), "*.po");
         UploadDirectoryToProject(GetProject(oneSkyService, "Unreal Engine", "PropertyNames"), new DirectoryInfo(CmdEnv.LocalRoot + "/Engine/Content/Localization/PropertyNames"), "*.po");
         UploadDirectoryToProject(GetProject(oneSkyService, "Unreal Engine", "ToolTips"), new DirectoryInfo(CmdEnv.LocalRoot + "/Engine/Content/Localization/ToolTips"), "*.po");
+        UploadDirectoryToProject(GetProject(oneSkyService, "Unreal Engine", "Category"), new DirectoryInfo(CmdEnv.LocalRoot + "/Engine/Content/Localization/Category"), "*.po");
+        UploadDirectoryToProject(GetProject(oneSkyService, "Unreal Engine", "Keywords"), new DirectoryInfo(CmdEnv.LocalRoot + "/Engine/Content/Localization/Keywords"), "*.po");
 
 		// Localisation statistics estimator.
 		if (P4Enabled)
