disable content as code builds in binary release
UE-11636
#ios

diff --git a/Engine/Source/Programs/AutomationTool/ProjectUtils.cs b/Engine/Source/Programs/AutomationTool/ProjectUtils.cs
--- a/Engine/Source/Programs/AutomationTool/ProjectUtils.cs
+++ b/Engine/Source/Programs/AutomationTool/ProjectUtils.cs
@@ -256,7 +256,7 @@ namespace AutomationTool
 
 			// detect if the project is content only, but has non-default build settings
 			List<string> ExtraSearchPaths = null;
-			if (!string.IsNullOrEmpty(RawProjectPath) && (!GlobalCommandLine.Rocket || BuildHostPlatform.Current.Platform == UnrealTargetPlatform.Mac))
+			if (!string.IsNullOrEmpty(RawProjectPath) && !GlobalCommandLine.Rocket)
 			{
 				if (RequiresTempTarget(RawProjectPath, ClientTargetPlatforms))
 				{
