UAT: Add support for checking equality of engine versions

diff --git a/Engine/Source/Programs/AutomationTool/Utils.cs b/Engine/Source/Programs/AutomationTool/Utils.cs
--- a/Engine/Source/Programs/AutomationTool/Utils.cs
+++ b/Engine/Source/Programs/AutomationTool/Utils.cs
@@ -865,6 +865,19 @@ namespace AutomationTool
                 BranchName);
         }
 
+		public override bool Equals(object obj)
+		{
+			var that = obj as FEngineVersionSupport;
+			if (that == null)
+				return false;
+			return this.ToString() == that.ToString();
+		}
+
+		public override int GetHashCode()
+		{
+			return this.ToString().GetHashCode();
+		}
+
         /// <summary>
         /// Ctor initializes with the values from the supplied Version file. The BranchName and CL are also taken from the current <see cref="CommandUtils.P4Env"/>.
         /// </summary>
