Disable precompiling all editor and plugin modules, and building feature packs, if we're not building the engine for distribution.

diff --git a/Engine/Source/Programs/AutomationTool/Scripts/GUBP.Automation.cs b/Engine/Source/Programs/AutomationTool/Scripts/GUBP.Automation.cs
--- a/Engine/Source/Programs/AutomationTool/Scripts/GUBP.Automation.cs
+++ b/Engine/Source/Programs/AutomationTool/Scripts/GUBP.Automation.cs
@@ -906,7 +906,11 @@ public class GUBP : BuildCommand
         {
             var Agenda = new UE4Build.BuildAgenda();
 
-            string AddArgs = "-nobuilduht -precompile";
+            string AddArgs = "-nobuilduht";
+			if(!bp.BranchOptions.bNoInstalledEngine)
+			{
+				AddArgs += " -precompile";
+			}
             if (bp.bOrthogonalizeEditorPlatforms)
             {
                 AddArgs += " -skipnonhostplatforms";
@@ -1570,7 +1574,12 @@ public class GUBP : BuildCommand
         {
             var Agenda = new UE4Build.BuildAgenda();
 
-            string Args = "-nobuilduht -skipactionhistory -skipnonhostplatforms -CopyAppBundleBackToDevice -forceheadergeneration -precompile";
+            string Args = "-nobuilduht -skipactionhistory -skipnonhostplatforms -CopyAppBundleBackToDevice -forceheadergeneration";
+			
+			if(!bp.BranchOptions.bNoInstalledEngine)
+			{
+				Args += " -precompile";
+			}
 
 			foreach(BranchInfo.BranchUProject GameProj in GameProjects)
 			{
@@ -2276,7 +2285,10 @@ public class GUBP : BuildCommand
                     AddDependency(ToolsCrossCompileNode.StaticGetFullName(HostPlatform));
                 }
             }
-			AddDependency(MakeFeaturePacksNode.StaticGetFullName(MakeFeaturePacksNode.GetDefaultBuildPlatform(bp)));
+			if(!bp.BranchOptions.bNoInstalledEngine)
+			{
+				AddDependency(MakeFeaturePacksNode.StaticGetFullName(MakeFeaturePacksNode.GetDefaultBuildPlatform(bp)));
+			}
         }
 		public override bool IsSeparatePromotable()
 		{
@@ -6097,7 +6109,7 @@ public class GUBP : BuildCommand
 				AddNode(new SharedCookAggregateNode(this, HostPlatforms, NonCodeProjectNames, NonCodeFormalBuilds));
 			}
 			
-			if(HostPlatform == MakeFeaturePacksNode.GetDefaultBuildPlatform(this))
+			if(HostPlatform == MakeFeaturePacksNode.GetDefaultBuildPlatform(this) && !BranchOptions.bNoInstalledEngine)
 			{
 				AddNode(new MakeFeaturePacksNode(HostPlatform, Branch.AllProjects.Where(x => MakeFeaturePacksNode.IsFeaturePack(x))));
 			}
