Allow removing targets from internal tools nodes.

diff --git a/Engine/Source/Programs/AutomationTool/Scripts/GUBP.Automation.cs b/Engine/Source/Programs/AutomationTool/Scripts/GUBP.Automation.cs
--- a/Engine/Source/Programs/AutomationTool/Scripts/GUBP.Automation.cs
+++ b/Engine/Source/Programs/AutomationTool/Scripts/GUBP.Automation.cs
@@ -5711,7 +5711,7 @@ public class GUBP : BuildCommand
 					bool SeparateNode;
 					bool CrossCompile;
 
-					if (ProgramTarget.Rules.GUBP_AlwaysBuildWithTools(HostPlatform, out bInternalOnly, out SeparateNode, out CrossCompile) && ProgramTarget.Rules.SupportsPlatform(HostPlatform) && SeparateNode)
+					if (ProgramTarget.Rules.GUBP_AlwaysBuildWithTools(HostPlatform, out bInternalOnly, out SeparateNode, out CrossCompile) && ProgramTarget.Rules.SupportsPlatform(HostPlatform) && SeparateNode && !BranchOptions.ExcludeNodes.Contains(ProgramTarget.TargetName))
 					{
 						if (bInternalOnly)
 						{
@@ -5735,7 +5735,7 @@ public class GUBP : BuildCommand
 						bool SeparateNode;
 						bool CrossCompile;
 
-						if(ProgramTarget.Rules.GUBP_AlwaysBuildWithTools(HostPlatform, out bInternalNodeOnly, out SeparateNode, out CrossCompile) && ProgramTarget.Rules.SupportsPlatform(HostPlatform) && SeparateNode)
+						if(ProgramTarget.Rules.GUBP_AlwaysBuildWithTools(HostPlatform, out bInternalNodeOnly, out SeparateNode, out CrossCompile) && ProgramTarget.Rules.SupportsPlatform(HostPlatform) && SeparateNode && !BranchOptions.ExcludeNodes.Contains(ProgramTarget.TargetName))
 						{
 							if(bInternalNodeOnly)
 							{
