Allow game branches to remove monolithic dependencies for cook nodes.

diff --git a/Engine/Source/Programs/AutomationTool/Scripts/GUBP.Automation.cs b/Engine/Source/Programs/AutomationTool/Scripts/GUBP.Automation.cs
--- a/Engine/Source/Programs/AutomationTool/Scripts/GUBP.Automation.cs
+++ b/Engine/Source/Programs/AutomationTool/Scripts/GUBP.Automation.cs
@@ -208,6 +208,7 @@ public class GUBP : BuildCommand
 			public bool bNoDocumentation = false;
 			public bool bNoInstalledEngine = false;
 			public bool bMakeFormalBuildWithoutLabelPromotable = false;
+			public bool bNoMonolithicDependenciesForCooks = false;
 			public Dictionary<string, sbyte> FrequencyBarriers = new Dictionary<string,sbyte>();
 			public int QuantumOverride = 0;
         }
@@ -2752,7 +2753,10 @@ public class GUBP : BuildCommand
                 {
                     bIsShared = true;
                 }
-                AddPseudodependency(GamePlatformMonolithicsNode.StaticGetFullName(HostPlatform, GameProj, TargetPlatform));
+				if(!bp.BranchOptions.bNoMonolithicDependenciesForCooks)
+				{
+					AddPseudodependency(GamePlatformMonolithicsNode.StaticGetFullName(HostPlatform, GameProj, TargetPlatform));
+				}
             }
             else
             {
