[INTEGRATE] Change 2554277 by [REDACTED_PERSON][REDACTED_HANDLE] on 2015/05/17 09:44:35

	Remove pseudo-dependency from tools node to root editor node for Orion, since it prevents us starting cooks sooner.

diff --git a/Engine/Source/Programs/AutomationTool/Scripts/GUBP.Automation.cs b/Engine/Source/Programs/AutomationTool/Scripts/GUBP.Automation.cs
--- a/Engine/Source/Programs/AutomationTool/Scripts/GUBP.Automation.cs
+++ b/Engine/Source/Programs/AutomationTool/Scripts/GUBP.Automation.cs
@@ -209,6 +209,7 @@ public class GUBP : BuildCommand
 			public bool bNoInstalledEngine = false;
 			public bool bMakeFormalBuildWithoutLabelPromotable = false;
 			public bool bNoMonolithicDependenciesForCooks = false;
+			public bool bNoEditorDependenciesForTools = false;
 			public Dictionary<string, sbyte> FrequencyBarriers = new Dictionary<string,sbyte>();
 			public int QuantumOverride = 0;
         }
@@ -1109,10 +1110,13 @@ public class GUBP : BuildCommand
 	}
     public class ToolsNode : CompileNode
     {
-        public ToolsNode(UnrealTargetPlatform InHostPlatform)
+        public ToolsNode(GUBP bp, UnrealTargetPlatform InHostPlatform)
             : base(InHostPlatform)
         {
-            AddPseudodependency(RootEditorNode.StaticGetFullName(HostPlatform));
+			if(!bp.BranchOptions.bNoEditorDependenciesForTools)
+			{
+				AddPseudodependency(RootEditorNode.StaticGetFullName(HostPlatform));
+			}
             AgentSharingGroup = "ToolsGroup" + StaticGetHostPlatformSuffix(HostPlatform);
         }
         public static string StaticGetFullName(UnrealTargetPlatform InHostPlatform)
@@ -1300,10 +1304,13 @@ public class GUBP : BuildCommand
 
     public class InternalToolsNode : CompileNode
     {
-        public InternalToolsNode(UnrealTargetPlatform InHostPlatform)
+        public InternalToolsNode(GUBP bp, UnrealTargetPlatform InHostPlatform)
             : base(InHostPlatform)
         {
-            AddPseudodependency(RootEditorNode.StaticGetFullName(HostPlatform));
+			if(!bp.BranchOptions.bNoEditorDependenciesForTools)
+			{
+				AddPseudodependency(RootEditorNode.StaticGetFullName(HostPlatform));
+			}
             AgentSharingGroup = "ToolsGroup" + StaticGetHostPlatformSuffix(HostPlatform);
         }
         public static string StaticGetFullName(UnrealTargetPlatform InHostPlatform)
@@ -1398,7 +1405,7 @@ public class GUBP : BuildCommand
         {
 			// Don't add rooteditor dependency if it isn't in the graph
 			var bRootEditorNodeDoesExit = bp.HasNode(RootEditorNode.StaticGetFullName(HostPlatform));
-			SetupSingleInternalToolsNode(InProgramTarget, !bRootEditorNodeDoesExit);
+			SetupSingleInternalToolsNode(InProgramTarget, !bRootEditorNodeDoesExit && bp.BranchOptions.bNoEditorDependenciesForTools);
         }
 		private void SetupSingleInternalToolsNode(SingleTargetProperties InProgramTarget, bool bSkipRootEditorPsuedoDependency = true)
 		{
@@ -5700,9 +5707,9 @@ public class GUBP : BuildCommand
 			if (!BranchOptions.ExcludePlatformsForEditor.Contains(HostPlatform))
 			{
 				AddNode(new RootEditorNode(HostPlatform));			
-				AddNode(new ToolsNode(HostPlatform));            
-				AddNode(new InternalToolsNode(HostPlatform));
-			if (HostPlatform == UnrealTargetPlatform.Win64 && ActivePlatforms.Contains(UnrealTargetPlatform.Linux))
+				AddNode(new ToolsNode(this, HostPlatform));            
+				AddNode(new InternalToolsNode(this, HostPlatform));
+				if (HostPlatform == UnrealTargetPlatform.Win64 && ActivePlatforms.Contains(UnrealTargetPlatform.Linux))
 				{
 					if (!BranchOptions.ExcludePlatformsForEditor.Contains(UnrealTargetPlatform.Linux))
 					{
