PR #986: Search IncrediBuild xgConsole.exe in PATH enviroment variable (Contributed by [REDACTED_PERSON])

diff --git a/Engine/Source/Programs/AutomationTool/UE4Build.cs b/Engine/Source/Programs/AutomationTool/UE4Build.cs
--- a/Engine/Source/Programs/AutomationTool/UE4Build.cs
+++ b/Engine/Source/Programs/AutomationTool/UE4Build.cs
@@ -18,6 +18,8 @@ namespace AutomationTool
 	{
 		private BuildCommand OwnerCommand;
 
+		private String XGEConsoleExePath = "";
+
 		public void PrepareBuildProduct(string File, bool DeleteProduct = true)
 		{
 			if (!BuildProductFiles.Contains(File))
@@ -747,13 +749,26 @@ namespace AutomationTool
 		}
 		public string XGEConsoleExe()
 		{
-			return @"C:\Program Files (x86)\Xoreax\IncrediBuild\xgConsole.exe";
+			if (string.IsNullOrEmpty(XGEConsoleExePath))
+			{
+				string[] PathDirs = Environment.GetEnvironmentVariable("PATH").Split(Path.PathSeparator);
+				foreach (string PathDir in PathDirs)
+				{
+					string FullPath = Path.Combine(PathDir, "xgConsole.exe");
+					if (FileExists(FullPath))
+					{
+						XGEConsoleExePath = FullPath;
+						break;
+					}
+				}
+			}
+			return XGEConsoleExePath;
 		}
 
 		public bool RunXGE(List<XGEItem> Actions, string TaskFilePath, bool DoRetries, bool SpecialTestFlag)
 		{
 			string XGEConsole = XGEConsoleExe();
-			if (!FileExists(XGEConsole))
+			if (string.IsNullOrEmpty(XGEConsole))
 			{
 				throw new AutomationException("Unable to find xge executable: " + XGEConsole);
 			}
@@ -1281,7 +1296,7 @@ namespace AutomationTool
 			bool bUsedXGE = false;
 			bool bCanUseXGE = !OwnerCommand.ParseParam("NoXGE") && !InForceNoXGE;
 			string XGEConsole = XGEConsoleExe();
-			if (bCanUseXGE && !FileExists(XGEConsole))
+			if (bCanUseXGE && string.IsNullOrEmpty(XGEConsole))
 			{
 				Log("XGE was requested, but is unavailable, so we won't use it.");
 				bCanUseXGE = false;
