Changing the commandline used to runall tests to use the new commandline argument makeup.

diff --git a/Engine/Source/Programs/AutomationTool/Scripts/RunProjectCommand.Automation.cs b/Engine/Source/Programs/AutomationTool/Scripts/RunProjectCommand.Automation.cs
--- a/Engine/Source/Programs/AutomationTool/Scripts/RunProjectCommand.Automation.cs
+++ b/Engine/Source/Programs/AutomationTool/Scripts/RunProjectCommand.Automation.cs
@@ -845,11 +845,11 @@ public partial class Project : CommandUtils
 		}
 		else if (Params.RunAutomationTest != "")
 		{
-			TempCmdLine += "-execcmds=\"automation list, automation run " + Params.RunAutomationTest + "\" ";
+			TempCmdLine += "-execcmds=\"automation list;runtests " + Params.RunAutomationTest + "\" ";
 		}
 		else if (Params.RunAutomationTests)
 		{
-			TempCmdLine += "-execcmds=\"automation list, automation runall\" ";
+			TempCmdLine += "-execcmds=\"automation list;runall\" ";
 		}
 		if (SC.StageTargetPlatform.UseAbsLog)
 		{
