Copy raw EXE and DLL files to the network so that crash reporter can use them.

diff --git a/Engine/Source/Programs/AutomationTool/Scripts/RocketBuild.Automation.cs b/Engine/Source/Programs/AutomationTool/Scripts/RocketBuild.Automation.cs
--- a/Engine/Source/Programs/AutomationTool/Scripts/RocketBuild.Automation.cs
+++ b/Engine/Source/Programs/AutomationTool/Scripts/RocketBuild.Automation.cs
@@ -903,7 +903,7 @@ namespace Rocket
 				foreach(string InputFileName in AllDependencyBuildProducts)
 				{
 					string Extension = Path.GetExtension(InputFileName);
-					if(DebugExtensions.Contains(Extension))
+					if(DebugExtensions.Contains(Extension) || Extension == ".exe" || Extension == ".dll") // Need all windows build products for crash reporter
 					{
 						string OutputFileName = CommandUtils.MakeRerootedFilePath(InputFileName, CommandUtils.CmdEnv.LocalRoot, SymbolsOutputDir);
 						CommandUtils.CopyFile(InputFileName, OutputFileName);
