Added word count and conflict report generation step into UE4 localization configuration scripts. Deleted word count localization configuration script and removed it from localization automation.

diff --git a/Engine/Config/Localization/Category.ini b/Engine/Config/Localization/Category.ini
--- a/Engine/Config/Localization/Category.ini
+++ b/Engine/Config/Localization/Category.ini
@@ -48,11 +48,11 @@ ResourceName=Category.locres
 CommandletClass=InternationalizationExport
 bExportLoc=true
 
-;;Write Text Localization Report
-;[GatherTextStep6]
-;CommandletClass=GenerateTextLocalizationReport
-;DestinationPath=\\epicgames.net\root\UE3\Localization\WordCounts
-;bWordCountReport=true
-;WordCountReportName=Category.csv
-;bConflictReport=true
-;ConflictReportName=CategoryConflicts.txt
\ No newline at end of file
+;Write Text Localization Report
+[GatherTextStep6]
+CommandletClass=GenerateTextLocalizationReport
+DestinationPath=\\epicgames.net\root\UE3\Localization\WordCounts
+bWordCountReport=true
+WordCountReportName=Category.csv
+bConflictReport=true
+ConflictReportName=CategoryConflicts.txt
\ No newline at end of file
diff --git a/Engine/Config/Localization/Editor.ini b/Engine/Config/Localization/Editor.ini
--- a/Engine/Config/Localization/Editor.ini
+++ b/Engine/Config/Localization/Editor.ini
@@ -72,11 +72,11 @@ ResourceName=Editor.locres
 CommandletClass=InternationalizationExport
 bExportLoc=true
 
-;;Write Text Localization Report
-;[GatherTextStep8]
-;CommandletClass=GenerateTextLocalizationReport
-;DestinationPath=\\epicgames.net\root\UE3\Localization\WordCounts
-;bWordCountReport=true
-;WordCountReportName=Editor.csv
-;bConflictReport=true
-;ConflictReportName=EditorConflicts.txt
\ No newline at end of file
+;Write Text Localization Report
+[GatherTextStep8]
+CommandletClass=GenerateTextLocalizationReport
+DestinationPath=\\epicgames.net\root\UE3\Localization\WordCounts
+bWordCountReport=true
+WordCountReportName=Editor.csv
+bConflictReport=true
+ConflictReportName=EditorConflicts.txt
\ No newline at end of file
diff --git a/Engine/Config/Localization/EditorTutorials.ini b/Engine/Config/Localization/EditorTutorials.ini
--- a/Engine/Config/Localization/EditorTutorials.ini
+++ b/Engine/Config/Localization/EditorTutorials.ini
@@ -46,11 +46,11 @@ ResourceName=EditorTutorials.locres
 CommandletClass=InternationalizationExport
 bExportLoc=true
 
-;;Write Text Localization Report
-;[GatherTextStep6]
-;CommandletClass=GenerateTextLocalizationReport
-;DestinationPath=\\epicgames.net\root\UE3\Localization\WordCounts
-;bWordCountReport=true
-;WordCountReportName=EditorTutorials.csv
-;bConflictReport=true
-;ConflictReportName=EditorTutorialsConflicts.txt
\ No newline at end of file
+;Write Text Localization Report
+[GatherTextStep6]
+CommandletClass=GenerateTextLocalizationReport
+DestinationPath=\\epicgames.net\root\UE3\Localization\WordCounts
+bWordCountReport=true
+WordCountReportName=EditorTutorials.csv
+bConflictReport=true
+ConflictReportName=EditorTutorialsConflicts.txt
\ No newline at end of file
diff --git a/Engine/Config/Localization/Engine.ini b/Engine/Config/Localization/Engine.ini
--- a/Engine/Config/Localization/Engine.ini
+++ b/Engine/Config/Localization/Engine.ini
@@ -72,11 +72,11 @@ ResourceName=Engine.locres
 CommandletClass=InternationalizationExport
 bExportLoc=true
 
-;;Write Text Localization Report
-;[GatherTextStep7]
-;CommandletClass=GenerateTextLocalizationReport
-;DestinationPath=\\partnersdc\NetworkShare\Users\SaulDavidAbreu
-;bWordCountReport=true
-;WordCountReportName=Engine.csv
-;bConflictReport=true
-;ConflictReportName=EngineConflicts.txt
\ No newline at end of file
+;Write Text Localization Report
+[GatherTextStep7]
+CommandletClass=GenerateTextLocalizationReport
+DestinationPath=\\epicgames.net\root\UE3\Localization\WordCounts
+bWordCountReport=true
+WordCountReportName=Engine.csv
+bConflictReport=true
+ConflictReportName=EngineConflicts.txt
\ No newline at end of file
diff --git a/Engine/Config/Localization/Keywords.ini b/Engine/Config/Localization/Keywords.ini
--- a/Engine/Config/Localization/Keywords.ini
+++ b/Engine/Config/Localization/Keywords.ini
@@ -48,11 +48,11 @@ ResourceName=Keywords.locres
 CommandletClass=InternationalizationExport
 bExportLoc=true
 
-;;Write Text Localization Report
-;[GatherTextStep6]
-;CommandletClass=GenerateTextLocalizationReport
-;DestinationPath=\\epicgames.net\root\UE3\Localization\WordCounts
-;bWordCountReport=true
-;WordCountReportName=Keywords.csv
-;bConflictReport=true
-;ConflictReportName=KeywordsConflicts.txt
\ No newline at end of file
+;Write Text Localization Report
+[GatherTextStep6]
+CommandletClass=GenerateTextLocalizationReport
+DestinationPath=\\epicgames.net\root\UE3\Localization\WordCounts
+bWordCountReport=true
+WordCountReportName=Keywords.csv
+bConflictReport=true
+ConflictReportName=KeywordsConflicts.txt
\ No newline at end of file
diff --git a/Engine/Config/Localization/PropertyNames.ini b/Engine/Config/Localization/PropertyNames.ini
--- a/Engine/Config/Localization/PropertyNames.ini
+++ b/Engine/Config/Localization/PropertyNames.ini
@@ -48,11 +48,11 @@ ResourceName=PropertyNames.locres
 CommandletClass=InternationalizationExport
 bExportLoc=true
 
-;;Write Text Localization Report
-;[GatherTextStep6]
-;CommandletClass=GenerateTextLocalizationReport
-;DestinationPath=\\epicgames.net\root\UE3\Localization\WordCounts
-;bWordCountReport=true
-;WordCountReportName=PropertyNames.csv
-;bConflictReport=true
-;ConflictReportName=PropertyNamesConflicts.txt
\ No newline at end of file
+;Write Text Localization Report
+[GatherTextStep6]
+CommandletClass=GenerateTextLocalizationReport
+DestinationPath=\\epicgames.net\root\UE3\Localization\WordCounts
+bWordCountReport=true
+WordCountReportName=PropertyNames.csv
+bConflictReport=true
+ConflictReportName=PropertyNamesConflicts.txt
\ No newline at end of file
diff --git a/Engine/Config/Localization/ToolTips.ini b/Engine/Config/Localization/ToolTips.ini
--- a/Engine/Config/Localization/ToolTips.ini
+++ b/Engine/Config/Localization/ToolTips.ini
@@ -51,11 +51,11 @@ ResourceName=ToolTips.locres
 CommandletClass=InternationalizationExport
 bExportLoc=true
 
-;;Write Text Localization Report
-;[GatherTextStep6]
-;CommandletClass=GenerateTextLocalizationReport
-;DestinationPath=\\epicgames.net\root\UE3\Localization\WordCounts
-;bWordCountReport=true
-;WordCountReportName=ToolTips.csv
-;bConflictReport=true
-;ConflictReportName=ToolTipsConflicts.txt
\ No newline at end of file
+;Write Text Localization Report
+[GatherTextStep6]
+CommandletClass=GenerateTextLocalizationReport
+DestinationPath=\\epicgames.net\root\UE3\Localization\WordCounts
+bWordCountReport=true
+WordCountReportName=ToolTips.csv
+bConflictReport=true
+ConflictReportName=ToolTipsConflicts.txt
\ No newline at end of file
diff --git a/Engine/Source/Programs/AutomationTool/Scripts/Localisation.Automation.cs b/Engine/Source/Programs/AutomationTool/Scripts/Localisation.Automation.cs
--- a/Engine/Source/Programs/AutomationTool/Scripts/Localisation.Automation.cs
+++ b/Engine/Source/Programs/AutomationTool/Scripts/Localisation.Automation.cs
@@ -91,7 +91,6 @@ class Localise : BuildCommand
                 String.Format("-config={0}", @"./Config/Localization/ToolTips.ini") + (string.IsNullOrEmpty(CommandletSCCArguments) ? "" : " " + CommandletSCCArguments),
                 String.Format("-config={0}", @"./Config/Localization/Category.ini") + (string.IsNullOrEmpty(CommandletSCCArguments) ? "" : " " + CommandletSCCArguments),
 				String.Format("-config={0}", @"./Config/Localization/Keywords.ini") + (string.IsNullOrEmpty(CommandletSCCArguments) ? "" : " " + CommandletSCCArguments),
-				String.Format("-config={0}", @"./Config/Localization/WordCount.ini"),
 			};
 
 		// Execute commandlet for each set of arguments.
