[INTEGRATE] Change 2557122 by [REDACTED_PERSON][REDACTED_HANDLE] on 2015/05/19 11:38:07

	Make all staged plugin files writable after we copy them.

diff --git a/Engine/Source/Programs/AutomationTool/Scripts/BuildPluginCommand.Automation.cs b/Engine/Source/Programs/AutomationTool/Scripts/BuildPluginCommand.Automation.cs
--- a/Engine/Source/Programs/AutomationTool/Scripts/BuildPluginCommand.Automation.cs
+++ b/Engine/Source/Programs/AutomationTool/Scripts/BuildPluginCommand.Automation.cs
@@ -131,6 +131,7 @@ class BuildPlugin : BuildCommand
 			string SourceFileName = Path.Combine(Path.GetDirectoryName(PluginFileName), MatchingFileName);
 			string TargetFileName = Path.Combine(PackageDirectory, MatchingFileName);
 			CommandUtils.CopyFile(SourceFileName, TargetFileName);
+			CommandUtils.SetFileAttributes(TargetFileName, false);
 		}
 
 		// Get the output plugin filename
