Add named parameter to SetFileAttributes call in BuildPlugin UAT script to clarify intent.

diff --git a/Engine/Source/Programs/AutomationTool/Scripts/BuildPluginCommand.Automation.cs b/Engine/Source/Programs/AutomationTool/Scripts/BuildPluginCommand.Automation.cs
--- a/Engine/Source/Programs/AutomationTool/Scripts/BuildPluginCommand.Automation.cs
+++ b/Engine/Source/Programs/AutomationTool/Scripts/BuildPluginCommand.Automation.cs
@@ -131,7 +131,7 @@ class BuildPlugin : BuildCommand
 			string SourceFileName = Path.Combine(Path.GetDirectoryName(PluginFileName), MatchingFileName);
 			string TargetFileName = Path.Combine(PackageDirectory, MatchingFileName);
 			CommandUtils.CopyFile(SourceFileName, TargetFileName);
-			CommandUtils.SetFileAttributes(TargetFileName, false);
+			CommandUtils.SetFileAttributes(TargetFileName, ReadOnly: false);
 		}
 
 		// Get the output plugin filename
