fix for pak files from a previous run sticking around during incremental deploy
UE-13112

diff --git a/Engine/Source/Programs/AutomationTool/Scripts/CopyBuildToStagingDirectory.Automation.cs b/Engine/Source/Programs/AutomationTool/Scripts/CopyBuildToStagingDirectory.Automation.cs
--- a/Engine/Source/Programs/AutomationTool/Scripts/CopyBuildToStagingDirectory.Automation.cs
+++ b/Engine/Source/Programs/AutomationTool/Scripts/CopyBuildToStagingDirectory.Automation.cs
@@ -8,6 +8,7 @@ using System.Net.NetworkInformation;
 using System.Reflection;
 using System.Text.RegularExpressions;
 using System.Threading;
+using System.Linq;
 using UnrealBuildTool;
 
 /// <summary>
@@ -963,6 +964,12 @@ public partial class Project : CommandUtils
         }
 	}
 
+	protected static void DeletePakFiles(string StagingDirectory)
+	{
+		var StagedFilesDir = new DirectoryInfo(StagingDirectory);
+		StagedFilesDir.GetFiles("*.pak", SearchOption.AllDirectories).ToList().ForEach(File => File.Delete());
+	}
+
 	public static void ApplyStagingManifest(ProjectParams Params, DeploymentContext SC)
 	{
 		MaybeConvertToLowerCase(Params, SC);
@@ -978,6 +985,21 @@ public partial class Project : CommandUtils
                 Log("Failed to delete staging directory "+SC.StageDirectory);
                 AutomationTool.ErrorReporter.Error("Stage Failed.", (int)AutomationTool.ErrorCodes.Error_FailedToDeleteStagingDirectory);
                 throw Ex;   
+            }
+		}
+		else
+		{
+			try
+			{
+				// delete old pak files
+				DeletePakFiles(SC.StageDirectory);
+			}
+            catch (Exception Ex)
+            {
+                // Delete cooked data (if any) as it may be incomplete / corrupted.
+                Log("Failed to delete pak files in "+SC.StageDirectory);
+                AutomationTool.ErrorReporter.Error("Stage Failed.", (int)AutomationTool.ErrorCodes.Error_FailedToDeleteStagingDirectory);
+                throw Ex;   
             }
 		}
 		if (ShouldCreatePak(Params, SC))
