build script, fix sticky preconditions being added as runconditions

diff --git a/Engine/Source/Programs/AutomationTool/Scripts/GUBP.Automation.cs b/Engine/Source/Programs/AutomationTool/Scripts/GUBP.Automation.cs
--- a/Engine/Source/Programs/AutomationTool/Scripts/GUBP.Automation.cs
+++ b/Engine/Source/Programs/AutomationTool/Scripts/GUBP.Automation.cs
@@ -7260,7 +7260,7 @@ public class GUBP : BuildCommand
                     string RunCondition = "";
                     var UncompletedEcDeps = new List<string>();
 					var UncompletedCompletedDeps = new List<string>();
-
+                    var PreConditionUncompletedEcDeps = new List<string>();
                     string MyAgentGroup = GUBPNodes[NodeToDo].AgentSharingGroup;
                     bool bDoNestedJobstep = false;
                     bool bDoFirstNestedJobstep = false;
@@ -7290,7 +7290,10 @@ public class GUBP : BuildCommand
                             }
                         }
                     }
-                    var PreConditionUncompletedEcDeps = UncompletedEcDeps;
+                    foreach(var Dep in UncompletedEcDeps)
+                    {
+                        PreConditionUncompletedEcDeps.Add(Dep);
+                    }
 					var CompletedDeps = GetCompletedOnlyDependencies(NodeToDo);					
 					foreach (var Dep in CompletedDeps)
 					{
