Add support for optional packaging blacklist per platform and build configuration
[Project]/Build/[Platform]/PakBlacklist-[Config].txt
#ue4

diff --git a/Engine/Source/Programs/AutomationTool/Scripts/CopyBuildToStagingDirectory.Automation.cs b/Engine/Source/Programs/AutomationTool/Scripts/CopyBuildToStagingDirectory.Automation.cs
--- a/Engine/Source/Programs/AutomationTool/Scripts/CopyBuildToStagingDirectory.Automation.cs
+++ b/Engine/Source/Programs/AutomationTool/Scripts/CopyBuildToStagingDirectory.Automation.cs
@@ -625,6 +625,18 @@ public partial class Project : CommandUtils
 	/// <returns></returns>
 	private static Dictionary<string, string> CreatePakResponseFileFromStagingManifest(DeploymentContext SC)
 	{
+		// look for optional packaging blacklist if only one config active
+		string[] Blacklist = null;
+		if (SC.StageTargetConfigurations.Count == 1)
+		{
+			var PakBlacklistFilename = CombinePaths(SC.ProjectRoot, "Build", SC.PlatformDir, string.Format("PakBlacklist-{0}.txt", SC.StageTargetConfigurations[0].ToString()));
+			if (File.Exists(PakBlacklistFilename))
+			{
+				Log("Applying PAK blacklist file {0}", PakBlacklistFilename);
+				Blacklist = File.ReadAllLines(PakBlacklistFilename);
+			}
+		}
+
 		var UnrealPakResponseFile = new Dictionary<string, string>(StringComparer.InvariantCultureIgnoreCase);
 		foreach (var Pair in SC.UFSStagingFiles)
 		{
@@ -633,6 +645,25 @@ public partial class Project : CommandUtils
 
 			Dest = CombinePaths(PathSeparator.Slash, SC.PakFileInternalRoot, Dest);
 
+			if (Blacklist != null)
+			{
+				bool bExcludeFile = false;
+				foreach (string ExcludePath in Blacklist)
+				{
+					if (Dest.StartsWith(ExcludePath))
+					{
+						bExcludeFile = true;
+						break;
+					}
+				}
+
+				if (bExcludeFile)
+				{
+					Log("Excluding {0}", Src);
+					continue;
+				}
+			}
+
 			// there can be files that only differ in case only, we don't support that in paks as paks are case-insensitive
 			if (UnrealPakResponseFile.ContainsKey(Src))
 			{
@@ -640,6 +671,7 @@ public partial class Project : CommandUtils
 			}
 			UnrealPakResponseFile.Add(Src, Dest);
 		}
+
 		return UnrealPakResponseFile;
 	}
 
