fix for target.cs files not being in the root of the source directory for some games

diff --git a/Engine/Source/Programs/AutomationTool/ProjectUtils.cs b/Engine/Source/Programs/AutomationTool/ProjectUtils.cs
--- a/Engine/Source/Programs/AutomationTool/ProjectUtils.cs
+++ b/Engine/Source/Programs/AutomationTool/ProjectUtils.cs
@@ -158,10 +158,14 @@ namespace AutomationTool
 		private static bool RequiresTempTarget(string RawProjectPath, List<UnrealTargetPlatform> ClientTargetPlatforms)
 		{
 			// check to see if we already have a Target.cs file
-			// @todo: Target.cs may not be the same name as the project in the future, so look at a better way to determine this
 			if (File.Exists(Path.Combine(Path.GetDirectoryName(RawProjectPath), "Source", Path.GetFileNameWithoutExtension(RawProjectPath) + ".Target.cs")))
 			{
-				return false;
+				//@todo: may want to read each target.cs to see if it has a target corresponding to the project name as a final check
+				FileInfo[] Files = (new DirectoryInfo(Path.GetDirectoryName(RawProjectPath)).GetFiles("*.target.cs", SearchOption.AllDirectories));
+				if (Files.Length > 0)
+				{
+					return false;
+				}
 			}
 
 			// no Target file, now check to see if build settings have changed
