Add support for progress feedback when using UE4Build. Specify the -progress parameter on the command line to MegaXGE to write progress data into the log.

diff --git a/Engine/Source/Programs/AutomationTool/CommandUtils.cs b/Engine/Source/Programs/AutomationTool/CommandUtils.cs
--- a/Engine/Source/Programs/AutomationTool/CommandUtils.cs
+++ b/Engine/Source/Programs/AutomationTool/CommandUtils.cs
@@ -234,6 +234,50 @@ namespace AutomationTool
 
 		#endregion
 
+		#region Progress Logging
+
+		public static void LogPushProgress(bool bShowProgress, int Numerator, int Denominator)
+		{
+			if(bShowProgress)
+			{
+				Log("[@progress push {0}/{1} skipline]", Numerator, Denominator);
+			}
+		}
+
+		public static void LogPopProgress(bool bShowProgress)
+		{
+			if(bShowProgress)
+			{
+				Log("[@progress pop skipline]");
+			}
+		}
+
+		public static void LogIncrementProgress(bool bShowProgress, int Numerator, int Denominator)
+		{
+			if(bShowProgress)
+			{
+				Log("[@progress increment {0}/{1} skipline]", Numerator, Denominator);
+			}
+		}
+
+		public static void LogSetProgress(bool bShowProgress, string Format, params string[] Args)
+		{
+			if(bShowProgress)
+			{
+				Log("[@progress '{0}' skipline]", String.Format(Format, Args));
+			}
+		}
+
+		public static void LogSetProgress(bool bShowProgress, int Numerator, int Denominator, string Format, params string[] Args)
+		{
+			if(bShowProgress)
+			{
+				Log("[@progress {0}/{1} '{2}' skipline]", Numerator, Denominator, String.Format(Format, Args));
+			}
+		}
+
+		#endregion
+
 		#region IO
 
 		/// <summary>
diff --git a/Engine/Source/Programs/AutomationTool/Scripts/MegaXGE.Automation.cs b/Engine/Source/Programs/AutomationTool/Scripts/MegaXGE.Automation.cs
--- a/Engine/Source/Programs/AutomationTool/Scripts/MegaXGE.Automation.cs
+++ b/Engine/Source/Programs/AutomationTool/Scripts/MegaXGE.Automation.cs
@@ -31,6 +31,8 @@ class MegaXGE : BuildCommand
 		bool Clean = ParseParam("Clean");
 		string CleanToolLocation = CombinePaths(CmdEnv.LocalRoot, "Engine", "Build", "Batchfiles", "Clean.bat");
 
+		bool ShowProgress = ParseParam("Progress");
+
 		var UE4Build = new UE4Build(this);
 
 		var Agenda = new UE4Build.BuildAgenda();
@@ -43,6 +45,7 @@ class MegaXGE : BuildCommand
 		Agenda.AddTargets(ProgramTargets, UnrealTargetPlatform.Win64, UnrealTargetConfiguration.Development);
 		if (Clean)
 		{
+			LogSetProgress(ShowProgress, "Cleaning previous builds...");
 			foreach (var CurTarget in ProgramTargets)
 			{
 				string Args = String.Format("{0} {1} {2}", CurTarget, UnrealTargetPlatform.Win64.ToString(), UnrealTargetConfiguration.Development.ToString());
@@ -142,7 +145,7 @@ class MegaXGE : BuildCommand
 		Log("*************************");
 
 		Agenda.DoRetries = ParseParam("Retry");
-		UE4Build.Build(Agenda, InUpdateVersionFiles: IsBuildMachine);
+		UE4Build.Build(Agenda, InUpdateVersionFiles: IsBuildMachine, InShowProgress: ShowProgress);
 
 		// 		if (WorkingCL > 0) // only move UAT files if we intend to check in some build products
 		// 		{
diff --git a/Engine/Source/Programs/AutomationTool/UE4Build.cs b/Engine/Source/Programs/AutomationTool/UE4Build.cs
--- a/Engine/Source/Programs/AutomationTool/UE4Build.cs
+++ b/Engine/Source/Programs/AutomationTool/UE4Build.cs
@@ -765,7 +765,7 @@ namespace AutomationTool
 			return XGEConsoleExePath;
 		}
 
-		public bool RunXGE(List<XGEItem> Actions, string TaskFilePath, bool DoRetries, bool SpecialTestFlag)
+		public bool RunXGE(List<XGEItem> Actions, string TaskFilePath, bool DoRetries, bool SpecialTestFlag, bool ShowProgress)
 		{
 			string XGEConsole = XGEConsoleExe();
 			if (string.IsNullOrEmpty(XGEConsole))
@@ -1075,6 +1075,34 @@ namespace AutomationTool
 			}
 			else
 			{
+				// Make sure all the tasks have a unique prefix
+				if(ShowProgress)
+				{
+					List<XmlElement> AllToolElements = new List<XmlElement>();
+					foreach(XmlElement EnvironmentElement in EnvironmentsElement.GetElementsByTagName("Environment"))
+					{
+						foreach(XmlElement ToolsElement in EnvironmentElement.GetElementsByTagName("Tools"))
+						{
+							foreach(XmlElement ToolElement in ToolsElement.GetElementsByTagName("Tool"))
+							{
+								AllToolElements.Add(ToolElement);
+							}
+						}
+					}
+					for(int Idx = 0; Idx < AllToolElements.Count; Idx++)
+					{
+						XmlElement ToolElement = AllToolElements[Idx];
+						if (ToolElement.HasAttribute("OutputPrefix"))
+						{
+							ToolElement.SetAttribute("OutputPrefix", ToolElement.Attributes["OutputPrefix"].Value + String.Format(" [@progress increment 1/{0}]", AllToolElements.Count));
+						}
+						else
+						{
+							ToolElement.SetAttribute("OutputPrefix", String.Format(" [@progress increment 1/{0} skipline]", AllToolElements.Count));
+						}
+					}
+				}
+
 				// Write the XGE task XML to a temporary file.
 				using (FileStream OutputFileStream = new FileStream(TaskFilePath, FileMode.Create, FileAccess.Write))
 				{
@@ -1187,7 +1215,7 @@ namespace AutomationTool
 		/// <param name="InDeleteBuildProducts">if specified, determines if the build products will be deleted before building. If not specified -clean parameter will be used,</param>
 		/// <param name="InUpdateVersionFiles">True if the version files are to be updated </param>
 		/// <param name="InForceNoXGE">If true will force XGE off</param>
-		public void Build(BuildAgenda Agenda, bool? InDeleteBuildProducts = null, bool InUpdateVersionFiles = true, bool InForceNoXGE = false, bool InForceNonUnity = false, bool InForceUnity = false, Dictionary<UnrealBuildTool.UnrealTargetPlatform, Dictionary<string, string>> PlatformEnvVars = null)
+		public void Build(BuildAgenda Agenda, bool? InDeleteBuildProducts = null, bool InUpdateVersionFiles = true, bool InForceNoXGE = false, bool InForceNonUnity = false, bool InForceUnity = false, bool InShowProgress = false, Dictionary<UnrealBuildTool.UnrealTargetPlatform, Dictionary<string, string>> PlatformEnvVars = null)
 		{
 			if (!CmdEnv.HasCapabilityToCompile)
 			{
@@ -1314,7 +1342,7 @@ namespace AutomationTool
 
 				if (DeleteBuildProducts)
 				{
-					foreach (var Target in Agenda.Targets)
+					foreach(var Target in Agenda.Targets)
 					{
 						CleanWithUBT(Target.ProjectName, Target.TargetName, Target.Platform, Target.Config.ToString(), Target.UprojectPath, bForceMonolithic, bForceNonUnity, bForceDebugInfo, Target.UBTArgs, bForceUnity);
 					}
@@ -1325,11 +1353,13 @@ namespace AutomationTool
 					{
 						// When building a target for Mac or iOS, use UBT's -flushmac option to clean up the remote builder
 						bool bForceFlushMac = DeleteBuildProducts && (Target.Platform == UnrealBuildTool.UnrealTargetPlatform.Mac || Target.Platform == UnrealBuildTool.UnrealTargetPlatform.IOS);
+						LogSetProgress(InShowProgress, "Building header tool...");
 						BuildWithUBT(Target.ProjectName, Target.TargetName, Target.Platform, Target.Config.ToString(), Target.UprojectPath, bForceMonolithic, bForceNonUnity, bForceDebugInfo, bForceFlushMac, true, Target.UBTArgs, bForceUnity);
 					}
 				}
 
 				List<XGEItem> XGEItems = new List<XGEItem>();
+				LogSetProgress(InShowProgress, "Generating headers...");
 				foreach (var Target in Agenda.Targets)
 				{
 					if (Target.TargetName != "UnrealHeaderTool" && CanUseXGE(Target.Platform))
@@ -1340,9 +1370,10 @@ namespace AutomationTool
 				}
 				if (XGEItems.Count > 0)
 				{
+					LogSetProgress(InShowProgress, "Building...");
                     var StartXGERun = DateTime.Now.ToString();
                     var FinishXGERun = "";
-					if (!RunXGE(XGEItems, TaskFilePath, Agenda.DoRetries, Agenda.SpecialTestFlag))
+					if (!RunXGE(XGEItems, TaskFilePath, Agenda.DoRetries, Agenda.SpecialTestFlag, InShowProgress))
 					{
                         FinishXGERun = DateTime.Now.ToString();
                         PrintCSVFile(String.Format("UAT,CompileWithXGE,{0},{1}", StartXGERun, FinishXGERun));
