Only check that the first file in a manifest exists. We're going to copy them anyway, so we might as well error there. Check for the first file existing so we can restore the network connection if it's been lost on Mac.

diff --git a/Engine/Source/Programs/AutomationTool/TempStorage.cs b/Engine/Source/Programs/AutomationTool/TempStorage.cs
--- a/Engine/Source/Programs/AutomationTool/TempStorage.cs
+++ b/Engine/Source/Programs/AutomationTool/TempStorage.cs
@@ -530,7 +530,6 @@ namespace AutomationTool
                     foreach (var ThisFileInfo in Dir.Value)
                     {
                         var NewFile = CombinePaths(NewBaseDir, ThisFileInfo.Name);
-                        Robust_FileExists_NoExceptions(false, NewFile, "Rebased manifest file does not exist {0}");
 
                         FileInfo Info = new FileInfo(NewFile);
 
@@ -541,6 +540,10 @@ namespace AutomationTool
                 {
                     throw new AutomationException("No files in attempt to GetFiles().");
                 }
+
+				Log("Checking that first file in manifest exists...");
+				Robust_FileExists_NoExceptions(false, Result[0], "Rebased manifest file does not exist {0}");
+
                 return Result;
             }
         }
