build script, remove cooks and packages that are unused by anyone - place the new list in front of a trigger, so qa have builds available, this is an option due to the significant drop in total builds

diff --git a/Engine/Source/Programs/AutomationTool/Scripts/GUBP.Automation.cs b/Engine/Source/Programs/AutomationTool/Scripts/GUBP.Automation.cs
--- a/Engine/Source/Programs/AutomationTool/Scripts/GUBP.Automation.cs
+++ b/Engine/Source/Programs/AutomationTool/Scripts/GUBP.Automation.cs
@@ -201,6 +201,7 @@ public class GUBP : BuildCommand
 			public List<string> ExcludeNodes = new List<string>();
 			public List<UnrealTargetPlatform> ExcludePlatformsForEditor = new List<UnrealTargetPlatform>();
             public List<UnrealTargetPlatform> RemovePlatformFromPromotable = new List<UnrealTargetPlatform>();
+            public List<string> ProjectsToCook = new List<string>();
             public List<string> PromotablesWithoutTools = new List<string>();
             public List<string> NodesToRemovePseudoDependencies = new List<string>();
             public List<string> EnhanceAgentRequirements = new List<string>();
@@ -2077,107 +2078,7 @@ public class GUBP : BuildCommand
             return StaticGetFullName(PromotionLabelPrefix);
         }
     }
-
-	public class SharedCookAggregateNode : AggregateNode
-	{
-		public SharedCookAggregateNode(GUBP bp, List<UnrealTargetPlatform> InHostPlatforms, Dictionary<string, List<UnrealTargetPlatform>> NonCodeProjectNames, Dictionary<string, List<TargetRules.GUBPFormalBuild>> NonCodeFormalBuilds)
-		{			
-			foreach (var HostPlatform in InHostPlatforms)
-			{
-				{
-					var Options = bp.Branch.BaseEngineProject.Options(HostPlatform);
-					if (Options.bIsPromotable && !Options.bSeparateGamePromotion)
-					{
-						var Kind = TargetRules.TargetType.Game;						
-						if (bp.Branch.BaseEngineProject.Properties.Targets.ContainsKey(Kind))
-						{
-							var Target = bp.Branch.BaseEngineProject.Properties.Targets[Kind];
-							var Platforms = Target.Rules.GUBP_GetPlatforms_MonolithicOnly(HostPlatform);
-							foreach (var Plat in Platforms)
-							{
-								if (bNoIOSOnPC && Plat == UnrealTargetPlatform.IOS && HostPlatform == UnrealTargetPlatform.Win64)
-								{
-									continue;
-								}
-								if (bp.ActivePlatforms.Contains(Plat))
-								{
-									string CookedPlatform = Platform.Platforms[Plat].GetCookPlatform(Kind == TargetRules.TargetType.Server, Kind == TargetRules.TargetType.Client, "");										
-									
-									foreach (var NonCodeProject in bp.Branch.NonCodeProjects)
-									{										
-										if (!NonCodeProjectNames.ContainsKey(NonCodeProject.GameName) || !NonCodeProjectNames[NonCodeProject.GameName].Contains(Plat))
-										{
-											continue;
-										}
-										if(NonCodeFormalBuilds.ContainsKey(NonCodeProject.GameName))
-										{
-											var PlatList = NonCodeFormalBuilds[NonCodeProject.GameName];
-											foreach (var PlatPair in PlatList)
-											{
-												if(PlatPair.TargetPlatform == Plat)
-												{													
-													AddCompletedDependency(CookNode.StaticGetFullName(HostPlatform, NonCodeProject, CookedPlatform));
-												}
-											}
-										}
-									}
-								}
-							}								
-						}						
-					}					
-				}
-				foreach (var CodeProj in bp.Branch.CodeProjects)
-				{
-					var Options = CodeProj.Options(HostPlatform);
-					if (!Options.bSeparateGamePromotion)
-					{
-						if (Options.bIsPromotable)
-						{
-							foreach (var Kind in BranchInfo.MonolithicKinds)
-							{
-								if (CodeProj.Properties.Targets.ContainsKey(Kind))
-								{
-									var Target = CodeProj.Properties.Targets[Kind];
-									var Platforms = new List<UnrealTargetPlatform>();
-									var PossPlats = Target.Rules.GUBP_GetPlatforms_MonolithicOnly(HostPlatform);
-									var FormalBuilds = Target.Rules.GUBP_GetConfigsForFormalBuilds_MonolithicOnly(HostPlatform);
-									foreach (var FormalBuild in FormalBuilds)
-									{
-										var Platform = FormalBuild.TargetPlatform;
-										if(!Platforms.Contains(Platform) && PossPlats.Contains(Platform))
-										{
-											Platforms.Add(Platform);
-										}
-									}
-									foreach (var Plat in Platforms)
-									{
-										if (bNoIOSOnPC && Plat == UnrealTargetPlatform.IOS && HostPlatform == UnrealTargetPlatform.Win64)
-										{
-											continue;
-										}
-										if (bp.ActivePlatforms.Contains(Plat))
-										{
-											string CookedPlatform = Platform.Platforms[Plat].GetCookPlatform(Kind == TargetRules.TargetType.Server, Kind == TargetRules.TargetType.Client, "");
-											AddCompletedDependency(CookNode.StaticGetFullName(HostPlatform, CodeProj, CookedPlatform));
-										}
-									}
-								}
-							}
-						}
-					}
-				}				
-			}
-		}
-		public static string StaticGetFullName()
-		{
-			return "SharedCookAggregate";
-		}
-		public override string GetFullName()
-		{
-			return StaticGetFullName();
-		}
-	}
-
+    	
     public class GameAggregatePromotableNode : AggregatePromotableNode
     {
         BranchInfo.BranchUProject GameProj;		
@@ -2378,23 +2279,6 @@ public class GUBP : BuildCommand
         }
     }
 
-	public class WaitForFormalUserInput : WaitForUserInput
-	{
-		
-		public WaitForFormalUserInput(GUBP bp)
-		{		
-			AddDependency(SharedCookAggregateNode.StaticGetFullName());
-			AddCompletedDependency(SharedCookAggregateNode.StaticGetFullName());
-		}
-		public static string StaticGetFullName()
-		{
-			return "Shared_WaitToMakeBuilds";
-		}
-		public override string GetFullName()
-		{
-			return StaticGetFullName();
-		}
-	}
     public class WaitForPromotionUserInput : WaitForUserInput
     {
         string PromotionLabelPrefix;
@@ -3104,11 +2988,7 @@ public class GUBP : BuildCommand
             }
 
             var AllTargetPlatforms = new List<UnrealTargetPlatform>();
-			var Options = InGameProj.Options(HostPlatform);
-			if(!Options.bSeparateGamePromotion && !bp.BranchOptions.bMakeFormalBuildWithoutLabelPromotable)
-			{
-				AddPseudodependency(WaitForFormalUserInput.StaticGetFullName());				
-			}
+			var Options = InGameProj.Options(HostPlatform);			
             if (ClientTargetPlatforms != null)
             {
                 if (!WorkingGameProject.Properties.Targets.ContainsKey(GameOrClient))
@@ -6008,72 +5888,70 @@ public class GUBP : BuildCommand
                                     {
                                         continue;
                                     }
-                                    if (!GUBPNodes.ContainsKey(CookNode.StaticGetFullName(HostPlatform, NonCodeProject, CookedPlatform)))
-                                    {
-                                        GameCookNodes.Add(AddNode(new CookNode(this, HostPlatform, NonCodeProject, Plat, CookedPlatform)));
-                                    }
-                                    if (!GUBPNodes.ContainsKey(GamePlatformCookedAndCompiledNode.StaticGetFullName(HostPlatform, NonCodeProject, Plat)))
+                                    if (BranchOptions.ProjectsToCook.Contains(NonCodeProject.GameName) || BranchOptions.ProjectsToCook.Count == 0)
                                     {
-                                        AddNode(new GamePlatformCookedAndCompiledNode(this, HostPlatform, NonCodeProject, Plat, false));
-                                        if (NonCodeFormalBuilds.ContainsKey(NonCodeProject.GameName))
+                                        if (!GUBPNodes.ContainsKey(CookNode.StaticGetFullName(HostPlatform, NonCodeProject, CookedPlatform)))
                                         {
+                                            GameCookNodes.Add(AddNode(new CookNode(this, HostPlatform, NonCodeProject, Plat, CookedPlatform)));
+                                        }
+                                        if (!GUBPNodes.ContainsKey(GamePlatformCookedAndCompiledNode.StaticGetFullName(HostPlatform, NonCodeProject, Plat)))
+                                        {
+                                            AddNode(new GamePlatformCookedAndCompiledNode(this, HostPlatform, NonCodeProject, Plat, false));
+
                                             var PlatList = NonCodeFormalBuilds[NonCodeProject.GameName];
                                             foreach (var PlatPair in PlatList)
                                             {
                                                 if (PlatPair.TargetPlatform == Plat)
                                                 {
-													var NodeName = AddNode(new FormalBuildNode(this, NonCodeProject, HostPlatform, new List<UnrealTargetPlatform>() { Plat }, new List<UnrealTargetConfiguration>() {PlatPair.TargetConfig}));
-													if(PlatPair.bBeforeTrigger)
-													{
-														RemovePseudodependencyFromNode(FormalBuildNode.StaticGetFullName(NonCodeProject, HostPlatform, new List<UnrealTargetPlatform>() { Plat }, new List<UnrealTargetConfiguration>() { PlatPair.TargetConfig }), WaitForFormalUserInput.StaticGetFullName());
-													}
-													// we don't want this delayed
-													// this would normally wait for the testing phase, we just want to build it right away
-													RemovePseudodependencyFromNode(
-														CookNode.StaticGetFullName(HostPlatform, NonCodeProject, CookedPlatform),
-														WaitForTestShared.StaticGetFullName());
-													string BuildAgentSharingGroup = "";
-													if (Options.bSeparateGamePromotion)
-													{
-														BuildAgentSharingGroup = NonCodeProject.GameName + "_MakeFormalBuild_" + Plat.ToString() + HostPlatformNode.StaticGetHostPlatformSuffix(HostPlatform);
-														if (Plat == UnrealTargetPlatform.IOS || Plat == UnrealTargetPlatform.Android) // These trash build products, so we need to use different agents
-														{
-															BuildAgentSharingGroup = "";
-														}
-														GUBPNodes[CookNode.StaticGetFullName(HostPlatform, NonCodeProject, CookedPlatform)].AgentSharingGroup = BuildAgentSharingGroup;
-														GUBPNodes[GamePlatformCookedAndCompiledNode.StaticGetFullName(HostPlatform, NonCodeProject, Plat)].AgentSharingGroup = BuildAgentSharingGroup;
-														GUBPNodes[NodeName].AgentSharingGroup = BuildAgentSharingGroup;
-													}
-													else
-													{
-														//GUBPNodes[NodeName].AgentSharingGroup = FormalAgentSharingGroup;
-														if(Plat == UnrealTargetPlatform.XboxOne)
-														{
-															GUBPNodes[NodeName].AgentSharingGroup = "";
-														}
-													}
-													if (PlatPair.bTest)
-													{
-														AddNode(new FormalBuildTestNode(this, NonCodeProject, HostPlatform, Plat, PlatPair.TargetConfig));
-													}													
-												}
-											}
-										}
-									}
-                                    if (!bNoAutomatedTesting)
-                                    {                                        
-										if (HostPlatform == UnrealTargetPlatform.Mac || HostPlatform == UnrealTargetPlatform.Linux) continue; //temp hack till Linux and Mac automated testing works
-										var ThisMonoGameTestNodes = new List<string>();
-										foreach (var Test in GameTests)
-										{
-											var TestName = Test.Key + "_" + Plat.ToString();
-											ThisMonoGameTestNodes.Add(AddNode(new UATTestNode(this, HostPlatform, NonCodeProject, TestName, Test.Value, CookedAgentSharingGroup, false, RequiredPlatforms)));
-										}
-										if (ThisMonoGameTestNodes.Count > 0)
-										{											
-											GameTestNodes.Add(AddNode(new GameAggregateNode(this, HostPlatform, NonCodeProject, "CookedTests_" + Plat.ToString() + "_" + Kind.ToString() + HostPlatformNode.StaticGetHostPlatformSuffix(HostPlatform), ThisMonoGameTestNodes, 0.0f)));
-										}
-									}
+                                                    var NodeName = AddNode(new FormalBuildNode(this, NonCodeProject, HostPlatform, new List<UnrealTargetPlatform>() { Plat }, new List<UnrealTargetConfiguration>() { PlatPair.TargetConfig }));                                                    
+                                                    // we don't want this delayed
+                                                    // this would normally wait for the testing phase, we just want to build it right away
+                                                    RemovePseudodependencyFromNode(
+                                                        CookNode.StaticGetFullName(HostPlatform, NonCodeProject, CookedPlatform),
+                                                        WaitForTestShared.StaticGetFullName());
+                                                    string BuildAgentSharingGroup = "";
+                                                    if (Options.bSeparateGamePromotion)
+                                                    {
+                                                        BuildAgentSharingGroup = NonCodeProject.GameName + "_MakeFormalBuild_" + Plat.ToString() + HostPlatformNode.StaticGetHostPlatformSuffix(HostPlatform);
+                                                        if (Plat == UnrealTargetPlatform.IOS || Plat == UnrealTargetPlatform.Android) // These trash build products, so we need to use different agents
+                                                        {
+                                                            BuildAgentSharingGroup = "";
+                                                        }
+                                                        GUBPNodes[CookNode.StaticGetFullName(HostPlatform, NonCodeProject, CookedPlatform)].AgentSharingGroup = BuildAgentSharingGroup;
+                                                        GUBPNodes[GamePlatformCookedAndCompiledNode.StaticGetFullName(HostPlatform, NonCodeProject, Plat)].AgentSharingGroup = BuildAgentSharingGroup;
+                                                        GUBPNodes[NodeName].AgentSharingGroup = BuildAgentSharingGroup;
+                                                    }
+                                                    else
+                                                    {
+                                                        //GUBPNodes[NodeName].AgentSharingGroup = FormalAgentSharingGroup;
+                                                        if (Plat == UnrealTargetPlatform.XboxOne)
+                                                        {
+                                                            GUBPNodes[NodeName].AgentSharingGroup = "";
+                                                        }
+                                                    }
+                                                    if (PlatPair.bTest)
+                                                    {
+                                                        AddNode(new FormalBuildTestNode(this, NonCodeProject, HostPlatform, Plat, PlatPair.TargetConfig));
+                                                    }
+                                                }
+                                            }
+                                        }
+
+                                        if (!bNoAutomatedTesting)
+                                        {
+                                            if (HostPlatform == UnrealTargetPlatform.Mac || HostPlatform == UnrealTargetPlatform.Linux) continue; //temp hack till Linux and Mac automated testing works
+                                            var ThisMonoGameTestNodes = new List<string>();
+                                            foreach (var Test in GameTests)
+                                            {
+                                                var TestName = Test.Key + "_" + Plat.ToString();
+                                                ThisMonoGameTestNodes.Add(AddNode(new UATTestNode(this, HostPlatform, NonCodeProject, TestName, Test.Value, CookedAgentSharingGroup, false, RequiredPlatforms)));
+                                            }
+                                            if (ThisMonoGameTestNodes.Count > 0)
+                                            {
+                                                GameTestNodes.Add(AddNode(new GameAggregateNode(this, HostPlatform, NonCodeProject, "CookedTests_" + Plat.ToString() + "_" + Kind.ToString() + HostPlatformNode.StaticGetHostPlatformSuffix(HostPlatform), ThisMonoGameTestNodes, 0.0f)));
+                                            }
+                                        }
+                                    }
 								}
 							}
 						}
@@ -6115,10 +5993,6 @@ public class GUBP : BuildCommand
                     AddNode(new GameAggregateNode(this, HostPlatform, Branch.BaseEngineProject, "AllCookedTests", GameTestNodes));
                 }
             }
-			if (!GUBPNodes.ContainsKey(SharedCookAggregateNode.StaticGetFullName()))
-			{
-				AddNode(new SharedCookAggregateNode(this, HostPlatforms, NonCodeProjectNames, NonCodeFormalBuilds));
-			}
 			
 			if(HostPlatform == MakeFeaturePacksNode.GetDefaultBuildPlatform(this) && !BranchOptions.bNoInstalledEngine)
 			{
@@ -6226,7 +6100,8 @@ public class GUBP : BuildCommand
 									}
                                     AddNode(new GamePlatformMonolithicsNode(this, HostPlatform, CodeProj, Plat));
                                 }
-								if (!AdditionalPlatforms.Contains(Plat))
+                                var FormalBuildConfigs = Target.Rules.GUBP_GetConfigsForFormalBuilds_MonolithicOnly(HostPlatform);
+								if (!AdditionalPlatforms.Contains(Plat) && (BranchOptions.ProjectsToCook.Contains(CodeProj.GameName) || BranchOptions.ProjectsToCook.Count == 0))
 								{
 									string CookedPlatform = Platform.Platforms[Plat].GetCookPlatform(Kind == TargetRules.TargetType.Server, Kind == TargetRules.TargetType.Client, "");
 									if (Target.Rules.GUBP_AlternateCookPlatform(HostPlatform, CookedPlatform) != "")
@@ -6241,7 +6116,7 @@ public class GUBP : BuildCommand
 									{
 										AddNode(new GamePlatformCookedAndCompiledNode(this, HostPlatform, CodeProj, Plat, true));
 									}
-									var FormalBuildConfigs = Target.Rules.GUBP_GetConfigsForFormalBuilds_MonolithicOnly(HostPlatform);
+									
 
 									foreach (var Config in FormalBuildConfigs)
 									{
@@ -6269,10 +6144,6 @@ public class GUBP : BuildCommand
 										}
 										if (FormalNodeName != null)
 										{
-											if (Config.bBeforeTrigger)
-											{
-												RemovePseudodependencyFromNode(FormalBuildNode.StaticGetFullName(CodeProj, HostPlatform, new List<UnrealTargetPlatform>() { Config.TargetPlatform }, new List<UnrealTargetConfiguration>() { Config.TargetConfig }), WaitForFormalUserInput.StaticGetFullName());
-											}
 											// we don't want this delayed
 											// this would normally wait for the testing phase, we just want to build it right away
 											RemovePseudodependencyFromNode(
@@ -6305,7 +6176,7 @@ public class GUBP : BuildCommand
 											}																				
 										}
 									}
-									if (!bNoAutomatedTesting)
+									if (!bNoAutomatedTesting && FormalBuildConfigs.Count > 0)
 									{
 										if (HostPlatform == UnrealTargetPlatform.Mac || HostPlatform == UnrealTargetPlatform.Linux) continue; //temp hack till Linux and Mac automated testing works
 										var GameTests = Target.Rules.GUBP_GetGameTests_MonolithicOnly(HostPlatform, GetAltHostPlatform(HostPlatform), Plat);
@@ -6407,8 +6278,7 @@ public class GUBP : BuildCommand
             AddNode(new SharedLabelPromotableSuccessNode());
             AddNode(new WaitForTestShared(this));
             AddNode(new WaitForSharedPromotionUserInput(this, true));
-            AddNode(new SharedLabelPromotableNode(this, true));
-			AddNode(new WaitForFormalUserInput(this));
+            AddNode(new SharedLabelPromotableNode(this, true));			
         }
         foreach (var HostPlatform in HostPlatforms)
         {
@@ -7268,9 +7138,8 @@ public class GUBP : BuildCommand
                     
                     string NodeParentPath = ParentPath;
 					string PreconditionParentPath;
-					if (GUBPNodes[NodeToDo].GetFullName().Contains("MakeBuild") && GUBPNodes[NodeToDo].FullNamesOfPseudosependencies.Contains(WaitForFormalUserInput.StaticGetFullName()) && !bGraphSubset)
-                    {
-						RemovePseudodependencyFromNode(NodeToDo, WaitForFormalUserInput.StaticGetFullName());
+					if (GUBPNodes[NodeToDo].GetFullName().Contains("MakeBuild") && !bGraphSubset)
+                    {						
 						PreconditionParentPath = GetPropertyFromStep("/myWorkflow/ParentJob");
 						UncompletedEcDeps = GetECDependencies(NodeToDo);
 					}
diff --git a/Engine/Source/UE4Editor.Target.cs b/Engine/Source/UE4Editor.Target.cs
--- a/Engine/Source/UE4Editor.Target.cs
+++ b/Engine/Source/UE4Editor.Target.cs
@@ -129,11 +129,11 @@ public class UE4EditorTarget : TargetRules
         List<UnrealTargetPlatform> TappyChickenPlats = null;
         if (HostPlatform == UnrealTargetPlatform.Mac)
         {
-            TappyChickenPlats = new List<UnrealTargetPlatform> {};
+            TappyChickenPlats = new List<UnrealTargetPlatform> { UnrealTargetPlatform.IOS };
         }
         else
         {
-            TappyChickenPlats = new List<UnrealTargetPlatform> { HostPlatform, UnrealTargetPlatform.Android, UnrealTargetPlatform.HTML5 };
+            TappyChickenPlats = new List<UnrealTargetPlatform> { UnrealTargetPlatform.Android };
         }
 
 		List<UnrealTargetPlatform> TP_2DSideScrollerPlats = null;
@@ -201,7 +201,7 @@ public class UE4EditorTarget : TargetRules
         NonCodeProjectNames.Add("BlackJack", DesktopAndMobilePlats);
         NonCodeProjectNames.Add("MemoryGame", DesktopAndMobilePlats);
         NonCodeProjectNames.Add("TappyChicken", TappyChickenPlats);
-        NonCodeProjectNames.Add("SwingNinja", DesktopAndMobilePlats);
+        NonCodeProjectNames.Add("SwingNinja", TappyChickenPlats);
         NonCodeProjectNames.Add("SunTemple", DesktopAndMobilePlats);
 
         NonCodeProjectNames.Add("FP_FirstPersonBP", AllSupportedPlats);
@@ -223,79 +223,17 @@ public class UE4EditorTarget : TargetRules
     {
         #region Sample Configurations
         //Generic Configurations
-        var PCBuildSettings = new List<GUBPFormalBuild>
-            {
-                    new GUBPFormalBuild(UnrealTargetPlatform.Win32, UnrealTargetConfiguration.Test, false),
-                    new GUBPFormalBuild(UnrealTargetPlatform.Win64, UnrealTargetConfiguration.Test, false),
-                    new GUBPFormalBuild(UnrealTargetPlatform.Mac, UnrealTargetConfiguration.Test, false),
-            };
-
-		//var MobileBuildSettings = new List<GUBPFormalBuild>
-		//	{
-		//			new GUBPFormalBuild(UnrealTargetPlatform.Android, UnrealTargetConfiguration.Test, false),
-		//			new GUBPFormalBuild(UnrealTargetPlatform.IOS, UnrealTargetConfiguration.Test, false),
-		//	};
-        // Non-Code Sample Configurations
-        var BlackJackBuildSettings = new List<GUBPFormalBuild>
-            {
-                    new GUBPFormalBuild(UnrealTargetPlatform.Win32, UnrealTargetConfiguration.Test, false),
-                    new GUBPFormalBuild(UnrealTargetPlatform.Win64, UnrealTargetConfiguration.Test, false),
-                    new GUBPFormalBuild(UnrealTargetPlatform.Mac, UnrealTargetConfiguration.Test, false),
-                    new GUBPFormalBuild(UnrealTargetPlatform.Android, UnrealTargetConfiguration.Test, false),
-                    new GUBPFormalBuild(UnrealTargetPlatform.IOS, UnrealTargetConfiguration.Test, false),
-            };
-        
-        var MemoryGameBuildSettings = new List<GUBPFormalBuild>
-            {
-                    new GUBPFormalBuild(UnrealTargetPlatform.Win32, UnrealTargetConfiguration.Test, false),
-                    new GUBPFormalBuild(UnrealTargetPlatform.Win64, UnrealTargetConfiguration.Test, false),
-                    new GUBPFormalBuild(UnrealTargetPlatform.Mac, UnrealTargetConfiguration.Test, false),
-                    new GUBPFormalBuild(UnrealTargetPlatform.Android, UnrealTargetConfiguration.Test, false),
-                    new GUBPFormalBuild(UnrealTargetPlatform.IOS, UnrealTargetConfiguration.Test, false),
-                    new GUBPFormalBuild(UnrealTargetPlatform.HTML5, UnrealTargetConfiguration.Test, false),
-            };
-        
-        var SwingNinjaBuildSettings = new List<GUBPFormalBuild>
-            {
-                    new GUBPFormalBuild(UnrealTargetPlatform.Win32, UnrealTargetConfiguration.Test, false),
-                    new GUBPFormalBuild(UnrealTargetPlatform.Win64, UnrealTargetConfiguration.Test, false),
-                    new GUBPFormalBuild(UnrealTargetPlatform.Mac, UnrealTargetConfiguration.Test, false),
-                    new GUBPFormalBuild(UnrealTargetPlatform.Android, UnrealTargetConfiguration.Test, false),
-                    new GUBPFormalBuild(UnrealTargetPlatform.IOS, UnrealTargetConfiguration.Test, false),
-            };
-        
+               
         var SunTempleBuildSettings = new List<GUBPFormalBuild>
             {
-                    new GUBPFormalBuild(UnrealTargetPlatform.Win32, UnrealTargetConfiguration.Test, false),
-                    new GUBPFormalBuild(UnrealTargetPlatform.Win64, UnrealTargetConfiguration.Test, false),
-                    new GUBPFormalBuild(UnrealTargetPlatform.Mac, UnrealTargetConfiguration.Test, false),
                     new GUBPFormalBuild(UnrealTargetPlatform.Android, UnrealTargetConfiguration.Test, false, true),
                     new GUBPFormalBuild(UnrealTargetPlatform.IOS, UnrealTargetConfiguration.Test, false, true),
-                    new GUBPFormalBuild(UnrealTargetPlatform.HTML5, UnrealTargetConfiguration.Test, false),
-            };
-
-        var ElementalBuildSettings = new List<GUBPFormalBuild>
-            {
-                    new GUBPFormalBuild(UnrealTargetPlatform.Win32, UnrealTargetConfiguration.Test, false),
-                    new GUBPFormalBuild(UnrealTargetPlatform.Win64, UnrealTargetConfiguration.Test, false),
-                    new GUBPFormalBuild(UnrealTargetPlatform.Mac, UnrealTargetConfiguration.Test, false),
-                    new GUBPFormalBuild(UnrealTargetPlatform.XboxOne, UnrealTargetConfiguration.Test, false),
-                    new GUBPFormalBuild(UnrealTargetPlatform.PS4, UnrealTargetConfiguration.Test, false),
-                    
-                    new GUBPFormalBuild(UnrealTargetPlatform.XboxOne, UnrealTargetConfiguration.Shipping, false),
-                    new GUBPFormalBuild(UnrealTargetPlatform.PS4, UnrealTargetConfiguration.Shipping, false),
-                    
-                    new GUBPFormalBuild(UnrealTargetPlatform.XboxOne, UnrealTargetConfiguration.Development, false),
-                    new GUBPFormalBuild(UnrealTargetPlatform.PS4, UnrealTargetConfiguration.Development, false),
             };
 
         var TappyChickenBuildSettings = new List<GUBPFormalBuild>
-            {
-                    new GUBPFormalBuild(UnrealTargetPlatform.Android, UnrealTargetConfiguration.Shipping, false),
-                    new GUBPFormalBuild(UnrealTargetPlatform.Android, UnrealTargetConfiguration.Test, false, true),                    
-                    new GUBPFormalBuild(UnrealTargetPlatform.Android, UnrealTargetConfiguration.Development, false),        			
-					new GUBPFormalBuild(UnrealTargetPlatform.HTML5, UnrealTargetConfiguration.Shipping, false),
-					new GUBPFormalBuild(UnrealTargetPlatform.HTML5, UnrealTargetConfiguration.Test, false),
+            {                    
+                    new GUBPFormalBuild(UnrealTargetPlatform.Android, UnrealTargetConfiguration.Test, false, true),
+                    new GUBPFormalBuild(UnrealTargetPlatform.IOS, UnrealTargetConfiguration.Test, false, true),
             };
         #endregion
 
@@ -313,77 +251,13 @@ public class UE4EditorTarget : TargetRules
                     new GUBPFormalBuild(UnrealTargetPlatform.XboxOne, UnrealTargetConfiguration.Test, false),
                     new GUBPFormalBuild(UnrealTargetPlatform.PS4, UnrealTargetConfiguration.Test, false),
             };
-
-        //TP_PUZZLE
-        var TP_PuzzleBuildSettings = new List<GUBPFormalBuild>
-            {
-                    new GUBPFormalBuild(UnrealTargetPlatform.IOS, UnrealTargetConfiguration.Test, false),
-                    new GUBPFormalBuild(UnrealTargetPlatform.Android, UnrealTargetConfiguration.Test, false),
-                    new GUBPFormalBuild(UnrealTargetPlatform.HTML5, UnrealTargetConfiguration.Test, false),
-                    new GUBPFormalBuild(UnrealTargetPlatform.Win32, UnrealTargetConfiguration.Test, false),
-                    new GUBPFormalBuild(UnrealTargetPlatform.Win64, UnrealTargetConfiguration.Test, false),
-                    new GUBPFormalBuild(UnrealTargetPlatform.Mac, UnrealTargetConfiguration.Test, false),
-            };
-
-        //TP_TOPDOWN
-        var TP_TopDownBuildSettings = new List<GUBPFormalBuild>
-            {
-                    new GUBPFormalBuild(UnrealTargetPlatform.IOS, UnrealTargetConfiguration.Test, false),
-                    new GUBPFormalBuild(UnrealTargetPlatform.Android, UnrealTargetConfiguration.Test, false),
-                    new GUBPFormalBuild(UnrealTargetPlatform.HTML5, UnrealTargetConfiguration.Test, false),
-                    new GUBPFormalBuild(UnrealTargetPlatform.Win32, UnrealTargetConfiguration.Test, false),
-                    new GUBPFormalBuild(UnrealTargetPlatform.Win64, UnrealTargetConfiguration.Test, false),
-                    new GUBPFormalBuild(UnrealTargetPlatform.Mac, UnrealTargetConfiguration.Test, false),
-            };
-
-        //TP_TWINSTICK
-        var TP_TwinStickBuildSettings = new List<GUBPFormalBuild>
-            {
-                    new GUBPFormalBuild(UnrealTargetPlatform.IOS, UnrealTargetConfiguration.Test, false),
-                    new GUBPFormalBuild(UnrealTargetPlatform.Android, UnrealTargetConfiguration.Test, false),
-                    new GUBPFormalBuild(UnrealTargetPlatform.Win32, UnrealTargetConfiguration.Test, false),
-                    new GUBPFormalBuild(UnrealTargetPlatform.Win64, UnrealTargetConfiguration.Test, false),
-                    new GUBPFormalBuild(UnrealTargetPlatform.Mac, UnrealTargetConfiguration.Test, false),
-                    new GUBPFormalBuild(UnrealTargetPlatform.XboxOne, UnrealTargetConfiguration.Test, false),
-                    new GUBPFormalBuild(UnrealTargetPlatform.PS4, UnrealTargetConfiguration.Test, false),
-            };
         #endregion
 
         var NonCodeProjectNames = new Dictionary<string, List<GUBPFormalBuild>>();
 
         //Add Samples to the list with its corresponding settings
-        NonCodeProjectNames.Add("TappyChicken", TappyChickenBuildSettings);    
-        NonCodeProjectNames.Add("BlackJack", BlackJackBuildSettings);
-        NonCodeProjectNames.Add("MemoryGame", MemoryGameBuildSettings);
-        NonCodeProjectNames.Add("SwingNinja", SwingNinjaBuildSettings);
+        NonCodeProjectNames.Add("TappyChicken", TappyChickenBuildSettings);
         NonCodeProjectNames.Add("SunTemple", SunTempleBuildSettings);
-        NonCodeProjectNames.Add("BlueprintOffice", PCBuildSettings);
-        NonCodeProjectNames.Add("ContentExamples", PCBuildSettings);
-        NonCodeProjectNames.Add("CouchKnights", PCBuildSettings);
-        NonCodeProjectNames.Add("EffectsCave", PCBuildSettings);
-        NonCodeProjectNames.Add("LandscapeMountains", PCBuildSettings);
-        NonCodeProjectNames.Add("GDC2014", PCBuildSettings);
-        NonCodeProjectNames.Add("MatineeFightScene", PCBuildSettings);
-        NonCodeProjectNames.Add("RealisticRendering", PCBuildSettings);
-        NonCodeProjectNames.Add("ReflectionsSubway", PCBuildSettings);
-        NonCodeProjectNames.Add("Stylized", PCBuildSettings);
-        NonCodeProjectNames.Add("ElementalDemo", ElementalBuildSettings);
-
-        //Add Templates to the list with its corresponding settings
-        NonCodeProjectNames.Add("FP_FirstPersonBP", AllBuildSettings);
-		NonCodeProjectNames.Add("TP_2DSideScrollerBP", AllBuildSettings);
-        NonCodeProjectNames.Add("TP_FirstPersonBP", AllBuildSettings);
-        NonCodeProjectNames.Add("TP_FlyingBP", AllBuildSettings);
-        NonCodeProjectNames.Add("TP_PuzzleBP", TP_PuzzleBuildSettings);
-        NonCodeProjectNames.Add("TP_RollingBP", AllBuildSettings);
-        NonCodeProjectNames.Add("TP_SideScrollerBP", AllBuildSettings);
-        NonCodeProjectNames.Add("TP_ThirdPersonBP", AllBuildSettings);
-        NonCodeProjectNames.Add("TP_TopDownBP", TP_TopDownBuildSettings);
-        NonCodeProjectNames.Add("TP_TwinStickBP", TP_TwinStickBuildSettings);
-        NonCodeProjectNames.Add("TP_VehicleBP", AllBuildSettings);
-        NonCodeProjectNames.Add("TP_VehicleAdvBP", AllBuildSettings);
-
-
         return NonCodeProjectNames;
     }
 }
