cis fix

diff --git a/Engine/Source/Programs/AutomationTool/Scripts/GUBP.Automation.cs b/Engine/Source/Programs/AutomationTool/Scripts/GUBP.Automation.cs
--- a/Engine/Source/Programs/AutomationTool/Scripts/GUBP.Automation.cs
+++ b/Engine/Source/Programs/AutomationTool/Scripts/GUBP.Automation.cs
@@ -2987,8 +2987,7 @@ public class GUBP : BuildCommand
                 WorkingGameProject = bp.Branch.BaseEngineProject;
             }
 
-            var AllTargetPlatforms = new List<UnrealTargetPlatform>();
-			var Options = InGameProj.Options(HostPlatform);			
+            var AllTargetPlatforms = new List<UnrealTargetPlatform>();					
             if (ClientTargetPlatforms != null)
             {
                 if (!WorkingGameProject.Properties.Targets.ContainsKey(GameOrClient))
