build script, finish removing sample package trigger

diff --git a/Engine/Source/Programs/AutomationTool/Scripts/GUBP.Automation.cs b/Engine/Source/Programs/AutomationTool/Scripts/GUBP.Automation.cs
--- a/Engine/Source/Programs/AutomationTool/Scripts/GUBP.Automation.cs
+++ b/Engine/Source/Programs/AutomationTool/Scripts/GUBP.Automation.cs
@@ -6483,8 +6483,7 @@ public class GUBP : BuildCommand
 		Dictionary<string, int> FullNodeListSortKey = GetDisplayOrder(FullNodeList.Keys.ToList(), FullNodeDirectDependencies, GUBPNodes);
 
         bool bOnlyNode = false;
-        bool bRelatedToNode = false;
-		bool bGraphSubset = false;
+        bool bRelatedToNode = false;		
         var NodesToDo = new HashSet<string>();
 
         {
@@ -6518,8 +6517,7 @@ public class GUBP : BuildCommand
                 }
             }
             if (!String.IsNullOrEmpty(NodeSpec))
-            {
-				bGraphSubset = true;
+            {				
                 if (NodeSpec.Equals("Noop", StringComparison.InvariantCultureIgnoreCase))
                 {
                     Log("Request for Noop node, done.");
@@ -6556,8 +6554,7 @@ public class GUBP : BuildCommand
         }
         string GameSpec = ParseParamValue("Game");
         if (!String.IsNullOrEmpty(GameSpec))
-        {
-			bGraphSubset = true;
+        {			
             List<string> Games = new List<string>(GameSpec.Split('+'));
             foreach (var GameArg in Games)
             {
@@ -7092,14 +7089,14 @@ public class GUBP : BuildCommand
             foreach (var NodeToDo in OrdereredToDo)
             {
                 if (GUBPNodes[NodeToDo].RunInEC() && !NodeIsAlreadyComplete(NodeToDo, LocalOnly)) // if something is already finished, we don't put it into EC  
-                {                                      
+                {
                     string EMails;
                     var NodeProps = GetECPropsForNode(NodeToDo, CLString, out EMails);
                     ECProps.AddRange(NodeProps);
 
                     bool Sticky = GUBPNodes[NodeToDo].IsSticky();
                     bool DoParallel = !Sticky;
-                    if(GUBPNodes[NodeToDo].ECProcedure() == "GUBP_UAT_Node_Parallel_AgentShare_Editor")
+                    if (GUBPNodes[NodeToDo].ECProcedure() == "GUBP_UAT_Node_Parallel_AgentShare_Editor")
                     {
                         DoParallel = true;
                     }
@@ -7108,7 +7105,7 @@ public class GUBP : BuildCommand
                         throw new AutomationException("Node {1} is sticky but has agent requirements.", NodeToDo);
                     }
                     string Procedure = GUBPNodes[NodeToDo].ECProcedure();
-                    if(GUBPNodes[NodeToDo].IsSticky() && NodeToDo == LastSticky)
+                    if (GUBPNodes[NodeToDo].IsSticky() && NodeToDo == LastSticky)
                     {
                         Procedure = Procedure + "_Release";
                     }
@@ -7128,23 +7125,18 @@ public class GUBP : BuildCommand
                     string PreCondition = "";
                     string RunCondition = "";
                     var UncompletedEcDeps = new List<string>();
-					var UncompletedCompletedDeps = new List<string>();
+                    var UncompletedCompletedDeps = new List<string>();
                     var PreConditionUncompletedEcDeps = new List<string>();
                     string MyAgentGroup = GUBPNodes[NodeToDo].AgentSharingGroup;
                     bool bDoNestedJobstep = false;
                     bool bDoFirstNestedJobstep = false;
 
-                    
+
                     string NodeParentPath = ParentPath;
-					string PreconditionParentPath;
-					if (GUBPNodes[NodeToDo].GetFullName().Contains("MakeBuild") && !bGraphSubset)
-                    {						
-						PreconditionParentPath = GetPropertyFromStep("/myWorkflow/ParentJob");
-						UncompletedEcDeps = GetECDependencies(NodeToDo);
-					}
-					else
-					{
-						PreconditionParentPath = ParentPath;
+                    string PreconditionParentPath;
+
+                    PreconditionParentPath = ParentPath;
+                    {
                         var EcDeps = GetECDependencies(NodeToDo);
                         foreach (var Dep in EcDeps)
                         {
@@ -7158,22 +7150,23 @@ public class GUBP : BuildCommand
                             }
                         }
                     }
-                    foreach(var Dep in UncompletedEcDeps)
+
+                    foreach (var Dep in UncompletedEcDeps)
                     {
                         PreConditionUncompletedEcDeps.Add(Dep);
                     }
-					var CompletedDeps = GetCompletedOnlyDependencies(NodeToDo);					
-					foreach (var Dep in CompletedDeps)
-					{
-						if (GUBPNodes[Dep].RunInEC() && !NodeIsAlreadyComplete(Dep, LocalOnly) && OrdereredToDo.Contains(Dep)) // if something is already finished, we don't put it into EC
-						{
-							if (OrdereredToDo.IndexOf(Dep) > OrdereredToDo.IndexOf(NodeToDo))
-							{
-								throw new AutomationException("Topological sort error, node {0} has a dependency of {1} which sorted after it.", NodeToDo, Dep);
-							}
-							UncompletedCompletedDeps.Add(Dep);
-						}
-					}
+                    var CompletedDeps = GetCompletedOnlyDependencies(NodeToDo);
+                    foreach (var Dep in CompletedDeps)
+                    {
+                        if (GUBPNodes[Dep].RunInEC() && !NodeIsAlreadyComplete(Dep, LocalOnly) && OrdereredToDo.Contains(Dep)) // if something is already finished, we don't put it into EC
+                        {
+                            if (OrdereredToDo.IndexOf(Dep) > OrdereredToDo.IndexOf(NodeToDo))
+                            {
+                                throw new AutomationException("Topological sort error, node {0} has a dependency of {1} which sorted after it.", NodeToDo, Dep);
+                            }
+                            UncompletedCompletedDeps.Add(Dep);
+                        }
+                    }
                     if (MyAgentGroup != "")
                     {
                         bDoNestedJobstep = true;
@@ -7183,7 +7176,7 @@ public class GUBP : BuildCommand
 
                         var MyChain = AgentGroupChains[MyAgentGroup];
                         int MyIndex = MyChain.IndexOf(NodeToDo);
-                        if (MyIndex  > 0)
+                        if (MyIndex > 0)
                         {
                             PreConditionUncompletedEcDeps.Add(MyChain[MyIndex - 1]);
                         }
@@ -7211,7 +7204,7 @@ public class GUBP : BuildCommand
                             }
                         }
                     }
-                    if(GUBPNodes[NodeToDo].IsSticky())
+                    if (GUBPNodes[NodeToDo].IsSticky())
                     {
                         var MyChain = StickyChain;
                         int MyIndex = MyChain.IndexOf(NodeToDo);
@@ -7251,7 +7244,7 @@ public class GUBP : BuildCommand
                         // these run "parallel", but we add preconditions to serialize them
                         int Index = 0;
                         foreach (var Dep in PreConditionUncompletedEcDeps)
-                        {                            
+                        {
                             PreCondition = PreCondition + "getProperty('" + GetJobStep(PreconditionParentPath, Dep) + "/status\') == \'completed\'";
                             Index++;
                             if (Index != PreConditionUncompletedEcDeps.Count)
@@ -7261,25 +7254,25 @@ public class GUBP : BuildCommand
                         }
                         PreCondition = PreCondition + ") true;]\"";
                     }
-					if(UncompletedCompletedDeps.Count > 0)
-					{
-						PreCondition = "\"\\$\" . \"[/javascript if(";
-						// these run "parallel", but we add preconditions to serialize them
-						int Index = 0;
-						foreach (var Dep in CompletedDeps)
-						{
-							if (GUBPNodes[Dep].RunInEC())
-							{
-								PreCondition = PreCondition + "getProperty('" + GetJobStep(PreconditionParentPath, Dep) + "/status\') == \'completed\'";
-								Index++;
-								if (Index != CompletedDeps.Count)
-								{
-									PreCondition = PreCondition + " && ";
-								}
-							}
-						}
-						PreCondition = PreCondition + ") true;]\"";
-					}
+                    if (UncompletedCompletedDeps.Count > 0)
+                    {
+                        PreCondition = "\"\\$\" . \"[/javascript if(";
+                        // these run "parallel", but we add preconditions to serialize them
+                        int Index = 0;
+                        foreach (var Dep in CompletedDeps)
+                        {
+                            if (GUBPNodes[Dep].RunInEC())
+                            {
+                                PreCondition = PreCondition + "getProperty('" + GetJobStep(PreconditionParentPath, Dep) + "/status\') == \'completed\'";
+                                Index++;
+                                if (Index != CompletedDeps.Count)
+                                {
+                                    PreCondition = PreCondition + " && ";
+                                }
+                            }
+                        }
+                        PreCondition = PreCondition + ") true;]\"";
+                    }
 
                     if (UncompletedEcDeps.Count > 0)
                     {
@@ -7287,8 +7280,8 @@ public class GUBP : BuildCommand
                         int Index = 0;
                         foreach (var Dep in UncompletedEcDeps)
                         {
-							RunCondition = RunCondition + "((\'\\$\" . \"[" + GetJobStep(PreconditionParentPath, Dep) + "/outcome]\' == \'success\') || ";
-							RunCondition = RunCondition + "(\'\\$\" . \"[" + GetJobStep(PreconditionParentPath, Dep) + "/outcome]\' == \'warning\'))";
+                            RunCondition = RunCondition + "((\'\\$\" . \"[" + GetJobStep(PreconditionParentPath, Dep) + "/outcome]\' == \'success\') || ";
+                            RunCondition = RunCondition + "(\'\\$\" . \"[" + GetJobStep(PreconditionParentPath, Dep) + "/outcome]\' == \'warning\'))";
 
                             Index++;
                             if (Index != UncompletedEcDeps.Count)
@@ -7368,14 +7361,14 @@ public class GUBP : BuildCommand
                     }
                     Args = Args + "});";
                     StepList.Add(Args);
-                    if (bFakeEC && 
+                    if (bFakeEC &&
                         !UnfinishedTriggers.Contains(NodeToDo) &&
                         (GUBPNodes[NodeToDo].ECProcedure().StartsWith("GUBP_UAT_Node") || GUBPNodes[NodeToDo].ECProcedure().StartsWith("GUBP_Mac_UAT_Node")) // other things we really can't test
                         ) // unfinished triggers are never run directly by EC, rather it does another job setup
                     {
-                        string Arg = String.Format("gubp -Node={0} -FakeEC {1} {2} {3} {4} {5}", 
+                        string Arg = String.Format("gubp -Node={0} -FakeEC {1} {2} {3} {4} {5}",
                             NodeToDo,
-                            bFake ? "-Fake" : "" , 
+                            bFake ? "-Fake" : "",
                             ParseParam("AllPlatforms") ? "-AllPlatforms" : "",
                             ParseParam("UnfinishedTriggersFirst") ? "-UnfinishedTriggersFirst" : "",
                             ParseParam("UnfinishedTriggersParallel") ? "-UnfinishedTriggersParallel" : "",
@@ -7390,7 +7383,7 @@ public class GUBP : BuildCommand
                         if (!String.IsNullOrEmpty(FakeFail))
                         {
                             Arg = Arg + " -FakeFail=" + FakeFail;
-                        } 
+                        }
                         FakeECArgs.Add(Arg);
                     }
 
