Parallelization of UAT startup tasks using MSBuild - NOTE: parallelization only works on Windows platforms. Mac uses xbuild, which is the MSBuild equivalent on Mac. XBuild does not perform parallel processing.

Successful preflight build: https://ec-01.epicgames.net/commander/link/jobDetails/jobs/GUBP___depot_UE4_2560607__82477_RunGraph_500023?s=Jobs&tabGroup=stepTableHeader

diff --git a/Engine/Source/Programs/AutomationTool/ScriptCompiler.cs b/Engine/Source/Programs/AutomationTool/ScriptCompiler.cs
--- a/Engine/Source/Programs/AutomationTool/ScriptCompiler.cs
+++ b/Engine/Source/Programs/AutomationTool/ScriptCompiler.cs
@@ -153,7 +153,24 @@ namespace AutomationTool
 					CommandUtils.LogVerbose("Script module {0} filtered by the Host Platform and will not be compiled.", ModuleFilename);
 				}
 			}
-			CompileModules(ModulesToCompile);
+
+			if ((UnrealBuildTool.BuildHostPlatform.Current.Platform == UnrealBuildTool.UnrealTargetPlatform.Win64) ||
+				(UnrealBuildTool.BuildHostPlatform.Current.Platform == UnrealBuildTool.UnrealTargetPlatform.Win32))
+			{
+				string Modules = string.Join(";", ModulesToCompile.ToArray());
+				var UATProj = CommandUtils.CombinePaths(CommandUtils.CmdEnv.LocalRoot, @"Engine\Source\Programs\AutomationTool\Scripts\UAT.proj");
+				var CmdLine = String.Format("\"{0}\" /p:Modules=\"{1}\" /p:Configuration={2}", UATProj, Modules, BuildConfig);
+				var Result = CommandUtils.Run(CommandUtils.CmdEnv.MsBuildExe, CmdLine);
+				if (Result.ExitCode != 0)
+				{
+					throw new AutomationException(String.Format("Failed to build \"{0}\":{1}{2}", UATProj, Environment.NewLine, Result.Output));
+				}
+			}
+			else
+			{
+				CompileModules(ModulesToCompile);
+			}
+			
 
 			Environment.CurrentDirectory = OldCWD;
 		}
diff --git a/Engine/Source/Programs/AutomationTool/Scripts/UAT.proj b/Engine/Source/Programs/AutomationTool/Scripts/UAT.proj
new file mode 100644
--- /dev/null
+++ b/Engine/Source/Programs/AutomationTool/Scripts/UAT.proj
@@ -0,0 +1,26 @@
+<Project xmlns="http://schemas.microsoft.com/developer/msbuild/2003">
+  <ItemGroup>
+    <ModulesToBuild Include="$(Modules)" />
+  </ItemGroup>
+  <PropertyGroup>
+    <BuildInParallel Condition="'$(BuildInParallel)' == ''">true</BuildInParallel>
+    <Configuration Condition="'$(Configuration)' == ''">Development</Configuration>
+    <Verbosity>quiet</Verbosity>
+    <Platform>AnyCPU</Platform>
+    <TreatWarningsAsErrors>true</TreatWarningsAsErrors>
+    <BuildProjectReferences>true</BuildProjectReferences>
+  </PropertyGroup>
+  <Target Name="Build">
+    <Error 
+	  Text="Unable to build Project %(ModulesToBuild.FullPath). Not a valid .csproj file."
+      Condition="'%(ModulesToBuild.Extension)'!='.csproj'" />
+	<Error 
+	  Text="Unable to build Project %(ModulesToBuild.FullPath). Project file not found."
+      Condition="!Exists('%(ModulesToBuild.FullPath)')" />
+	<MSBuild
+      Projects="@(ModulesToBuild)"
+      BuildInParallel="$(BuildInParallel)"
+	  Targets="Build"
+      Properties="NoLogo=true;Verbosity=$(Verbosity);Configuration=$(Configuration);Platform=$(Platform);TreatWarningsAsErrors=$(TreatWarningsAsErrors);BuildProjectReferences=$(BuildProjectReferences)"/>
+  </Target>
+</Project>
\ No newline at end of file
