Handle UAT arguments which end in a blackslash correctly (eg. -MyArg1=D:\\UE4\\ -MyArg2)

diff --git a/Engine/Source/Programs/AutomationTool/SharedUtils.cs b/Engine/Source/Programs/AutomationTool/SharedUtils.cs
--- a/Engine/Source/Programs/AutomationTool/SharedUtils.cs
+++ b/Engine/Source/Programs/AutomationTool/SharedUtils.cs
@@ -32,6 +32,12 @@ namespace AutomationTool
 				char C = CmdLine[Index];
 				if (!bQuote && Char.IsWhiteSpace(C))
 				{
+					// There can be an escape character here when passing a Windows-style path with a trailing backslash (eg. "-MyArg1=D:\UE4\ -MyArg2"). We don't want to push it into the next argument.
+					if(bEscape)
+					{
+						Arg.Append('\\');
+						bEscape = false;
+					}
 					if (Arg.Length > 0)
 					{
 						Args.Add(Arg.ToString());
