Remove a couple of overloads that just result in infinite loops.

diff --git a/Engine/Source/Programs/AutomationTool/FileFilter.cs b/Engine/Source/Programs/AutomationTool/FileFilter.cs
--- a/Engine/Source/Programs/AutomationTool/FileFilter.cs
+++ b/Engine/Source/Programs/AutomationTool/FileFilter.cs
@@ -209,15 +209,6 @@ public class FileFilter
 		AddRule(Pattern, FileFilterType.Include);
 	}
 
-	/// <summary>
-	/// Adds several exclude rules to the filter
-	/// </summary>
-	/// <param name="Patterns">Patterns to match. See CreateRegex() for details.</param>
-	public void Include(params string[] Patterns)
-	{
-		Include(Patterns);
-	}
-
 	/// <summary>
 	/// Adds several exclude rules to the filter
 	/// </summary>
@@ -239,15 +230,6 @@ public class FileFilter
 		AddRule(Pattern, FileFilterType.Exclude);
 	}
 
-	/// <summary>
-	/// Adds several exclude rules to the filter
-	/// </summary>
-	/// <param name="Patterns">Patterns to match. See CreateRegex() for details.</param>
-	public void Exclude(params string[] Patterns)
-	{
-		Exclude(Patterns);
-	}
-
 	/// <summary>
 	/// Adds several exclude rules to the filter
 	/// </summary>
