UAT: Ensure UHT is built when building BPT from UE4BuildUtils.BuildBuildPatchTool.

diff --git a/Engine/Source/Programs/AutomationTool/Scripts/UE4BuildUtils.cs b/Engine/Source/Programs/AutomationTool/Scripts/UE4BuildUtils.cs
--- a/Engine/Source/Programs/AutomationTool/Scripts/UE4BuildUtils.cs
+++ b/Engine/Source/Programs/AutomationTool/Scripts/UE4BuildUtils.cs
@@ -19,7 +19,41 @@ public class UE4BuildUtils : CommandUtils
 	/// <param name="InPlatform"></param>
 	public static void BuildBuildPatchTool(BuildCommand Command, UnrealBuildTool.UnrealTargetPlatform InPlatform)
 	{
-		Log("Building BuildPatchTool");
+		BuildUnrealHeaderTool(Command, InPlatform);
+
+		BuildProduct(Command, new UE4Build.BuildTarget()
+			{
+				ProjectName = "",
+				TargetName = "BuildPatchTool",
+				Platform = InPlatform,
+				Config = UnrealBuildTool.UnrealTargetConfiguration.Development,
+			});
+	}
+
+	/// <summary>
+	/// Builds UnrealHeaderTool for the specified platform.
+	/// </summary>
+	/// <param name="Command"></param>
+	/// <param name="InPlatform"></param>
+	public static void BuildUnrealHeaderTool(BuildCommand Command, UnrealBuildTool.UnrealTargetPlatform InPlatform)
+	{
+		BuildProduct(Command, new UE4Build.BuildTarget()
+		{
+			ProjectName = "",
+			TargetName = "UnrealHeaderTool",
+			Platform = InPlatform,
+			Config = UnrealBuildTool.UnrealTargetConfiguration.Development,
+		});
+	}
+
+	private static void BuildProduct(BuildCommand Command, UE4Build.BuildTarget Target)
+	{
+		if (Target == null)
+		{
+			throw new AutomationException("Target is required when calling UE4BuildUtils.BuildProduct");
+		}
+
+		Log("Building {0}", Target.TargetName);
 
 		if (Command == null)
 		{
@@ -29,13 +63,7 @@ public class UE4BuildUtils : CommandUtils
 		var UE4Build = new UE4Build(Command);
 
 		var Agenda = new UE4Build.BuildAgenda();
-		Agenda.Targets.Add(new UE4Build.BuildTarget()
-		{
-			ProjectName = "",
-			TargetName = "BuildPatchTool",
-			Platform = InPlatform,
-			Config = UnrealBuildTool.UnrealTargetConfiguration.Development,
-		});
+		Agenda.Targets.Add(Target);
 
 		UE4Build.Build(Agenda, InDeleteBuildProducts: true, InUpdateVersionFiles: true);
 		UE4Build.CheckBuildProducts(UE4Build.BuildProductFiles);
