enable plugin linking in the binary release
UE-16129

diff --git a/Engine/Source/Programs/AutomationTool/ProjectUtils.cs b/Engine/Source/Programs/AutomationTool/ProjectUtils.cs
--- a/Engine/Source/Programs/AutomationTool/ProjectUtils.cs
+++ b/Engine/Source/Programs/AutomationTool/ProjectUtils.cs
@@ -209,7 +209,7 @@ namespace AutomationTool
 			foreach (UnrealTargetPlatform TargetPlatformType in TargetPlatforms)
 			{
 				IUEBuildPlatform BuildPlat = UEBuildPlatform.GetBuildPlatform(TargetPlatformType, true);
-				if (BuildPlat != null && !(BuildPlat as UEBuildPlatform).HasDefaultBuildConfig(TargetPlatformType, Path.GetDirectoryName(RawProjectPath)))
+				if (!GlobalCommandLine.Rocket && BuildPlat != null && !(BuildPlat as UEBuildPlatform).HasDefaultBuildConfig(TargetPlatformType, Path.GetDirectoryName(RawProjectPath)))
 				{
 					RetVal = true;
 					break;
