missed a rocket check
UE-16129

diff --git a/Engine/Source/Programs/AutomationTool/ProjectUtils.cs b/Engine/Source/Programs/AutomationTool/ProjectUtils.cs
--- a/Engine/Source/Programs/AutomationTool/ProjectUtils.cs
+++ b/Engine/Source/Programs/AutomationTool/ProjectUtils.cs
@@ -265,7 +265,7 @@ namespace AutomationTool
 
 			// detect if the project is content only, but has non-default build settings
 			List<string> ExtraSearchPaths = null;
-			if (!string.IsNullOrEmpty(RawProjectPath) && !GlobalCommandLine.Rocket)
+			if (!string.IsNullOrEmpty(RawProjectPath))
 			{
 				if (RequiresTempTarget(RawProjectPath, ClientTargetPlatforms))
 				{
