We want to display an error message more prominantly when a binary is missing during the packaging of their game.
This used to be displayed, but seems to have been removed accidentally in CL# 2506770.

diff --git a/Engine/Source/Programs/AutomationTool/Scripts/CopyBuildToStagingDirectory.Automation.cs b/Engine/Source/Programs/AutomationTool/Scripts/CopyBuildToStagingDirectory.Automation.cs
--- a/Engine/Source/Programs/AutomationTool/Scripts/CopyBuildToStagingDirectory.Automation.cs
+++ b/Engine/Source/Programs/AutomationTool/Scripts/CopyBuildToStagingDirectory.Automation.cs
@@ -1474,6 +1474,7 @@ public partial class Project : CommandUtils
 							if (bRequireStagedFilesToExist)
 							{
 								// if we aren't collecting multiple platforms, then it is expected to exist
+								AutomationTool.ErrorReporter.Error("Stage Failed.", (int)AutomationTool.ErrorCodes.Error_MissingExecutable);
 								throw new AutomationException("Missing receipt '{0}'. Check that this target has been built.", Path.GetFileName(ReceiptFileName));
 							}
 							else
