catch problems writing job step details and just log them as warnings

diff --git a/Engine/Source/Programs/AutomationTool/CommandUtils.cs b/Engine/Source/Programs/AutomationTool/CommandUtils.cs
--- a/Engine/Source/Programs/AutomationTool/CommandUtils.cs
+++ b/Engine/Source/Programs/AutomationTool/CommandUtils.cs
@@ -1925,10 +1925,17 @@ namespace AutomationTool
         {
             if (IsBuildMachine && CmdEnv.CSVFile != "")
             {
-                var CSVBuilder = new StringBuilder();                
+                var CSVBuilder = new StringBuilder();
                 var CSVLineToAppend = String.Format("{0}{1}", Input, Environment.NewLine);
-                CSVBuilder.Append(CSVLineToAppend);                
-                File.AppendAllText(CmdEnv.CSVFile, CSVBuilder.ToString());
+                CSVBuilder.Append(CSVLineToAppend);
+                try
+                {
+                    File.AppendAllText(CmdEnv.CSVFile, CSVBuilder.ToString());
+                }
+                catch(Exception ex)
+                {
+                    Log(System.Diagnostics.TraceEventType.Warning, "Could not append to csv file.");
+                }
             }
         }
 
