build script, remove unused variable

diff --git a/Engine/Source/Programs/AutomationTool/CommandUtils.cs b/Engine/Source/Programs/AutomationTool/CommandUtils.cs
--- a/Engine/Source/Programs/AutomationTool/CommandUtils.cs
+++ b/Engine/Source/Programs/AutomationTool/CommandUtils.cs
@@ -1932,7 +1932,7 @@ namespace AutomationTool
                 {
                     File.AppendAllText(CmdEnv.CSVFile, CSVBuilder.ToString());
                 }
-                catch(Exception ex)
+                catch
                 {
                     Log(System.Diagnostics.TraceEventType.Warning, "Could not append to csv file.");
                 }
