UAT: Remove temporary requirement that building BuildPatchTool forced a build of UHT.

diff --git a/Engine/Source/Programs/AutomationTool/Scripts/UE4BuildUtils.cs b/Engine/Source/Programs/AutomationTool/Scripts/UE4BuildUtils.cs
--- a/Engine/Source/Programs/AutomationTool/Scripts/UE4BuildUtils.cs
+++ b/Engine/Source/Programs/AutomationTool/Scripts/UE4BuildUtils.cs
@@ -19,8 +19,6 @@ public class UE4BuildUtils : CommandUtils
 	/// <param name="InPlatform"></param>
 	public static void BuildBuildPatchTool(BuildCommand Command, UnrealBuildTool.UnrealTargetPlatform InPlatform)
 	{
-		BuildUnrealHeaderTool(Command, InPlatform);
-
 		BuildProduct(Command, new UE4Build.BuildTarget()
 			{
 				ProjectName = "",
