- integration of ParallelExecutor into build process
- correct setting of environment variables per tool based on the Environment node the tool belongs to
- correct MissingDependencyCount to exclude duplicate dependencies

preflight build: https://ec-01.epicgames.net/commander/link/jobDetails/jobs/3521469?filterName0=&firstElement=501&s=Jobs

diff --git a/Engine/Source/Programs/AutomationTool/Scripts/GUBP.Automation.cs b/Engine/Source/Programs/AutomationTool/Scripts/GUBP.Automation.cs
--- a/Engine/Source/Programs/AutomationTool/Scripts/GUBP.Automation.cs
+++ b/Engine/Source/Programs/AutomationTool/Scripts/GUBP.Automation.cs
@@ -707,9 +707,11 @@ public class GUBP : BuildCommand
 
     public class CompileNode : HostPlatformNode
     {
+		bool bDependentOnCompileTools = true;
         public CompileNode(UnrealTargetPlatform InHostPlatform, bool DependentOnCompileTools = true)
             : base(InHostPlatform)
         {
+			bDependentOnCompileTools = DependentOnCompileTools;
             if (DependentOnCompileTools)
             {
                 AddDependency(ToolsForCompileNode.StaticGetFullName(HostPlatform));
@@ -742,7 +744,8 @@ public class GUBP : BuildCommand
             {
                 bool ReallyDeleteBuildProducts = DeleteBuildProducts() && !GUBP.bForceIncrementalCompile;
                 Agenda.DoRetries = false; // these would delete build products
-                UE4Build.Build(Agenda, InDeleteBuildProducts: ReallyDeleteBuildProducts, InUpdateVersionFiles: false, InForceUnity: true);
+				bool UseParallelExecutor = bDependentOnCompileTools && (HostPlatform == UnrealTargetPlatform.Win64);
+				UE4Build.Build(Agenda, InDeleteBuildProducts: ReallyDeleteBuildProducts, InUpdateVersionFiles: false, InForceUnity: true, InUseParallelExecutor: UseParallelExecutor);
                 var StartPostBuild = DateTime.Now.ToString();
                 PostBuild(bp, UE4Build);
                 var FinishPostBuild = DateTime.Now.ToString();
@@ -836,6 +839,10 @@ public class GUBP : BuildCommand
             string AddArgs = "-CopyAppBundleBackToDevice";
 
             Agenda.AddTargets(new string[] { "UnrealHeaderTool" }, HostPlatform, UnrealTargetConfiguration.Development, InAddArgs: AddArgs);
+			if (HostPlatform == UnrealTargetPlatform.Win64)
+			{
+				Agenda.AddTargets(new string[] { "ParallelExecutor" }, HostPlatform, UnrealTargetConfiguration.Development, InAddArgs: AddArgs);
+			}
             return Agenda;
         }
         public override void PostBuild(GUBP bp, UE4Build UE4Build)
@@ -5731,7 +5738,7 @@ public class GUBP : BuildCommand
 
 			if (!BranchOptions.ExcludePlatformsForEditor.Contains(HostPlatform))
 			{
-				AddNode(new NonUnityTestNode(HostPlatform));
+			AddNode(new NonUnityTestNode(HostPlatform));
 			}
 
             if (DoASharedPromotable)
@@ -5900,6 +5907,8 @@ public class GUBP : BuildCommand
                                         {
                                             AddNode(new GamePlatformCookedAndCompiledNode(this, HostPlatform, NonCodeProject, Plat, false));
 
+											if (NonCodeFormalBuilds.ContainsKey(NonCodeProject.GameName))
+											{
                                             var PlatList = NonCodeFormalBuilds[NonCodeProject.GameName];
                                             foreach (var PlatPair in PlatList)
                                             {
@@ -5938,6 +5947,7 @@ public class GUBP : BuildCommand
                                                 }
                                             }
                                         }
+                                        }
 
                                         if (!bNoAutomatedTesting)
                                         {
diff --git a/Engine/Source/Programs/AutomationTool/UE4Build.cs b/Engine/Source/Programs/AutomationTool/UE4Build.cs
--- a/Engine/Source/Programs/AutomationTool/UE4Build.cs
+++ b/Engine/Source/Programs/AutomationTool/UE4Build.cs
@@ -19,6 +19,7 @@ namespace AutomationTool
 		private BuildCommand OwnerCommand;
 
 		private String XGEConsoleExePath = "";
+		private String ParallelExecutorPath = "";
 
 		public void PrepareBuildProduct(string File, bool DeleteProduct = true)
 		{
@@ -747,6 +748,7 @@ namespace AutomationTool
 			Result.Sort();
 			return Result;
 		}
+
 		public string XGEConsoleExe()
 		{
 			if (string.IsNullOrEmpty(XGEConsoleExePath))
@@ -765,304 +767,29 @@ namespace AutomationTool
 			return XGEConsoleExePath;
 		}
 
-		public bool RunXGE(List<XGEItem> Actions, string TaskFilePath, bool DoRetries, bool SpecialTestFlag, bool ShowProgress)
+		public string ParallelExecutorExec()
 		{
-			string XGEConsole = XGEConsoleExe();
-			if (string.IsNullOrEmpty(XGEConsole))
+			if (string.IsNullOrEmpty(ParallelExecutorPath))
 			{
-				throw new AutomationException("Unable to find xge executable: " + XGEConsole);
+				string RootPath = Path.Combine(CmdEnv.LocalRoot, @"Engine\Binaries\Win64");
+				string FullPath = Path.Combine(RootPath, "ParallelExecutor.exe");
+				if (FileExists(FullPath))
+				{
+					ParallelExecutorPath = FullPath;
+				}
 			}
+			return ParallelExecutorPath;
+		}
 
-			XmlDocument XGETaskDocument = new XmlDocument();
-
-			// <BuildSet FormatVersion="1">...</BuildSet>
-			XmlElement BuildSetElement = XGETaskDocument.CreateElement("BuildSet");
-			XGETaskDocument.AppendChild(BuildSetElement);
-			BuildSetElement.SetAttribute("FormatVersion", "1");
-
-			// <Environments>...</Environments>
-			XmlElement EnvironmentsElement = XGETaskDocument.CreateElement("Environments");
-			BuildSetElement.AppendChild(EnvironmentsElement);
-
-			int Job = 0;
-			int Env = 0;
-			Dictionary<string, XmlElement> EnvStringToEnv = new Dictionary<string, XmlElement>();
-			Dictionary<string, XmlElement> EnvStringToProject = new Dictionary<string, XmlElement>();
-			Dictionary<string, string> ParamsToTool = new Dictionary<string, string>();
-			Dictionary<string, XmlElement> ParamsToToolElement = new Dictionary<string, XmlElement>();
-			Dictionary<string, string> ToolToAction = new Dictionary<string, string>();
-			foreach (var Item in Actions)
+		public bool ProcessXGEItems(List<XGEItem> Actions, string XGETool, string Args, string TaskFilePath, bool DoRetries, bool SpecialTestFlag, bool ShowProgress)
+		{
+			if (string.IsNullOrEmpty(XGETool))
 			{
-				var CurrentDependencies = new List<string>();
-				foreach (var XGEFile in Item.XgeXmlFiles)
-				{
-					if (!FileExists_NoExceptions(XGEFile))
-					{
-						throw new AutomationException("BUILD FAILED: Couldn't find file: {0}", XGEFile);
-					}
-					var TargetFile = TaskFilePath + "." + Path.GetFileName(XGEFile);
-					CopyFile(XGEFile, TargetFile);
-					CopyFile_NoExceptions(XGEFile, TaskFilePath);
-					XmlDocument UBTTask = new XmlDocument();
-					UBTTask.XmlResolver = null;
-					UBTTask.Load(XGEFile);
-					DeleteFile(XGEFile);
-
-					var All = new List<string>();
-					{
-						var Elements = UBTTask.GetElementsByTagName("Variable");
-						foreach (XmlElement Element in Elements)
-						{
-							string Pair = Element.Attributes["Name"].Value + "=" + Element.Attributes["Value"].Value;
-							All.Add(Pair);
-						}
-					}
-					All.Sort();
-					string AllString = "";
-					foreach (string Element in All)
-					{
-						AllString += Element + "\n";
-					}
-					XmlElement ToolsElement;
-					XmlElement ProjectElement;
-
-					if (EnvStringToEnv.ContainsKey(AllString))
-					{
-						ToolsElement = EnvStringToEnv[AllString];
-						ProjectElement = EnvStringToProject[AllString];
-					}
-					else
-					{
-						string EnvName = string.Format("Env_{0}", Env);
-						Env++;
-						// <Environment Name="Win32">...</Environment>
-						XmlElement EnvironmentElement = XGETaskDocument.CreateElement("Environment");
-						EnvironmentsElement.AppendChild(EnvironmentElement);
-						EnvironmentElement.SetAttribute("Name", EnvName);
-
-						// <Tools>...</Tools>
-						ToolsElement = XGETaskDocument.CreateElement("Tools");
-						EnvironmentElement.AppendChild(ToolsElement);
-
-						{
-							// <Variables>...</Variables>
-							XmlElement VariablesElement = XGETaskDocument.CreateElement("Variables");
-							EnvironmentElement.AppendChild(VariablesElement);
-
-							var Elements = UBTTask.GetElementsByTagName("Variable");
-							foreach (XmlElement Element in Elements)
-							{
-								// <Variable>...</Variable>
-								XmlElement VariableElement = XGETaskDocument.CreateElement("Variable");
-								VariablesElement.AppendChild(VariableElement);
-								VariableElement.SetAttribute("Name", Element.Attributes["Name"].Value);
-								VariableElement.SetAttribute("Value", Element.Attributes["Value"].Value);
-							}
-						}
-
-						// <Project Name="Default" Env="Default">...</Project>
-						ProjectElement = XGETaskDocument.CreateElement("Project");
-						BuildSetElement.AppendChild(ProjectElement);
-						ProjectElement.SetAttribute("Name", EnvName);
-						ProjectElement.SetAttribute("Env", EnvName);
-
-						EnvStringToEnv.Add(AllString, ToolsElement);
-						EnvStringToProject.Add(AllString, ProjectElement);
-
-					}
-
-					Dictionary<string, string> ToolToTool = new Dictionary<string, string>();
-					Dictionary<string, string> ActionToAction = new Dictionary<string, string>();
-
-					{
-						var Elements = UBTTask.GetElementsByTagName("Tool");
-						foreach (XmlElement Element in Elements)
-						{
-							string Key = Element.Attributes["Path"].Value;
-							Key += " ";
-							Key += Element.Attributes["Params"].Value;
-
-							//hack hack hack
-							string ElementParams = Element.Attributes["Params"].Value;
-							if (!String.IsNullOrEmpty(ElementParams))
-							{
-								int YcIndex = ElementParams.IndexOf(" /Yc\"");
-								if (YcIndex >= 0)
-								{
-									// /Fp&quot;D:\BuildFarm\buildmachine_++depot+UE4\Engine\Intermediate\BuildData\Win64\UE4Editor\Development\SharedPCHs\CoreUObject.h.pch&quot
-									string Fp = " /Fp\"";
-									int FpIndex = ElementParams.IndexOf(Fp, YcIndex);
-									if (FpIndex >= 0)
-									{
-										int EndIndex = ElementParams.IndexOf("\"", FpIndex + Fp.Length);
-										if (EndIndex >= 0)
-										{
-											string PCHFileName = ElementParams.Substring(FpIndex + Fp.Length, EndIndex - FpIndex - Fp.Length);
-											if (PCHFileName.Contains(@"\SharedPCHs\") && PCHFileName.Contains(@"\UE4Editor\"))
-											{
-												Key = "SharedEditorPCH$ " + PCHFileName;
-												Log("Hack: detected Shared PCH, which will use a different key {0}", Key);
-
-											}
-										}
-									}
-								}
-
-							}
-
-							string ToolName = string.Format("{0}_{1}", Element.Attributes["Name"].Value, Job);
-							string OriginalToolName = ToolName;
-
-							if (ParamsToTool.ContainsKey(Key))
-							{
-								ToolName = ParamsToTool[Key];
-								ToolToTool.Add(OriginalToolName, ToolName);
-
-								XmlElement ToolElement = ParamsToToolElement[Key];
-								ToolElement.SetAttribute("GroupPrefix", ToolElement.Attributes["GroupPrefix"].Value + " + " + Item.OutputCaption);
-							}
-							else
-							{
-								// <Tool ... />
-								XmlElement ToolElement = XGETaskDocument.CreateElement("Tool");
-								ToolsElement.AppendChild(ToolElement);
-
-								ParamsToTool.Add(Key, ToolName);
-								ParamsToToolElement.Add(Key, ToolElement);
-
-								ToolElement.SetAttribute("Name", ToolName);
-								ToolElement.SetAttribute("AllowRemote", Element.Attributes["AllowRemote"].Value);
-								if (Element.HasAttribute("OutputPrefix"))
-								{
-									ToolElement.SetAttribute("OutputPrefix", Element.Attributes["OutputPrefix"].Value);
-								}
-								ToolElement.SetAttribute("GroupPrefix", "** For " + Item.OutputCaption);
-
-								ToolElement.SetAttribute("Params", Element.Attributes["Params"].Value);
-								ToolElement.SetAttribute("Path", Element.Attributes["Path"].Value);
-								if(Element.HasAttribute("VCCompiler"))
-								{
-									ToolElement.SetAttribute("VCCompiler", Element.Attributes["VCCompiler"].Value);
-								}
-								ToolElement.SetAttribute("SkipIfProjectFailed", Element.Attributes["SkipIfProjectFailed"].Value);
-								if (Element.HasAttribute("AutoReserveMemory"))
-								{
-									ToolElement.SetAttribute("AutoReserveMemory", Element.Attributes["AutoReserveMemory"].Value);
-								}
-								ToolElement.SetAttribute("OutputFileMasks", Element.Attributes["OutputFileMasks"].Value);
-								//ToolElement.SetAttribute("AllowRestartOnLocal", "false");  //vs2012 can't really restart, so we will go with this for now
-								if (Element.Attributes["OutputFileMasks"].Value == "PCLaunch.rc.res")
-								{
-									// total hack, when doing clean compiles, this output directory does not exist, we need to create it now
-									string Parms = Element.Attributes["Params"].Value;
-									string Start = "/fo \"";
-									int StartIndex = Parms.IndexOf(Start);
-									if (StartIndex >= 0)
-									{
-										Parms = Parms.Substring(StartIndex + Start.Length);
-										int EndIndex = Parms.IndexOf("\"");
-										if (EndIndex > 0)
-										{
-											string ResLocation = CombinePaths(Parms.Substring(0, EndIndex));
-											if (!DirectoryExists_NoExceptions(GetDirectoryName(ResLocation)))
-											{
-												CreateDirectory(GetDirectoryName(ResLocation));
-											}
-										}
-									}
-								}
-							}
-						}
-					}
-					{
-						var NextDependencies = new List<string>();
-
-						var Elements = UBTTask.GetElementsByTagName("Task");
-						foreach (XmlElement Element in Elements)
-						{
-							string ToolName = string.Format("{0}_{1}", Element.Attributes["Tool"].Value, Job);
-							string ActionName = string.Format("{0}_{1}", Element.Attributes["Name"].Value, Job);
-							string OriginalActionName = ActionName;
-
-							if (ToolToTool.ContainsKey(ToolName))
-							{
-								ToolName = ToolToTool[ToolName];
-								ActionName = ToolToAction[ToolName];
-								ActionToAction.Add(OriginalActionName, ActionName);
-							}
-							else
-							{
-								ActionToAction.Add(OriginalActionName, ActionName);
-								ToolToAction.Add(ToolName, ActionName);
-								// <Task ... />
-								XmlElement TaskElement = XGETaskDocument.CreateElement("Task");
-								ProjectElement.AppendChild(TaskElement);
-
-								TaskElement.SetAttribute("SourceFile", Element.Attributes["SourceFile"].Value);
-								if (Element.HasAttribute("Caption"))
-								{
-									TaskElement.SetAttribute("Caption", Element.Attributes["Caption"].Value);
-								}
-								TaskElement.SetAttribute("Name", ActionName);
-								NextDependencies.Add(ActionName);
-								TaskElement.SetAttribute("Tool", ToolName);
-								TaskElement.SetAttribute("WorkingDir", Element.Attributes["WorkingDir"].Value);
-								TaskElement.SetAttribute("SkipIfProjectFailed", Element.Attributes["SkipIfProjectFailed"].Value);
-
-								string NewDepends = "";
-								if (Element.HasAttribute("DependsOn"))
-								{
-									string Depends = Element.Attributes["DependsOn"].Value;
-									while (Depends.Length > 0)
-									{
-										string ThisAction = Depends;
-										int Semi = Depends.IndexOf(";");
-										if (Semi >= 0)
-										{
-											ThisAction = Depends.Substring(0, Semi);
-											Depends = Depends.Substring(Semi + 1);
-										}
-										else
-										{
-											Depends = "";
-										}
-										if (ThisAction.Length > 0)
-										{
-											if (NewDepends.Length > 0)
-											{
-												NewDepends += ";";
-											}
-											string ActionJob = ThisAction + string.Format("_{0}", Job);
-											if (!ActionToAction.ContainsKey(ActionJob))
-											{
-												throw new AutomationException("Action not found '{0}' in {1}->{2}", ActionJob, XGEFile, TargetFile);
-												// the XGE schedule is not topologically sorted. Hmmm. Well, we make a scary assumption here that this 
-											}
-											NewDepends += ActionToAction[ActionJob];
-										}
-									}
-								}
-								foreach (var Dep in CurrentDependencies)
-								{
-									if (NewDepends.Length > 0)
-									{
-										NewDepends += ";";
-									}
-									NewDepends += Dep;
-								}
-								if (NewDepends != "")
-								{
-									TaskElement.SetAttribute("DependsOn", NewDepends);
-								}
-							}
-
-						}
-						CurrentDependencies.AddRange(NextDependencies);
-					}
-					Job++;
-				}
+				throw new AutomationException("Unable to find xge tool: " + XGETool);
 			}
-			if (Job == 0)
+
+			XmlDocument XGETaskDocument;	
+			if (!CombineXGEItemFiles(Actions, TaskFilePath, out XGETaskDocument))
 			{
 				if (DeleteBuildProducts)
 				{
@@ -1075,7 +802,7 @@ namespace AutomationTool
 			}
 			else
 			{
-				// Make sure all the tasks have a unique prefix
+				XmlElement EnvironmentsElement = XGETaskDocument.CreateElement("Environments");// Make sure all the tasks have a unique prefix
 				if(ShowProgress)
 				{
 					List<XmlElement> AllToolElements = new List<XmlElement>();
@@ -1113,7 +840,7 @@ namespace AutomationTool
 					throw new AutomationException("Unable to find xge xml: " + TaskFilePath);
 				}
 
-				string Args = "\"" + TaskFilePath + "\" /Rebuild /MaxCPUS=200";
+				System.IO.File.Copy(TaskFilePath, @"P:\Developers\Paul.Fazio\XGE.xml", true);
 
 				int Retries = DoRetries ? 2 : 1;
                 int ConnectionRetries = 4;
@@ -1123,15 +850,15 @@ namespace AutomationTool
 					{
                         while (true)
                         {
-                            Log("Running XGE *******");
+                            Log("Running {0} *******", XGETool);
                             PushDir(CombinePaths(CmdEnv.LocalRoot, @"\Engine\Source"));
                             int SuccesCode;
-                            string LogFile = GetRunAndLogLogName(CmdEnv, XGEConsole);
-                            string Output = RunAndLog(XGEConsole, Args, out SuccesCode, LogFile);
+                            string LogFile = GetRunAndLogLogName(CmdEnv, XGETool);
+                            string Output = RunAndLog(XGETool, Args, out SuccesCode, LogFile);
                             PopDir();
                             if (ConnectionRetries > 0 && (SuccesCode == 4 || SuccesCode == 2) && !Output.Contains("------Project:"))
                             {
-                                Log(System.Diagnostics.TraceEventType.Warning, "XGE failure on the local connection timeout");
+                                Log(System.Diagnostics.TraceEventType.Warning, String.Format("{0} failure on the local connection timeout", XGETool));
                                 if (ConnectionRetries < 2)
                                 {
                                     System.Threading.Thread.Sleep(60000);
@@ -1141,18 +868,18 @@ namespace AutomationTool
                             }
                             else if (SuccesCode != 0)
                             {
-                                Log("XGE did not succeed *******");
-                                throw new AutomationException(String.Format("Command failed (Result:{3}): {0} {1}. See logfile for details: '{2}' ",
-                                                                XGEConsole, Args, Path.GetFileName(LogFile), SuccesCode));
+                                Log("{0} did not succeed *******", XGETool);
+                                throw new AutomationException("Command failed (Result:{3}): {0} {1}. See logfile for details: '{2}' ",
+                                                                XGETool, Args, Path.GetFileName(LogFile), SuccesCode);
                             }
-                            Log("XGE Done *******");
+                            Log("{0} {1} Done *******", XGETool, Args);
                             break;
                         }
                         break;
 					}
 					catch (Exception Ex)
 					{
-						Log("XGE failed on try {0}: {1}", i + 1, Ex.ToString());
+						Log("{0} failed on try {1}: {2}", XGETool, i + 1, Ex.ToString());
 						if (i + 1 >= Retries)
 						{
 							return false;
@@ -1173,6 +900,300 @@ namespace AutomationTool
 			return true;
 		}
 
+        private static bool CombineXGEItemFiles(List<XGEItem> Actions, string TaskFilePath, out XmlDocument XGETaskDocument)
+        {
+            XGETaskDocument = new XmlDocument();
+
+            // <BuildSet FormatVersion="1">...</BuildSet>
+            XmlElement BuildSetElement = XGETaskDocument.CreateElement("BuildSet");
+            XGETaskDocument.AppendChild(BuildSetElement);
+            BuildSetElement.SetAttribute("FormatVersion", "1");
+
+            // <Environments>...</Environments>
+            XmlElement EnvironmentsElement = XGETaskDocument.CreateElement("Environments");
+            BuildSetElement.AppendChild(EnvironmentsElement);
+
+            int Job = 0;
+            int Env = 0;
+            Dictionary<string, XmlElement> EnvStringToEnv = new Dictionary<string, XmlElement>();
+            Dictionary<string, XmlElement> EnvStringToProject = new Dictionary<string, XmlElement>();
+            Dictionary<string, string> ParamsToTool = new Dictionary<string, string>();
+            Dictionary<string, XmlElement> ParamsToToolElement = new Dictionary<string, XmlElement>();
+            Dictionary<string, string> ToolToAction = new Dictionary<string, string>();
+            foreach (var Item in Actions)
+            {
+                var CurrentDependencies = new List<string>();
+                foreach (var XGEFile in Item.XgeXmlFiles)
+                {
+                    if (!FileExists_NoExceptions(XGEFile))
+                    {
+                        throw new AutomationException("BUILD FAILED: Couldn't find file: {0}", XGEFile);
+                    }
+                    var TargetFile = TaskFilePath + "." + Path.GetFileName(XGEFile);
+                    CopyFile(XGEFile, TargetFile);
+                    CopyFile_NoExceptions(XGEFile, TaskFilePath);
+                    XmlDocument UBTTask = new XmlDocument();
+                    UBTTask.XmlResolver = null;
+                    UBTTask.Load(XGEFile);
+                    DeleteFile(XGEFile);
+
+                    var All = new List<string>();
+                    {
+                        var Elements = UBTTask.GetElementsByTagName("Variable");
+                        foreach (XmlElement Element in Elements)
+                        {
+                            string Pair = Element.Attributes["Name"].Value + "=" + Element.Attributes["Value"].Value;
+                            All.Add(Pair);
+                        }
+                    }
+                    All.Sort();
+                    string AllString = "";
+                    foreach (string Element in All)
+                    {
+                        AllString += Element + "\n";
+                    }
+                    XmlElement ToolsElement;
+                    XmlElement ProjectElement;
+
+                    if (EnvStringToEnv.ContainsKey(AllString))
+                    {
+                        ToolsElement = EnvStringToEnv[AllString];
+                        ProjectElement = EnvStringToProject[AllString];
+                    }
+                    else
+                    {
+                        string EnvName = string.Format("Env_{0}", Env);
+                        Env++;
+                        // <Environment Name="Win32">...</Environment>
+                        XmlElement EnvironmentElement = XGETaskDocument.CreateElement("Environment");
+                        EnvironmentsElement.AppendChild(EnvironmentElement);
+                        EnvironmentElement.SetAttribute("Name", EnvName);
+
+                        // <Tools>...</Tools>
+                        ToolsElement = XGETaskDocument.CreateElement("Tools");
+                        EnvironmentElement.AppendChild(ToolsElement);
+
+                        {
+                            // <Variables>...</Variables>
+                            XmlElement VariablesElement = XGETaskDocument.CreateElement("Variables");
+                            EnvironmentElement.AppendChild(VariablesElement);
+
+                            var Elements = UBTTask.GetElementsByTagName("Variable");
+                            foreach (XmlElement Element in Elements)
+                            {
+                                // <Variable>...</Variable>
+                                XmlElement VariableElement = XGETaskDocument.CreateElement("Variable");
+                                VariablesElement.AppendChild(VariableElement);
+                                VariableElement.SetAttribute("Name", Element.Attributes["Name"].Value);
+                                VariableElement.SetAttribute("Value", Element.Attributes["Value"].Value);
+                            }
+                        }
+
+                        // <Project Name="Default" Env="Default">...</Project>
+                        ProjectElement = XGETaskDocument.CreateElement("Project");
+                        BuildSetElement.AppendChild(ProjectElement);
+                        ProjectElement.SetAttribute("Name", EnvName);
+                        ProjectElement.SetAttribute("Env", EnvName);
+
+                        EnvStringToEnv.Add(AllString, ToolsElement);
+                        EnvStringToProject.Add(AllString, ProjectElement);
+
+                    }
+
+                    Dictionary<string, string> ToolToTool = new Dictionary<string, string>();
+                    Dictionary<string, string> ActionToAction = new Dictionary<string, string>();
+
+                    {
+                        var Elements = UBTTask.GetElementsByTagName("Tool");
+                        foreach (XmlElement Element in Elements)
+                        {
+                            string Key = Element.Attributes["Path"].Value;
+                            Key += " ";
+                            Key += Element.Attributes["Params"].Value;
+
+                            //hack hack hack
+                            string ElementParams = Element.Attributes["Params"].Value;
+                            if (!String.IsNullOrEmpty(ElementParams))
+                            {
+                                int YcIndex = ElementParams.IndexOf(" /Yc\"");
+                                if (YcIndex >= 0)
+                                {
+                                    // /Fp&quot;D:\BuildFarm\buildmachine_++depot+UE4\Engine\Intermediate\BuildData\Win64\UE4Editor\Development\SharedPCHs\CoreUObject.h.pch&quot
+                                    string Fp = " /Fp\"";
+                                    int FpIndex = ElementParams.IndexOf(Fp, YcIndex);
+                                    if (FpIndex >= 0)
+                                    {
+                                        int EndIndex = ElementParams.IndexOf("\"", FpIndex + Fp.Length);
+                                        if (EndIndex >= 0)
+                                        {
+                                            string PCHFileName = ElementParams.Substring(FpIndex + Fp.Length, EndIndex - FpIndex - Fp.Length);
+                                            if (PCHFileName.Contains(@"\SharedPCHs\") && PCHFileName.Contains(@"\UE4Editor\"))
+                                            {
+                                                Key = "SharedEditorPCH$ " + PCHFileName;
+                                                Log("Hack: detected Shared PCH, which will use a different key {0}", Key);
+
+                                            }
+                                        }
+                                    }
+                                }
+
+                            }
+
+                            string ToolName = string.Format("{0}_{1}", Element.Attributes["Name"].Value, Job);
+                            string OriginalToolName = ToolName;
+
+                            if (ParamsToTool.ContainsKey(Key))
+                            {
+                                ToolName = ParamsToTool[Key];
+                                ToolToTool.Add(OriginalToolName, ToolName);
+
+                                XmlElement ToolElement = ParamsToToolElement[Key];
+                                ToolElement.SetAttribute("GroupPrefix", ToolElement.Attributes["GroupPrefix"].Value + " + " + Item.OutputCaption);
+                            }
+                            else
+                            {
+                                // <Tool ... />
+                                XmlElement ToolElement = XGETaskDocument.CreateElement("Tool");
+                                ToolsElement.AppendChild(ToolElement);
+
+                                ParamsToTool.Add(Key, ToolName);
+                                ParamsToToolElement.Add(Key, ToolElement);
+
+                                ToolElement.SetAttribute("Name", ToolName);
+                                ToolElement.SetAttribute("AllowRemote", Element.Attributes["AllowRemote"].Value);
+                                if (Element.HasAttribute("OutputPrefix"))
+                                {
+                                    ToolElement.SetAttribute("OutputPrefix", Element.Attributes["OutputPrefix"].Value);
+                                }
+                                ToolElement.SetAttribute("GroupPrefix", "** For " + Item.OutputCaption);
+
+                                ToolElement.SetAttribute("Params", Element.Attributes["Params"].Value);
+                                ToolElement.SetAttribute("Path", Element.Attributes["Path"].Value);
+                                if(Element.HasAttribute("VCCompiler"))
+                                {
+                                    ToolElement.SetAttribute("VCCompiler", Element.Attributes["VCCompiler"].Value);
+                                }
+                                ToolElement.SetAttribute("SkipIfProjectFailed", Element.Attributes["SkipIfProjectFailed"].Value);
+                                if (Element.HasAttribute("AutoReserveMemory"))
+                                {
+                                    ToolElement.SetAttribute("AutoReserveMemory", Element.Attributes["AutoReserveMemory"].Value);
+                                }
+                                ToolElement.SetAttribute("OutputFileMasks", Element.Attributes["OutputFileMasks"].Value);
+                                //ToolElement.SetAttribute("AllowRestartOnLocal", "false");  //vs2012 can't really restart, so we will go with this for now
+                                if (Element.Attributes["OutputFileMasks"].Value == "PCLaunch.rc.res")
+                                {
+                                    // total hack, when doing clean compiles, this output directory does not exist, we need to create it now
+                                    string Parms = Element.Attributes["Params"].Value;
+                                    string Start = "/fo \"";
+                                    int StartIndex = Parms.IndexOf(Start);
+                                    if (StartIndex >= 0)
+                                    {
+                                        Parms = Parms.Substring(StartIndex + Start.Length);
+                                        int EndIndex = Parms.IndexOf("\"");
+                                        if (EndIndex > 0)
+                                        {
+                                            string ResLocation = CombinePaths(Parms.Substring(0, EndIndex));
+                                            if (!DirectoryExists_NoExceptions(GetDirectoryName(ResLocation)))
+                                            {
+                                                CreateDirectory(GetDirectoryName(ResLocation));
+                                            }
+                                        }
+                                    }
+                                }
+                            }
+                        }
+                    }
+                    {
+                        var NextDependencies = new List<string>();
+
+                        var Elements = UBTTask.GetElementsByTagName("Task");
+                        foreach (XmlElement Element in Elements)
+                        {
+                            string ToolName = string.Format("{0}_{1}", Element.Attributes["Tool"].Value, Job);
+                            string ActionName = string.Format("{0}_{1}", Element.Attributes["Name"].Value, Job);
+                            string OriginalActionName = ActionName;
+
+                            if (ToolToTool.ContainsKey(ToolName))
+                            {
+                                ToolName = ToolToTool[ToolName];
+                                ActionName = ToolToAction[ToolName];
+                                ActionToAction.Add(OriginalActionName, ActionName);
+                            }
+                            else
+                            {
+                                ActionToAction.Add(OriginalActionName, ActionName);
+                                ToolToAction.Add(ToolName, ActionName);
+                                // <Task ... />
+                                XmlElement TaskElement = XGETaskDocument.CreateElement("Task");
+                                ProjectElement.AppendChild(TaskElement);
+
+                                TaskElement.SetAttribute("SourceFile", Element.Attributes["SourceFile"].Value);
+                                if (Element.HasAttribute("Caption"))
+                                {
+                                    TaskElement.SetAttribute("Caption", Element.Attributes["Caption"].Value);
+                                }
+                                TaskElement.SetAttribute("Name", ActionName);
+                                NextDependencies.Add(ActionName);
+                                TaskElement.SetAttribute("Tool", ToolName);
+                                TaskElement.SetAttribute("WorkingDir", Element.Attributes["WorkingDir"].Value);
+                                TaskElement.SetAttribute("SkipIfProjectFailed", Element.Attributes["SkipIfProjectFailed"].Value);
+
+                                string NewDepends = "";
+                                if (Element.HasAttribute("DependsOn"))
+                                {
+                                    string Depends = Element.Attributes["DependsOn"].Value;
+                                    while (Depends.Length > 0)
+                                    {
+                                        string ThisAction = Depends;
+                                        int Semi = Depends.IndexOf(";");
+                                        if (Semi >= 0)
+                                        {
+                                            ThisAction = Depends.Substring(0, Semi);
+                                            Depends = Depends.Substring(Semi + 1);
+                                        }
+                                        else
+                                        {
+                                            Depends = "";
+                                        }
+                                        if (ThisAction.Length > 0)
+                                        {
+                                            if (NewDepends.Length > 0)
+                                            {
+                                                NewDepends += ";";
+                                            }
+                                            string ActionJob = ThisAction + string.Format("_{0}", Job);
+                                            if (!ActionToAction.ContainsKey(ActionJob))
+                                            {
+                                                throw new AutomationException("Action not found '{0}' in {1}->{2}", ActionJob, XGEFile, TargetFile);
+                                                // the XGE schedule is not topologically sorted. Hmmm. Well, we make a scary assumption here that this 
+                                            }
+                                            NewDepends += ActionToAction[ActionJob];
+                                        }
+                                    }
+                                }
+                                foreach (var Dep in CurrentDependencies)
+                                {
+                                    if (NewDepends.Length > 0)
+                                    {
+                                        NewDepends += ";";
+                                    }
+                                    NewDepends += Dep;
+                                }
+                                if (NewDepends != "")
+                                {
+                                    TaskElement.SetAttribute("DependsOn", NewDepends);
+                                }
+                            }
+
+                        }
+                        CurrentDependencies.AddRange(NextDependencies);
+                    }
+                    Job++;
+                }
+            }
+			return (Job > 0);
+        }
+
 		public void ClearExportedXGEXML()
 		{
 			foreach (var XGEFile in FindXGEFiles())
@@ -1215,7 +1236,8 @@ namespace AutomationTool
 		/// <param name="InDeleteBuildProducts">if specified, determines if the build products will be deleted before building. If not specified -clean parameter will be used,</param>
 		/// <param name="InUpdateVersionFiles">True if the version files are to be updated </param>
 		/// <param name="InForceNoXGE">If true will force XGE off</param>
-		public void Build(BuildAgenda Agenda, bool? InDeleteBuildProducts = null, bool InUpdateVersionFiles = true, bool InForceNoXGE = false, bool InForceNonUnity = false, bool InForceUnity = false, bool InShowProgress = false, Dictionary<UnrealBuildTool.UnrealTargetPlatform, Dictionary<string, string>> PlatformEnvVars = null)
+		/// <param name="InUseParallelExecutor">If true AND XGE not present or not being used then use ParallelExecutor</param>
+		public void Build(BuildAgenda Agenda, bool? InDeleteBuildProducts = null, bool InUpdateVersionFiles = true, bool InForceNoXGE = false, bool InUseParallelExecutor = false, bool InForceNonUnity = false, bool InForceUnity = false, bool InShowProgress = false, Dictionary<UnrealBuildTool.UnrealTargetPlatform, Dictionary<string, string>> PlatformEnvVars = null)
 		{
 			if (!CmdEnv.HasCapabilityToCompile)
 			{
@@ -1330,13 +1352,17 @@ namespace AutomationTool
 				bCanUseXGE = false;
 			}
 
+			// only run ParallelExecutor if not running XGE (and we've requested ParallelExecutor and it exists)
+			bool bCanUseParallelExecutor = !bCanUseXGE && InUseParallelExecutor && !string.IsNullOrEmpty(ParallelExecutorExec());
 			Log("************************* UE4Build:");
 			Log("************************* ForceMonolithic: {0}", bForceMonolithic);
 			Log("************************* ForceNonUnity:{0} ", bForceNonUnity);
 			Log("************************* ForceDebugInfo: {0}", bForceDebugInfo);
 			Log("************************* UseXGE: {0}", bCanUseXGE);
+            Log("************************* UseParallelExecutor: {0}", bCanUseParallelExecutor);
 
-			if (bCanUseXGE)
+			// process XGE files
+			if (bCanUseXGE || bCanUseParallelExecutor)
 			{
 				string TaskFilePath = CombinePaths(CmdEnv.LogFolder, @"UAT_XGE.xml");
 
@@ -1370,27 +1396,46 @@ namespace AutomationTool
 				}
 				if (XGEItems.Count > 0)
 				{
-					LogSetProgress(InShowProgress, "Building...");
-                    var StartXGERun = DateTime.Now.ToString();
-                    var FinishXGERun = "";
-					if (!RunXGE(XGEItems, TaskFilePath, Agenda.DoRetries, Agenda.SpecialTestFlag, InShowProgress))
+                    LogSetProgress(InShowProgress, "Building...");
+                    var StartRun = DateTime.Now.ToString();
+					var FinishRun = "";
+					string XGETool = "";
+					string Args = "";
+					if (bCanUseXGE) 
+					{
+						XGETool = XGEConsoleExePath;
+						Args = "\"" + TaskFilePath + "\" /Rebuild /MaxCPUS=200";
+					}
+					else if (bCanUseParallelExecutor)
+					{
+						XGETool = ParallelExecutorExec();
+						Args = "\"" + TaskFilePath + "\"";
+					}
+					if (XGETool.Equals(""))
 					{
-                        FinishXGERun = DateTime.Now.ToString();
-                        PrintCSVFile(String.Format("UAT,CompileWithXGE,{0},{1}", StartXGERun, FinishXGERun));
-						throw new AutomationException("BUILD FAILED: XGE failed, retries not enabled:");
+						throw new AutomationException("BUILD FAILED: no tool present to process XGE files");
+					}
+
+					if (!ProcessXGEItems(XGEItems, XGETool, Args, TaskFilePath, Agenda.DoRetries, Agenda.SpecialTestFlag, InShowProgress))
+					{
+						FinishRun = DateTime.Now.ToString();
+						PrintCSVFile(String.Format("UAT,CompileWithXGE,{0},{1}", StartRun, FinishRun));
+						throw new AutomationException("BUILD FAILED: {0} failed, retries not enabled:", XGETool);
 					}
 					else
 					{
 						bUsedXGE = true;
 					}
-                    FinishXGERun = DateTime.Now.ToString();
-                    PrintCSVFile(String.Format("UAT,CompileWithXGE,{0},{1}", StartXGERun, FinishXGERun));
+					
+					FinishRun = DateTime.Now.ToString();
+					PrintCSVFile(String.Format("UAT,CompileWithXGE,{0},{1}", StartRun, FinishRun));
 				}
 				else
 				{
 					bUsedXGE = true; // if there was nothing to build, we still consider this a success
 				}
 			}
+
 			if (!bUsedXGE)
 			{
 				if (DeleteBuildProducts)
diff --git a/Engine/Source/Programs/ParallelExecutor/Private/BuildGraph.cpp b/Engine/Source/Programs/ParallelExecutor/Private/BuildGraph.cpp
--- a/Engine/Source/Programs/ParallelExecutor/Private/BuildGraph.cpp
+++ b/Engine/Source/Programs/ParallelExecutor/Private/BuildGraph.cpp
@@ -43,7 +43,7 @@ uint32 FBuildActionExecutor::Run()
 
 int32 FBuildActionExecutor::ExecuteAction()
 {
-	int32 ExitCode = -1;
+	ExitCode = -1;
 
 	void* ReadPipe;
 	void* WritePipe;
@@ -220,6 +220,7 @@ bool FBuildGraph::ReadFromFile(const FString& InputPath)
 
 	// Read the tool environment
 	TMap<FString, const FXmlNode*> NameToTool;
+	TMap<FString, const FXmlNode*> EnvironmentVariables;
 	for(const FXmlNode* EnvironmentNode = EnvironmentsNode->GetFirstChildNode(); EnvironmentNode != nullptr; EnvironmentNode = EnvironmentNode->GetNextNode())
 	{
 		if(EnvironmentNode->GetTag() == TEXT("Environment"))
@@ -241,21 +242,12 @@ bool FBuildGraph::ReadFromFile(const FString& InputPath)
 				}
 			}
 
-			// Read the environment variables
+			// Read the environment variables for this environment. Each environment has its own set of variables 
 			const FXmlNode* VariablesNode = EnvironmentNode->FindChildNode(TEXT("Variables"));
 			if(VariablesNode != nullptr)
 			{
-				for(const FXmlNode* VariableNode = VariablesNode->GetFirstChildNode(); VariableNode != nullptr; VariableNode = VariableNode->GetNextNode())
-				{
-					if(VariableNode->GetTag() == TEXT("Variable"))
-					{
-						FString Key = VariableNode->GetAttribute(TEXT("Name"));
-						if(Key.Len() > 0)
-						{
-							Variables.FindOrAdd(Key) = VariableNode->GetAttribute(TEXT("Value"));
-						}
-					}
-				}
+				FString Key = EnvironmentNode->GetAttribute(TEXT("Name"));
+				EnvironmentVariables.FindOrAdd(Key) = VariablesNode;
 			}
 		}
 	}
@@ -295,8 +287,26 @@ bool FBuildGraph::ReadFromFile(const FString& InputPath)
 						}
 
 						VisitedActions.Empty();
-						RecursiveIncDependants(Action, VisitedActions);
-						Action->MissingDependencyCount = DependsOnList.Num();
+						RecursiveIncDependents(Action, VisitedActions);
+
+						Action->MissingDependencyCount = Action->Dependencies.Num();
+
+						// set environment variables for this action
+						const FXmlNode** VariablesNode = EnvironmentVariables.Find(ProjectNode->GetAttribute(TEXT("Env")));
+						if (VariablesNode != nullptr)
+						{
+							for (const FXmlNode* VariableNode = (*VariablesNode)->GetFirstChildNode(); VariableNode != nullptr; VariableNode = VariableNode->GetNextNode())
+							{
+								if (VariableNode->GetTag() == TEXT("Variable"))
+								{
+									FString Key = VariableNode->GetAttribute(TEXT("Name"));
+									if (Key.Len() > 0)
+									{
+										Action->Variables.FindOrAdd(Key) = VariableNode->GetAttribute(TEXT("Value"));
+									}
+								}
+							}
+						}
 					}
 				}
 			}
@@ -365,7 +375,7 @@ int32 FBuildGraph::ExecuteInParallel(int32 MaxProcesses)
 		while(ExecutingActions.Num() < MaxProcesses && QueuedActions.Num() > 0)
 		{
 			FBuildAction* Action = QueuedActions.Pop();
-			FBuildActionExecutor* ExecutingAction = new FBuildActionExecutor(Action, CriticalSection, CompletedEvent, CompletedActions, Variables, JobObject);
+			FBuildActionExecutor* ExecutingAction = new FBuildActionExecutor(Action, CriticalSection, CompletedEvent, CompletedActions, Action->Variables, JobObject);
 			FRunnableThread* ExecutingThread = FRunnableThread::Create(ExecutingAction, *FString::Printf(TEXT("Build:%s"), *Action->Caption));
 			ExecutingActions.Add(ExecutingAction, ExecutingThread);
 		}
@@ -412,7 +422,7 @@ int32 FBuildGraph::ExecuteInParallel(int32 MaxProcesses)
 				}
 			}
 
-			// Mark all the dependants as done
+			// Mark all the dependents as done
 			for(FBuildAction* DependantAction : CompletedAction->Action->Dependants)
 			{
 				if(--DependantAction->MissingDependencyCount == 0)
@@ -442,7 +452,7 @@ FBuildAction* FBuildGraph::FindOrAddAction(TMap<FString, FBuildAction*>& NameToA
 	return Action;
 }
 
-void FBuildGraph::RecursiveIncDependants(FBuildAction* Action, TSet<FBuildAction*>& VisitedActions)
+void FBuildGraph::RecursiveIncDependents(FBuildAction* Action, TSet<FBuildAction*>& VisitedActions)
 {
 	for(FBuildAction* Dependency : Action->Dependencies)
 	{
@@ -450,7 +460,7 @@ void FBuildGraph::RecursiveIncDependants(FBuildAction* Action, TSet<FBuildAction
 		{
 			VisitedActions.Add(Action);
 			Dependency->TotalDependants++;
-			RecursiveIncDependants(Dependency, VisitedActions);
+			RecursiveIncDependents(Dependency, VisitedActions);
 		}
 	}
 }
diff --git a/Engine/Source/Programs/ParallelExecutor/Private/BuildGraph.h b/Engine/Source/Programs/ParallelExecutor/Private/BuildGraph.h
--- a/Engine/Source/Programs/ParallelExecutor/Private/BuildGraph.h
+++ b/Engine/Source/Programs/ParallelExecutor/Private/BuildGraph.h
@@ -20,6 +20,8 @@ public:
 	int32 TotalDependants;
 	int32 MissingDependencyCount;
 
+	TMap<FString, FString> Variables;
+
 	FBuildAction();
 	~FBuildAction();
 };
@@ -57,10 +59,9 @@ public:
 
 private:
 	TArray<FBuildAction*> Actions;
-	TMap<FString, FString> Variables;
 
 	FBuildAction* FindOrAddAction(TMap<FString, FBuildAction*>& NameToAction, const FString& Name);
-	static void RecursiveIncDependants(FBuildAction* Action, TSet<FBuildAction*>& VisitedActions);
+	static void RecursiveIncDependents(FBuildAction* Action, TSet<FBuildAction*>& VisitedActions);
 
 	static void* CreateJobObject();
 };
