fix for Desktop in the path when trying to package
UE-15370

diff --git a/Engine/Source/Programs/AutomationTool/DeploymentContext.cs b/Engine/Source/Programs/AutomationTool/DeploymentContext.cs
--- a/Engine/Source/Programs/AutomationTool/DeploymentContext.cs
+++ b/Engine/Source/Programs/AutomationTool/DeploymentContext.cs
@@ -628,11 +628,37 @@ public class DeploymentContext //: ProjectParams
 					bool OtherPlatform = false;
 					foreach (UnrealTargetPlatform Plat in Enum.GetValues(typeof(UnrealTargetPlatform)))
 					{
-						if (Plat != StageTargetPlatform.PlatformType && Plat != UnrealTargetPlatform.Unknown && FileToCopy.IndexOf(CommandUtils.CombinePaths("/" + Plat.ToString() + "/"), 0, StringComparison.InvariantCultureIgnoreCase) >= 0)
-						{
-							OtherPlatform = true;
-							break;
-						}
+                        if (Plat != StageTargetPlatform.PlatformType && Plat != UnrealTargetPlatform.Unknown)
+                        {
+                            var Search = FileToCopy;
+                            if (Search.StartsWith(LocalRoot, StringComparison.InvariantCultureIgnoreCase))
+                            {
+                                if (LocalRoot.EndsWith("\\") || LocalRoot.EndsWith("/"))
+                                {
+                                    Search = Search.Substring(LocalRoot.Length - 1);
+                                }
+                                else
+                                {
+                                    Search = Search.Substring(LocalRoot.Length);
+                                }
+                            }
+							if (Search.StartsWith(ProjectRoot, StringComparison.InvariantCultureIgnoreCase))
+							{
+								if (ProjectRoot.EndsWith("\\") || ProjectRoot.EndsWith("/"))
+								{
+									Search = Search.Substring(ProjectRoot.Length - 1);
+								}
+								else
+								{
+									Search = Search.Substring(ProjectRoot.Length);
+								}
+							}
+                            if (Search.IndexOf(CommandUtils.CombinePaths("/" + Plat.ToString() + "/"), 0, StringComparison.InvariantCultureIgnoreCase) >= 0)
+                            {
+                                OtherPlatform = true;
+                                break;
+                            }
+                        }
 					}
 					if (OtherPlatform)
 					{
