cook on the fly server launch no longer needs the target platform

diff --git a/Engine/Source/Programs/AutomationTool/Scripts/RunProjectCommand.Automation.cs b/Engine/Source/Programs/AutomationTool/Scripts/RunProjectCommand.Automation.cs
--- a/Engine/Source/Programs/AutomationTool/Scripts/RunProjectCommand.Automation.cs
+++ b/Engine/Source/Programs/AutomationTool/Scripts/RunProjectCommand.Automation.cs
@@ -1007,9 +1007,8 @@ public partial class Project : CommandUtils
 	private static ProcessResult RunCookOnTheFlyServer(string ProjectName, string ServerLogFile, string TargetPlatform, string AdditionalCommandLine)
 	{
 		var ServerApp = HostPlatform.Current.GetUE4ExePath("UE4Editor.exe");
-		var Args = String.Format("{0} -run=cook -targetplatform={1} -cookonthefly -unattended -CrashForUAT -FORCELOGFLUSH -log",
-			CommandUtils.MakePathSafeToUseWithCommandLine(ProjectName),
-			TargetPlatform);
+		var Args = String.Format("{0} -run=cook -cookonthefly -unattended -CrashForUAT -FORCELOGFLUSH -log",
+			CommandUtils.MakePathSafeToUseWithCommandLine(ProjectName));
 		if (!String.IsNullOrEmpty(ServerLogFile))
 		{
 			Args += " -abslog=" + CommandUtils.MakePathSafeToUseWithCommandLine(ServerLogFile);
