remove Copy of TaskFilePath to my personal developer folder

diff --git a/Engine/Source/Programs/AutomationTool/UE4Build.cs b/Engine/Source/Programs/AutomationTool/UE4Build.cs
--- a/Engine/Source/Programs/AutomationTool/UE4Build.cs
+++ b/Engine/Source/Programs/AutomationTool/UE4Build.cs
@@ -840,8 +840,6 @@ namespace AutomationTool
 					throw new AutomationException("Unable to find xge xml: " + TaskFilePath);
 				}
 
-				System.IO.File.Copy(TaskFilePath, @"P:\Developers\Paul.Fazio\XGE.xml", true);
-
 				int Retries = DoRetries ? 2 : 1;
                 int ConnectionRetries = 4;
                 for (int i = 0; i < Retries; i++)
