fix for installed plugins enabled by default kicking a link build for content-only projects
UE-16385

diff --git a/Engine/Source/Programs/AutomationTool/ProjectUtils.cs b/Engine/Source/Programs/AutomationTool/ProjectUtils.cs
--- a/Engine/Source/Programs/AutomationTool/ProjectUtils.cs
+++ b/Engine/Source/Programs/AutomationTool/ProjectUtils.cs
@@ -218,7 +218,8 @@ namespace AutomationTool
 				// find if there are any plugins enabled or disabled which differ from the default
 				foreach(PluginInfo Plugin in AvailablePlugins)
 				{
-					if(UProjectInfo.IsPluginEnabledForProject(Plugin, Project, TargetPlatformType) != Plugin.Descriptor.bEnabledByDefault)
+					bool bPluginEnabledForProject = UProjectInfo.IsPluginEnabledForProject(Plugin, Project, TargetPlatformType);
+					if (bPluginEnabledForProject != Plugin.Descriptor.bEnabledByDefault || (bPluginEnabledForProject && Plugin.Descriptor.bInstalled))
 					{
 						if(Plugin.Descriptor.Modules.Any(Module => Module.IsCompiledInConfiguration(TargetPlatformType, TargetRules.TargetType.Game, bBuildDeveloperTools: false, bBuildEditor: false)))
 						{
diff --git a/Engine/Source/Runtime/Projects/Private/ProjectManager.cpp b/Engine/Source/Runtime/Projects/Private/ProjectManager.cpp
--- a/Engine/Source/Runtime/Projects/Private/ProjectManager.cpp
+++ b/Engine/Source/Runtime/Projects/Private/ProjectManager.cpp
@@ -266,7 +266,7 @@ bool FProjectManager::IsNonDefaultPluginEnabled() const
 
 	for(const FPluginStatus& Plugin: IPluginManager::Get().QueryStatusForAllPlugins())
 	{
-		if((Plugin.LoadedFrom == EPluginLoadedFrom::GameProject || !Plugin.Descriptor.bEnabledByDefault) && EnabledPlugins.Contains(Plugin.Name))
+		if ((Plugin.LoadedFrom == EPluginLoadedFrom::GameProject || !Plugin.Descriptor.bEnabledByDefault || Plugin.Descriptor.bInstalled) && EnabledPlugins.Contains(Plugin.Name))
 		{
 			return true;
 		}
