build script, disable xge in all gubp compiles

diff --git a/Engine/Source/Programs/AutomationTool/Scripts/GUBP.Automation.cs b/Engine/Source/Programs/AutomationTool/Scripts/GUBP.Automation.cs
--- a/Engine/Source/Programs/AutomationTool/Scripts/GUBP.Automation.cs
+++ b/Engine/Source/Programs/AutomationTool/Scripts/GUBP.Automation.cs
@@ -745,7 +745,7 @@ public class GUBP : BuildCommand
                 bool ReallyDeleteBuildProducts = DeleteBuildProducts() && !GUBP.bForceIncrementalCompile;
                 Agenda.DoRetries = false; // these would delete build products
 				bool UseParallelExecutor = bDependentOnCompileTools && (HostPlatform == UnrealTargetPlatform.Win64);
-				UE4Build.Build(Agenda, InDeleteBuildProducts: ReallyDeleteBuildProducts, InUpdateVersionFiles: false, InForceUnity: true, InUseParallelExecutor: UseParallelExecutor);
+				UE4Build.Build(Agenda, InDeleteBuildProducts: ReallyDeleteBuildProducts, InUpdateVersionFiles: false, InForceNoXGE: true, InForceUnity: true, InUseParallelExecutor: UseParallelExecutor);
                 var StartPostBuild = DateTime.Now.ToString();
                 PostBuild(bp, UE4Build);
                 var FinishPostBuild = DateTime.Now.ToString();
