build script, retry copy publish on mac

diff --git a/Engine/Source/Programs/AutomationTool/Scripts/RocketBuild.Automation.cs b/Engine/Source/Programs/AutomationTool/Scripts/RocketBuild.Automation.cs
--- a/Engine/Source/Programs/AutomationTool/Scripts/RocketBuild.Automation.cs
+++ b/Engine/Source/Programs/AutomationTool/Scripts/RocketBuild.Automation.cs
@@ -906,7 +906,26 @@ namespace Rocket
 					if(DebugExtensions.Contains(Extension) || Extension == ".exe" || Extension == ".dll") // Need all windows build products for crash reporter
 					{
 						string OutputFileName = CommandUtils.MakeRerootedFilePath(InputFileName, CommandUtils.CmdEnv.LocalRoot, SymbolsOutputDir);
-						CommandUtils.CopyFile(InputFileName, OutputFileName);
+                        if (SymbolsOutputDir.StartsWith("/Volumes"))
+                        {
+                            int Retry = 0;
+                            int NumRetries = 60;
+                            bool bCopied = false;
+                            while (!bCopied && Retry < NumRetries)
+                            {
+                                if (Retry > 0)
+                                {
+                                    CommandUtils.Log("*** Mac temp storage retry {0}", OutputFileName);
+                                    System.Threading.Thread.Sleep(1000);
+                                }
+                                bCopied = CommandUtils.CopyFile_NoExceptions(InputFileName, OutputFileName, true);
+                                Retry++;
+                            }
+                        }
+                        else
+                        {
+                            CommandUtils.CopyFile(InputFileName, OutputFileName);
+                        }
 					}
 				}
 			}
