Match ProjectUtils plugin enabled check to ProjectManager check (don't expect build if enabled by default but disabled in project)
#ue4

diff --git a/Engine/Source/Programs/AutomationTool/ProjectUtils.cs b/Engine/Source/Programs/AutomationTool/ProjectUtils.cs
--- a/Engine/Source/Programs/AutomationTool/ProjectUtils.cs
+++ b/Engine/Source/Programs/AutomationTool/ProjectUtils.cs
@@ -221,7 +221,7 @@ namespace AutomationTool
 				foreach(PluginInfo Plugin in AvailablePlugins)
 				{
 					bool bPluginEnabledForProject = UProjectInfo.IsPluginEnabledForProject(Plugin, Project, TargetPlatformType);
-					if (bPluginEnabledForProject != Plugin.Descriptor.bEnabledByDefault || (bPluginEnabledForProject && Plugin.Descriptor.bInstalled))
+					if ((bPluginEnabledForProject && !Plugin.Descriptor.bEnabledByDefault) || (bPluginEnabledForProject && Plugin.Descriptor.bInstalled))
 					{
 						if(Plugin.Descriptor.Modules.Any(Module => Module.IsCompiledInConfiguration(TargetPlatformType, TargetRules.TargetType.Game, bBuildDeveloperTools: false, bBuildEditor: false)))
 						{
