potential fix for mis-matching directories when creating the path to the release directory for a platform

diff --git a/Engine/Source/Programs/AutomationTool/Scripts/CopyBuildToStagingDirectory.Automation.cs b/Engine/Source/Programs/AutomationTool/Scripts/CopyBuildToStagingDirectory.Automation.cs
--- a/Engine/Source/Programs/AutomationTool/Scripts/CopyBuildToStagingDirectory.Automation.cs
+++ b/Engine/Source/Programs/AutomationTool/Scripts/CopyBuildToStagingDirectory.Automation.cs
@@ -678,7 +678,7 @@ public partial class Project : CommandUtils
 
     private static string GetReleasePakFilePath(DeploymentContext SC, ProjectParams Params, string ReleaseVersion, string PakName)
     {
-        string ReleaseVersionPath = CombinePaths(SC.ProjectRoot, "Releases", ReleaseVersion, SC.StageTargetPlatform.GetCookPlatform(Params.DedicatedServer, false, Params.CookFlavor), PakName);
+        string ReleaseVersionPath = CombinePaths(SC.ProjectRoot, "Releases", ReleaseVersion, SC.StageTargetPlatform.GetCookPlatform(SC.DedicatedServer, false, Params.CookFlavor), PakName);
         return ReleaseVersionPath;
     }
 
