Add a command line option to use parallel executor.

diff --git a/Engine/Source/Programs/AutomationTool/Scripts/MegaXGE.Automation.cs b/Engine/Source/Programs/AutomationTool/Scripts/MegaXGE.Automation.cs
--- a/Engine/Source/Programs/AutomationTool/Scripts/MegaXGE.Automation.cs
+++ b/Engine/Source/Programs/AutomationTool/Scripts/MegaXGE.Automation.cs
@@ -145,7 +145,7 @@ class MegaXGE : BuildCommand
 		Log("*************************");
 
 		Agenda.DoRetries = ParseParam("Retry");
-		UE4Build.Build(Agenda, InUpdateVersionFiles: IsBuildMachine, InShowProgress: ShowProgress);
+		UE4Build.Build(Agenda, InUpdateVersionFiles: IsBuildMachine, InUseParallelExecutor: ParseParam("useparallelexecutor"), InShowProgress: ShowProgress);
 
 		// 		if (WorkingCL > 0) // only move UAT files if we intend to check in some build products
 		// 		{
