[INTEGRATE] Change 2563617 by [REDACTED_PERSON][REDACTED_HANDLE] on 2015/05/22 19:15:48

	Increase the priority of the StripRootEditor node, so it runs directly after RootEditor (and before some random sample decides to relink plugins and invalidate the PDBs).

diff --git a/Engine/Source/Programs/AutomationTool/Scripts/RocketBuild.Automation.cs b/Engine/Source/Programs/AutomationTool/Scripts/RocketBuild.Automation.cs
--- a/Engine/Source/Programs/AutomationTool/Scripts/RocketBuild.Automation.cs
+++ b/Engine/Source/Programs/AutomationTool/Scripts/RocketBuild.Automation.cs
@@ -426,6 +426,11 @@ namespace Rocket
 			AgentSharingGroup = "Editor" + StaticGetHostPlatformSuffix(HostPlatform);
 		}
 
+		public override float Priority()
+		{
+			return 1000000.0f;
+		}
+
 		public static string StaticGetFullName(UnrealTargetPlatform InHostPlatform)
 		{
 			return "Strip" + GUBP.RootEditorNode.StaticGetFullName(InHostPlatform);
