make it so DeploymentContext now includes a ProjectBinariesFolder, this variable is the path to the binaries folder for the specific context (including platform)

diff --git a/Engine/Source/Programs/AutomationTool/DeploymentContext.cs b/Engine/Source/Programs/AutomationTool/DeploymentContext.cs
--- a/Engine/Source/Programs/AutomationTool/DeploymentContext.cs
+++ b/Engine/Source/Programs/AutomationTool/DeploymentContext.cs
@@ -164,6 +164,10 @@ public class DeploymentContext //: ProjectParams
 	///  Directory to archive all of the files in: d:\archivedir\WindowsNoEditor
 	/// </summary>
 	public string ArchiveDirectory;
+	/// <summary>
+	///  Directory to project binaries
+	/// </summary>
+	public string ProjectBinariesFolder;
 
 	/// <summary>
 	/// Filename for the manifest of file changes for iterative deployment.
@@ -326,6 +330,7 @@ public class DeploymentContext //: ProjectParams
 			CommandUtils.CreateDirectory(ArchiveDirectory);
 		}
 		ProjectArgForCommandLines = ProjectArgForCommandLines.Replace("\\", "/");
+		ProjectBinariesFolder = CommandUtils.CombinePaths(ProjectUtils.GetClientProjectBinariesRootPath(RawProjectPath, TargetRules.TargetType.Game, IsCodeBasedProject), PlatformDir);
 	}
 
 	public void StageFile(StagedFileType FileType, string InputPath, string OutputPath = null, bool bRemap = true)
