fix for not finding plugins properly when running a packaged content-only game
UE-16540

diff --git a/Engine/Source/Programs/AutomationTool/Win/WinPlatform.Automation.cs b/Engine/Source/Programs/AutomationTool/Win/WinPlatform.Automation.cs
--- a/Engine/Source/Programs/AutomationTool/Win/WinPlatform.Automation.cs
+++ b/Engine/Source/Programs/AutomationTool/Win/WinPlatform.Automation.cs
@@ -236,7 +236,7 @@ public abstract class BaseWinPlatform : Platform
 
 	public override bool ShouldStageCommandLine(ProjectParams Params, DeploymentContext SC)
 	{
-		return !String.IsNullOrEmpty(Params.StageCommandline) || !String.IsNullOrEmpty(Params.RunCommandline) || (!Params.IsCodeBasedProject && Params.NoBootstrapExe);
+		return false; // !String.IsNullOrEmpty(Params.StageCommandline) || !String.IsNullOrEmpty(Params.RunCommandline) || (!Params.IsCodeBasedProject && Params.NoBootstrapExe);
 	}
 
 	public override List<string> GetDebugFileExtentions()
