Allow missing receipts when we're specifying more than one platform, to handle cases when targets only exist for certain platforms.

diff --git a/Engine/Source/Programs/AutomationTool/Scripts/CopyBuildToStagingDirectory.Automation.cs b/Engine/Source/Programs/AutomationTool/Scripts/CopyBuildToStagingDirectory.Automation.cs
--- a/Engine/Source/Programs/AutomationTool/Scripts/CopyBuildToStagingDirectory.Automation.cs
+++ b/Engine/Source/Programs/AutomationTool/Scripts/CopyBuildToStagingDirectory.Automation.cs
@@ -1463,7 +1463,7 @@ public partial class Project : CommandUtils
 					UnrealTargetPlatform[] SubPlatformsToStage = PlatformInstance.GetStagePlatforms();
 
 					// if we are attempting to gathering multiple platforms, the files aren't required
-					bool bRequireStagedFilesToExist = SubPlatformsToStage.Length == 1;
+					bool bRequireStagedFilesToExist = SubPlatformsToStage.Length == 1 && PlatformsToStage.Count == 1;
 
 					foreach (UnrealTargetPlatform ReceiptPlatform in SubPlatformsToStage)
 					{
