fix for checking the entire project for a target.cs file when we should only check the entire Source directory
UE-16449

diff --git a/Engine/Source/Programs/AutomationTool/ProjectUtils.cs b/Engine/Source/Programs/AutomationTool/ProjectUtils.cs
--- a/Engine/Source/Programs/AutomationTool/ProjectUtils.cs
+++ b/Engine/Source/Programs/AutomationTool/ProjectUtils.cs
@@ -164,11 +164,11 @@ namespace AutomationTool
 			{
 				return false;
 			}
-			else 
+			else if (Directory.Exists(Path.Combine(Path.GetDirectoryName(RawProjectPath), "Source")))
 			{
 				// wasn't one in the main Source directory, let's check all sub-directories
 				//@todo: may want to read each target.cs to see if it has a target corresponding to the project name as a final check
-				FileInfo[] Files = (new DirectoryInfo (Path.GetDirectoryName (RawProjectPath)).GetFiles ("*.Target.cs", SearchOption.AllDirectories));
+				FileInfo[] Files = (new DirectoryInfo( Path.Combine (Path.GetDirectoryName (RawProjectPath), "Source")).GetFiles ("*.Target.cs", SearchOption.AllDirectories));
 				if (Files.Length > 0)
 				{
 					return false;
