UE-16552 - Application can crash when returning from a paused OBB download
- Changed logic for dealing with the download activity being backed out of
- Updated .bat file writing so that a second bat file which doesn't push the obb (but can remove it) is written

diff --git a/Engine/Build/Android/Java/JavaTemplates/DownloaderActivity.java.template b/Engine/Build/Android/Java/JavaTemplates/DownloaderActivity.java.template
--- a/Engine/Build/Android/Java/JavaTemplates/DownloaderActivity.java.template
+++ b/Engine/Build/Android/Java/JavaTemplates/DownloaderActivity.java.template
@@ -663,10 +663,13 @@ public class DownloaderActivity extends Activity implements IDownloaderClient {
 	@Override
 	protected void onPause() {
 		super.onPause();
-	
+		GameActivity.Log.debug("In onPause");
+		
 		if(OutputData.getIntExtra(GameActivity.DOWNLOAD_RETURN_NAME, GameActivity.DOWNLOAD_NO_RETURN_CODE) == GameActivity.DOWNLOAD_NO_RETURN_CODE)
 		{
+			GameActivity.Log.debug("onPause returning that user quit the download.");
 			OutputData.putExtra(GameActivity.DOWNLOAD_RETURN_NAME, GameActivity.DOWNLOAD_USER_QUIT);
+			setResult(RESULT_OK, OutputData);
 		}
 	}
 
@@ -679,6 +682,7 @@ public class DownloaderActivity extends Activity implements IDownloaderClient {
             mDownloaderClientStub.disconnect(this);
         }
         super.onStop();
+		setResult(RESULT_OK, OutputData);
     }
 
     /**
diff --git a/Engine/Build/Android/Java/src/com/epicgames/ue4/GameActivity.java b/Engine/Build/Android/Java/src/com/epicgames/ue4/GameActivity.java
--- a/Engine/Build/Android/Java/src/com/epicgames/ue4/GameActivity.java
+++ b/Engine/Build/Android/Java/src/com/epicgames/ue4/GameActivity.java
@@ -1076,38 +1076,48 @@ public class GameActivity extends NativeActivity
 	{
 		if( requestCode == DOWNLOAD_ACTIVITY_ID)
 		{
-			int errorCode = data.getIntExtra(DOWNLOAD_RETURN_NAME, DOWNLOAD_NO_RETURN_CODE);
-			
-			String logMsg = "DownloadActivity Returned with ";
-			switch(errorCode)
+			int errorCode = 0;
+			if(resultCode == RESULT_OK)
 			{
-			case DOWNLOAD_FILES_PRESENT:
-				logMsg += "Download Files Present";
-				break;
-			case DOWNLOAD_COMPLETED_OK:
-				logMsg += "Download Completed OK";
-				break;
-			case DOWNLOAD_NO_RETURN_CODE:
-				logMsg += "Download No Return Code";
-				break;
-			case DOWNLOAD_USER_QUIT:
-				logMsg += "Download User Quit";
-				break;
-			case DOWNLOAD_FAILED:
-				logMsg += "Download Failed";
-				break;
-			case DOWNLOAD_INVALID:
-				logMsg += "Download Invalid";
-				break;
-			case DOWNLOAD_NO_PLAY_KEY:
-				logMsg +="Download No Play Key";
-				break;
-			default:
-				logMsg += "Unknown message!";
-				break;
+								
+				errorCode = data.getIntExtra(DOWNLOAD_RETURN_NAME, DOWNLOAD_NO_RETURN_CODE);
+				
+				String logMsg = "DownloadActivity Returned with ";
+				switch(errorCode)
+				{
+				case DOWNLOAD_FILES_PRESENT:
+					logMsg += "Download Files Present";
+					break;
+				case DOWNLOAD_COMPLETED_OK:
+					logMsg += "Download Completed OK";
+					break;
+				case DOWNLOAD_NO_RETURN_CODE:
+					logMsg += "Download No Return Code";
+					break;
+				case DOWNLOAD_USER_QUIT:
+					logMsg += "Download User Quit";
+					break;
+				case DOWNLOAD_FAILED:
+					logMsg += "Download Failed";
+					break;
+				case DOWNLOAD_INVALID:
+					logMsg += "Download Invalid";
+					break;
+				case DOWNLOAD_NO_PLAY_KEY:
+					logMsg +="Download No Play Key";
+					break;
+				default:
+					logMsg += "Unknown message!";
+					break;
+				}
+				
+				Log.debug(logMsg);
+			}
+			else
+			{
+				Log.debug("Download activity cancelled by user.");
+				errorCode = DOWNLOAD_USER_QUIT;
 			}
-			
-			Log.debug(logMsg);
 			
 			HasAllFiles = (errorCode == DOWNLOAD_FILES_PRESENT || errorCode == DOWNLOAD_COMPLETED_OK);
 			
diff --git a/Engine/Source/Programs/AutomationTool/Android/AndroidPlatform.Automation.cs b/Engine/Source/Programs/AutomationTool/Android/AndroidPlatform.Automation.cs
--- a/Engine/Source/Programs/AutomationTool/Android/AndroidPlatform.Automation.cs
+++ b/Engine/Source/Programs/AutomationTool/Android/AndroidPlatform.Automation.cs
@@ -107,9 +107,9 @@ public class AndroidPlatform : Platform
 		}
     }
 
-	private static string GetFinalBatchName(string ApkName, ProjectParams Params, string Architecture, string GPUArchitecture)
+	private static string GetFinalBatchName(string ApkName, ProjectParams Params, string Architecture, string GPUArchitecture, bool NoOBBInstall)
 	{
-		return Path.Combine(Path.GetDirectoryName(ApkName), "Install_" + Params.ShortProjectName + "_" + Params.ClientConfigsToBuild[0].ToString() + Architecture + GPUArchitecture + (Utils.IsRunningOnMono ? ".command" : ".bat"));
+		return Path.Combine(Path.GetDirectoryName(ApkName), "Install_" + Params.ShortProjectName + (!NoOBBInstall ? "_" : "_NoOBBInstall_") + Params.ClientConfigsToBuild[0].ToString() + Architecture + GPUArchitecture + (Utils.IsRunningOnMono ? ".command" : ".bat"));
 	}
 
 	public override void Package(ProjectParams Params, DeploymentContext SC, int WorkingCL)
@@ -174,8 +174,7 @@ public class AndroidPlatform : Platform
 					}
 				}
 				
-				string BatchName = GetFinalBatchName(ApkName, Params, bMakeSeparateApks ? Architecture : "", bMakeSeparateApks ? GPUArchitecture : "");
-
+				              
 				if (!Params.Prebuilt)
 				{
 					string CookFlavor = SC.FinalCookPlatform.IndexOf("_") > 0 ? SC.FinalCookPlatform.Substring(SC.FinalCookPlatform.IndexOf("_")) : "";
@@ -194,14 +193,45 @@ public class AndroidPlatform : Platform
 			    }
 
 				// Write install batch file(s).
-
+                string BatchName = GetFinalBatchName(ApkName, Params, bMakeSeparateApks ? Architecture : "", bMakeSeparateApks ? GPUArchitecture : "", false);
                 string PackageName = GetPackageInfo(ApkName, false);
 				// make a batch file that can be used to install the .apk and .obb files
-				string[] BatchLines;
+				string[] BatchLines = GenerateInstallBatchFile(bPackageDataInsideApk, PackageName, ApkName, Params, ObbName, DeviceObbName, false);
+				File.WriteAllLines(BatchName, BatchLines);
+
+                // If we aren't packaging data in the APK then lets write out a bat file to also let us test without the OBB
+                // on the device.
+                String NoInstallBatchName = GetFinalBatchName(ApkName, Params, bMakeSeparateApks ? Architecture : "", bMakeSeparateApks ? GPUArchitecture : "", true);
+                // if(!bPackageDataInsideApk)
+                {
+                    BatchLines = GenerateInstallBatchFile(bPackageDataInsideApk, PackageName, ApkName, Params, ObbName, DeviceObbName, true);
+                    File.WriteAllLines(NoInstallBatchName, BatchLines);
+                }
+
 				if (Utils.IsRunningOnMono)
 				{
-					Log("Writing shell script for install with {0}", bPackageDataInsideApk ? "data in APK" : "separate obb");
-					BatchLines = new string[] {
+					CommandUtils.FixUnixFilePermissions(BatchName);
+                    if(File.Exists(NoInstallBatchName)) 
+                    {
+                        CommandUtils.FixUnixFilePermissions(NoInstallBatchName);
+                    }
+				}
+			}
+		}
+
+		PrintRunTime();
+	}
+
+    private string[] GenerateInstallBatchFile(bool bPackageDataInsideApk, string PackageName, string ApkName, ProjectParams Params, string ObbName, string DeviceObbName, bool bNoObbInstall)
+    {
+        string[] BatchLines = null;
+        
+        if (Utils.IsRunningOnMono)
+        {
+            string OBBInstallCommand = bNoObbInstall ? "shell 'rm -r $EXTERNAL_STORAGE/" + DeviceObbName + "'" : "push " + Path.GetFileName(ObbName) + " $STORAGE/" + DeviceObbName;
+
+            Log("Writing shell script for install with {0}", bPackageDataInsideApk ? "data in APK" : "separate obb");
+            BatchLines = new string[] {
 						"#!/bin/sh",
 						"cd \"`dirname \"$0\"`\"",
                         "ADB=",
@@ -223,7 +253,7 @@ public class AndroidPlatform : Platform
 						bPackageDataInsideApk ? "" : "\techo",
 						bPackageDataInsideApk ? "" : "\techo Installing new data. Failures here indicate storage problems \\(missing SD card or bad permissions\\) and are fatal.",
 						bPackageDataInsideApk ? "" : "\tSTORAGE=$(echo \"`$ADB $DEVICE shell 'echo $EXTERNAL_STORAGE'`\" | cat -v | tr -d '^M')",
-						bPackageDataInsideApk ? "" : "\t$ADB $DEVICE push " + Path.GetFileName(ObbName) + " $STORAGE/" + DeviceObbName,
+						bPackageDataInsideApk ? "" : "\t$ADB $DEVICE " + OBBInstallCommand,
 						bPackageDataInsideApk ? "if [ 1 ]; then" : "\tif [ $? -eq 0 ]; then",
 						"\t\techo",
 						"\t\techo Installation successful",
@@ -239,11 +269,13 @@ public class AndroidPlatform : Platform
 						"echo Check that the device has an SD card.",
 						"exit 1"
 					};
-				}
-				else
-				{
-					Log("Writing bat for install with {0}", bPackageDataInsideApk ? "data in APK" : "separate OBB");
-					BatchLines = new string[] {
+        }
+        else
+        {
+            string OBBInstallCommand = bNoObbInstall ? "shell rm -r %STORAGE%/" + DeviceObbName : "push " + Path.GetFileName(ObbName) + " %STORAGE%/" + DeviceObbName;
+
+            Log("Writing bat for install with {0}", bPackageDataInsideApk ? "data in APK" : "separate OBB");
+            BatchLines = new string[] {
 						"setlocal",
                         "set ANDROIDHOME=%ANDROID_HOME%",
                         "if \"%ANDROIDHOME%\"==\"\" set ANDROIDHOME="+Environment.GetEnvironmentVariable("ANDROID_HOME"),
@@ -263,7 +295,7 @@ public class AndroidPlatform : Platform
 						"%ADB% %DEVICE% shell rm -r %STORAGE%/" + TargetAndroidLocation + PackageName,
 						bPackageDataInsideApk ? "" : "@echo.",
 						bPackageDataInsideApk ? "" : "@echo Installing new data. Failures here indicate storage problems (missing SD card or bad permissions) and are fatal.",
-						bPackageDataInsideApk ? "" : "%ADB% %DEVICE% push " + Path.GetFileName(ObbName) + " %STORAGE%/" + DeviceObbName,
+						bPackageDataInsideApk ? "" : "%ADB% %DEVICE% " + OBBInstallCommand,
 						bPackageDataInsideApk ? "" : "if \"%ERRORLEVEL%\" NEQ \"0\" goto Error",
 						"@echo.",
 						"@echo Installation successful",
@@ -278,19 +310,9 @@ public class AndroidPlatform : Platform
 						"@echo Check that the device has an SD card.",
 						"@pause"
 					};
-				}
-				File.WriteAllLines(BatchName, BatchLines);
-
-				if (Utils.IsRunningOnMono)
-				{
-					CommandUtils.FixUnixFilePermissions(BatchName);
-				}
-			}
-		}
-
-		PrintRunTime();
-	}
-
+        } 
+        return BatchLines;
+    }
 	public override void GetFilesToArchive(ProjectParams Params, DeploymentContext SC)
 	{
 		if (SC.StageTargetConfigurations.Count != 1)
@@ -312,7 +334,8 @@ public class AndroidPlatform : Platform
 			{
 				string ApkName = GetFinalApkName(Params, SC.StageExecutables[0], true, bMakeSeparateApks ? Architecture : "", bMakeSeparateApks ? GPUArchitecture : "");
 				string ObbName = GetFinalObbName(ApkName);
-				string BatchName = GetFinalBatchName(ApkName, Params, bMakeSeparateApks ? Architecture : "", bMakeSeparateApks ? GPUArchitecture : "");
+				string BatchName = GetFinalBatchName(ApkName, Params, bMakeSeparateApks ? Architecture : "", bMakeSeparateApks ? GPUArchitecture : "", false);
+                string NoOBBBatchName = GetFinalBatchName(ApkName, Params, bMakeSeparateApks ? Architecture : "", bMakeSeparateApks ? GPUArchitecture : "", true);
 
 				// verify the files exist
 				if (!FileExists(ApkName))
@@ -336,6 +359,8 @@ public class AndroidPlatform : Platform
 				}
 
 				SC.ArchiveFiles(Path.GetDirectoryName(BatchName), Path.GetFileName(BatchName));
+                SC.ArchiveFiles(Path.GetDirectoryName(NoOBBBatchName), Path.GetFileName(NoOBBBatchName));
+
 			}
 		}
 	}
