UAT: Added two new project params:

-IgnoreCookErrors - ignores cook errors and continues to launch/stage builds
-FastCook - fast cook option to pass into the cooker

diff --git a/Engine/Source/Programs/AutomationTool/ProjectParams.cs b/Engine/Source/Programs/AutomationTool/ProjectParams.cs
--- a/Engine/Source/Programs/AutomationTool/ProjectParams.cs
+++ b/Engine/Source/Programs/AutomationTool/ProjectParams.cs
@@ -257,6 +257,8 @@ namespace AutomationTool
 			this.Package = InParams.Package;
 			this.Deploy = InParams.Deploy;
 			this.IterativeDeploy = InParams.IterativeDeploy;
+			this.IgnoreCookErrors = InParams.IgnoreCookErrors;
+			this.FastCook = InParams.FastCook;
 			this.Device = InParams.Device;
 			this.DeviceName = InParams.DeviceName;
 			this.ServerDevice = InParams.ServerDevice;
@@ -391,7 +393,9 @@ namespace AutomationTool
             bool? Prebuilt = null,
             int? RunTimeoutSeconds = null,
 			string SpecifiedArchitecture = null,
-            bool? IterativeDeploy = null
+            bool? IterativeDeploy = null,
+			bool? FastCook = null,
+			bool? IgnoreCookErrors = null
 			)
 		{
 			//
@@ -534,6 +538,8 @@ namespace AutomationTool
 			this.Package = GetParamValueIfNotSpecified(Command, Package, this.Package, "package");
 			this.Deploy = GetParamValueIfNotSpecified(Command, Deploy, this.Deploy, "deploy");
 			this.IterativeDeploy = GetParamValueIfNotSpecified(Command, IterativeDeploy, this.IterativeDeploy, new string[] {"iterativedeploy", "iterate" } );
+			this.FastCook = GetParamValueIfNotSpecified(Command, FastCook, this.FastCook, "FastCook");
+			this.IgnoreCookErrors = GetParamValueIfNotSpecified(Command, IgnoreCookErrors, this.IgnoreCookErrors, "IgnoreCookErrors");
 			this.Device = ParseParamValueIfNotSpecified(Command, Device, "device", String.Empty).Trim(new char[] { '\"' });
 
 			// strip the platform prefix the specified device.
@@ -1098,7 +1104,14 @@ namespace AutomationTool
 		[Help("iterativecooking", "Uses the iterative cooking, command line: -iterativedeploy or -iterate")]
 		public bool IterativeDeploy;
 
+		[Help("FastCook", "Uses fast cook path if supported by target")]
+		public bool FastCook;
 
+		/// <summary>
+		/// Cook: Ignores cook errors and continues with packaging etc.
+		/// </summary>
+		[Help("IgnoreCookErrors", "Ignores cook errors and continues with packaging etc")]
+		public bool IgnoreCookErrors { private set; get; }
 
 		#endregion
 
@@ -2024,6 +2037,7 @@ namespace AutomationTool
                 CommandUtils.Log("CookMapsOnly={0}", CookMapsOnly);
                 CommandUtils.Log("Deploy={0}", Deploy);
 				CommandUtils.Log("IterativeDeploy={0}", IterativeDeploy);
+				CommandUtils.Log("FastCook={0}", FastCook);
 				CommandUtils.Log("LogWindow={0}", LogWindow);
 				CommandUtils.Log("Manifests={0}", Manifests);
 				CommandUtils.Log("MapToRun={0}", MapToRun);
diff --git a/Engine/Source/Programs/AutomationTool/Scripts/CookCommand.Automation.cs b/Engine/Source/Programs/AutomationTool/Scripts/CookCommand.Automation.cs
--- a/Engine/Source/Programs/AutomationTool/Scripts/CookCommand.Automation.cs
+++ b/Engine/Source/Programs/AutomationTool/Scripts/CookCommand.Automation.cs
@@ -135,6 +135,10 @@ public partial class Project : CommandUtils
                 {
                     CommandletParams += " -unversioned";
                 }
+				if (Params.FastCook)
+				{
+					CommandletParams += " -FastCook";
+				}
                 if (Params.UseDebugParamForEditorExe)
                 {
                     CommandletParams += " -debug";
@@ -214,11 +218,18 @@ public partial class Project : CommandUtils
             }
 			catch (Exception Ex)
 			{
-				// Delete cooked data (if any) as it may be incomplete / corrupted.
-				Log("Cook failed. Deleting cooked data.");
-				CleanupCookedData(PlatformsToCook.ToList(), Params);
-				AutomationTool.ErrorReporter.Error("Cook failed.", (int)AutomationTool.ErrorCodes.Error_UnknownCookFailure);
-				throw Ex;
+				if (Params.IgnoreCookErrors)
+				{
+					LogWarning("Ignoring cook failure.");
+				}
+				else
+				{
+					// Delete cooked data (if any) as it may be incomplete / corrupted.
+					Log("Cook failed. Deleting cooked data.");
+					CleanupCookedData(PlatformsToCook.ToList(), Params);
+					AutomationTool.ErrorReporter.Error("Cook failed.", (int)AutomationTool.ErrorCodes.Error_UnknownCookFailure);
+					throw Ex;
+				}
 			}
 		}
 
