Add a Robust_CopyFile function to deal with Macs losing connections to the network share.

diff --git a/Engine/Source/Programs/AutomationTool/Scripts/RocketBuild.Automation.cs b/Engine/Source/Programs/AutomationTool/Scripts/RocketBuild.Automation.cs
--- a/Engine/Source/Programs/AutomationTool/Scripts/RocketBuild.Automation.cs
+++ b/Engine/Source/Programs/AutomationTool/Scripts/RocketBuild.Automation.cs
@@ -853,8 +853,8 @@ namespace Rocket
 
 			// Copy the files to their final location
 			CommandUtils.Log("Copying files to {0}", PublishDir);
-			CommandUtils.CopyFile(FullZipFileName, Path.Combine(PublishDir, Path.GetFileName(FullZipFileName)));
-			CommandUtils.CopyFile(EditorZipFileName, Path.Combine(PublishDir, Path.GetFileName(EditorZipFileName)));
+			CommandUtils.Robust_CopyFile(FullZipFileName, Path.Combine(PublishDir, Path.GetFileName(FullZipFileName)));
+			CommandUtils.Robust_CopyFile(EditorZipFileName, Path.Combine(PublishDir, Path.GetFileName(EditorZipFileName)));
 			CommandUtils.DeleteFile(FullZipFileName);
 			CommandUtils.DeleteFile(EditorZipFileName);
 
@@ -911,26 +911,7 @@ namespace Rocket
 					if(DebugExtensions.Contains(Extension) || Extension == ".exe" || Extension == ".dll") // Need all windows build products for crash reporter
 					{
 						string OutputFileName = CommandUtils.MakeRerootedFilePath(InputFileName, CommandUtils.CmdEnv.LocalRoot, SymbolsOutputDir);
-                        if (SymbolsOutputDir.StartsWith("/Volumes"))
-                        {
-                            int Retry = 0;
-                            int NumRetries = 60;
-                            bool bCopied = false;
-                            while (!bCopied && Retry < NumRetries)
-                            {
-                                if (Retry > 0)
-                                {
-                                    CommandUtils.Log("*** Mac temp storage retry {0}", OutputFileName);
-                                    System.Threading.Thread.Sleep(1000);
-                                }
-                                bCopied = CommandUtils.CopyFile_NoExceptions(InputFileName, OutputFileName, true);
-                                Retry++;
-                            }
-                        }
-                        else
-                        {
-                            CommandUtils.CopyFile(InputFileName, OutputFileName);
-                        }
+						CommandUtils.Robust_CopyFile(InputFileName, OutputFileName);
 					}
 				}
 			}
diff --git a/Engine/Source/Programs/AutomationTool/TempStorage.cs b/Engine/Source/Programs/AutomationTool/TempStorage.cs
--- a/Engine/Source/Programs/AutomationTool/TempStorage.cs
+++ b/Engine/Source/Programs/AutomationTool/TempStorage.cs
@@ -93,6 +93,30 @@ namespace AutomationTool
             }
         }
 
+		public static void Robust_CopyFile(string InputFileName, string OutputFileName)
+		{
+            if (OutputFileName.StartsWith("/Volumes/"))
+            {
+                int Retry = 0;
+                int NumRetries = 60;
+                bool bCopied = false;
+                while (!bCopied && Retry < NumRetries)
+                {
+                    if (Retry > 0)
+                    {
+                        CommandUtils.Log("*** Mac temp storage retry {0}", OutputFileName);
+                        System.Threading.Thread.Sleep(1000);
+                    }
+                    bCopied = CommandUtils.CopyFile_NoExceptions(InputFileName, OutputFileName, true);
+                    Retry++;
+                }
+            }
+            else
+            {
+                CommandUtils.CopyFile(InputFileName, OutputFileName);
+            }
+		}
+
         public static void Robust_FileExists_NoExceptions(string Filename, string Message)
         {
             Robust_FileExists_NoExceptions(false, Filename, Message);
