Removed execution of translated word counts calculator for UDN files from the localization automation script for UE4.

diff --git a/Engine/Source/Programs/AutomationTool/Scripts/Localisation.Automation.cs b/Engine/Source/Programs/AutomationTool/Scripts/Localisation.Automation.cs
--- a/Engine/Source/Programs/AutomationTool/Scripts/Localisation.Automation.cs
+++ b/Engine/Source/Programs/AutomationTool/Scripts/Localisation.Automation.cs
@@ -117,29 +117,6 @@ class Localise : BuildCommand
         UploadDirectoryToProject(GetProject(oneSkyService, "Unreal Engine", "ToolTips"), new DirectoryInfo(CmdEnv.LocalRoot + "/Engine/Content/Localization/ToolTips"), "*.po");
         UploadDirectoryToProject(GetProject(oneSkyService, "Unreal Engine", "Category"), new DirectoryInfo(CmdEnv.LocalRoot + "/Engine/Content/Localization/Category"), "*.po");
         UploadDirectoryToProject(GetProject(oneSkyService, "Unreal Engine", "Keywords"), new DirectoryInfo(CmdEnv.LocalRoot + "/Engine/Content/Localization/Keywords"), "*.po");
-
-		// Localisation statistics estimator.
-		if (P4Enabled)
-		{
-			// Available only for P4
-			var EstimatorExePath = CombinePaths(CmdEnv.LocalRoot, @"Engine/Binaries/DotNET/TranslatedWordsCountEstimator.exe");
-			var StatisticsFilePath = @"\\epicgames.net\root\UE3\Localization\WordCounts\udn.csv";
-
-			var Arguments = string.Format(
-				"{0} {1} {2} {3} {4}",
-				StatisticsFilePath,
-				P4Env.P4Port,
-				P4Env.User,
-				P4Env.Client,
-				Environment.GetEnvironmentVariable("P4PASSWD"));
-
-			var RunResult = Run(EstimatorExePath, Arguments);
-
-			if (RunResult.ExitCode != 0)
-			{
-				throw new AutomationException("Error while executing TranslatedWordsCountEstimator with arguments '{0}'", Arguments);
-			}
-		}
 	}
 
     private static ProjectGroup GetProjectGroup(OneSkyService oneSkyService, string GroupName)
