Treat shadow variable warnings as errors in non-unity builds.

diff --git a/Engine/Source/Programs/AutomationTool/Scripts/GUBP.Automation.cs b/Engine/Source/Programs/AutomationTool/Scripts/GUBP.Automation.cs
--- a/Engine/Source/Programs/AutomationTool/Scripts/GUBP.Automation.cs
+++ b/Engine/Source/Programs/AutomationTool/Scripts/GUBP.Automation.cs
@@ -3515,14 +3515,14 @@ public class GUBP : BuildCommand
             Agenda.AddTargets(new string[] { "UnrealHeaderTool" }, HostPlatform, UnrealTargetConfiguration.Development);
             Agenda.AddTargets(
                 new string[] { bp.Branch.BaseEngineProject.Properties.Targets[TargetRules.TargetType.Editor].TargetName },
-                HostPlatform, UnrealTargetConfiguration.Development, InAddArgs: "-skipnonhostplatforms");
+                HostPlatform, UnrealTargetConfiguration.Development, InAddArgs: "-skipnonhostplatforms -shadowvariableerrors");
 
             foreach (var Kind in BranchInfo.MonolithicKinds)
             {
                 if (bp.Branch.BaseEngineProject.Properties.Targets.ContainsKey(Kind))
                 {
                     var Target = bp.Branch.BaseEngineProject.Properties.Targets[Kind];
-                    Agenda.AddTargets(new string[] { Target.TargetName }, HostPlatform, UnrealTargetConfiguration.Development);
+                    Agenda.AddTargets(new string[] { Target.TargetName }, HostPlatform, UnrealTargetConfiguration.Development, InAddArgs: "-shadowvariableerrors");
                 }
             }
 
diff --git a/Engine/Source/Programs/UnrealBuildTool/Configuration/BuildConfiguration.cs b/Engine/Source/Programs/UnrealBuildTool/Configuration/BuildConfiguration.cs
--- a/Engine/Source/Programs/UnrealBuildTool/Configuration/BuildConfiguration.cs
+++ b/Engine/Source/Programs/UnrealBuildTool/Configuration/BuildConfiguration.cs
@@ -37,6 +37,12 @@ namespace UnrealBuildTool
 		[XmlConfig]
 		public static int MinGameModuleSourceFilesForUnityBuild;
 
+		/// <summary>
+		/// Forces shadow variable warnings to be treated as errors on platforms that support it.
+		/// </summary>
+		[XmlConfig]
+		public static bool bShadowVariableErrors;
+
 		/// <summary>
 		/// New Monolithic Graphics drivers have optional "fast calls" replacing various D3d functions
 		/// </summary>
diff --git a/Engine/Source/Programs/UnrealBuildTool/HTML5/HTML5ToolChain.cs b/Engine/Source/Programs/UnrealBuildTool/HTML5/HTML5ToolChain.cs
--- a/Engine/Source/Programs/UnrealBuildTool/HTML5/HTML5ToolChain.cs
+++ b/Engine/Source/Programs/UnrealBuildTool/HTML5/HTML5ToolChain.cs
@@ -66,7 +66,7 @@ namespace UnrealBuildTool
 
 			if (bEnableShadowVariableWarning)
 			{
-				Result += " -Wshadow -Wno-error=shadow";
+				Result += " -Wshadow" + (BuildConfiguration.bShadowVariableErrors? "" : " -Wno-error=shadow");
 			}
 
             // JavsScript option overrides (see src/settings.js)
diff --git a/Engine/Source/Programs/UnrealBuildTool/IOS/IOSToolChain.cs b/Engine/Source/Programs/UnrealBuildTool/IOS/IOSToolChain.cs
--- a/Engine/Source/Programs/UnrealBuildTool/IOS/IOSToolChain.cs
+++ b/Engine/Source/Programs/UnrealBuildTool/IOS/IOSToolChain.cs
@@ -232,7 +232,7 @@ namespace UnrealBuildTool
 				
 			if (CompileEnvironment.Config.bEnableShadowVariableWarning)
 			{
-				Result += " -Wshadow -Wno-error=shadow";
+				Result += " -Wshadow" + (BuildConfiguration.bShadowVariableErrors? "" : " -Wno-error=shadow");
 			}
 
 			Result += " -Wno-unused-variable";
diff --git a/Engine/Source/Programs/UnrealBuildTool/Linux/LinuxToolChain.cs b/Engine/Source/Programs/UnrealBuildTool/Linux/LinuxToolChain.cs
--- a/Engine/Source/Programs/UnrealBuildTool/Linux/LinuxToolChain.cs
+++ b/Engine/Source/Programs/UnrealBuildTool/Linux/LinuxToolChain.cs
@@ -353,7 +353,7 @@ namespace UnrealBuildTool
 
 			if (CompileEnvironment.Config.bEnableShadowVariableWarning)
 			{
-				Result += " -Wshadow -Wno-error=shadow";
+				Result += " -Wshadow" + (BuildConfiguration.bShadowVariableErrors? "" : " -Wno-error=shadow");
 			}
 
             //Result += " -DOPERATOR_NEW_INLINE=FORCENOINLINE";
diff --git a/Engine/Source/Programs/UnrealBuildTool/Mac/MacToolChain.cs b/Engine/Source/Programs/UnrealBuildTool/Mac/MacToolChain.cs
--- a/Engine/Source/Programs/UnrealBuildTool/Mac/MacToolChain.cs
+++ b/Engine/Source/Programs/UnrealBuildTool/Mac/MacToolChain.cs
@@ -172,7 +172,7 @@ namespace UnrealBuildTool
 
 			if (CompileEnvironment.Config.bEnableShadowVariableWarning)
 			{
-				Result += " -Wshadow -Wno-error=shadow";
+				Result += " -Wshadow" + (BuildConfiguration.bShadowVariableErrors? "" : " -Wno-error=shadow");
 			}
 
 			// @todo: Remove these two when the code is fixed and they're no longer needed
diff --git a/Engine/Source/Programs/UnrealBuildTool/System/UnrealBuildTool.cs b/Engine/Source/Programs/UnrealBuildTool/System/UnrealBuildTool.cs
--- a/Engine/Source/Programs/UnrealBuildTool/System/UnrealBuildTool.cs
+++ b/Engine/Source/Programs/UnrealBuildTool/System/UnrealBuildTool.cs
@@ -749,6 +749,10 @@ namespace UnrealBuildTool
                 {
                     BuildConfiguration.bForceUnityBuild = true;
                 }
+                else if (LowercaseArg == "-shadowvariableerrors")
+                {
+                    BuildConfiguration.bShadowVariableErrors = true;
+                }
                 // New Monolithic Graphics drivers have optional "fast calls" replacing various D3d functions
                 else if (LowercaseArg == "-fastmonocalls")
                 {
diff --git a/Engine/Source/Programs/UnrealBuildTool/Windows/VCToolChain.cs b/Engine/Source/Programs/UnrealBuildTool/Windows/VCToolChain.cs
--- a/Engine/Source/Programs/UnrealBuildTool/Windows/VCToolChain.cs
+++ b/Engine/Source/Programs/UnrealBuildTool/Windows/VCToolChain.cs
@@ -677,7 +677,7 @@ namespace UnrealBuildTool
 
 				if (CompileEnvironment.Config.bEnableShadowVariableWarning)
 				{
-					Arguments.Append(" -Wshadow -Wno-error=shadow");
+					Arguments.Append(" -Wshadow" + (BuildConfiguration.bShadowVariableErrors? "" : " -Wno-error=shadow"));
 				}
 
 				// @todo clang: Kind of a shame to turn these off.  We'd like to catch unused variables, but it is tricky with how our assertion macros work.
