build script, fix for -game on code projects that no longer contain promotion

diff --git a/Engine/Source/Programs/AutomationTool/Scripts/GUBP.Automation.cs b/Engine/Source/Programs/AutomationTool/Scripts/GUBP.Automation.cs
--- a/Engine/Source/Programs/AutomationTool/Scripts/GUBP.Automation.cs
+++ b/Engine/Source/Programs/AutomationTool/Scripts/GUBP.Automation.cs
@@ -6577,8 +6577,10 @@ public class GUBP : BuildCommand
                     {
                         if (GameProj.GameName.Equals(GameName, StringComparison.InvariantCultureIgnoreCase))
                         {
-
-                            NodesToDo.Add(GameAggregatePromotableNode.StaticGetFullName(GameProj));
+                            if (GameProj.Options(UnrealTargetPlatform.Win64).bIsPromotable)
+                            {
+                                NodesToDo.Add(GameAggregatePromotableNode.StaticGetFullName(GameProj));
+                            }
                             foreach (var Node in GUBPNodes)
                             {
                                 if (Node.Value.GameNameIfAnyForTempStorage() == GameProj.GameName)
