#html5 fix emitting compression for code projects. use consistent file extentions for compressed files.

diff --git a/Engine/Build/HTML5/GameX.html.template b/Engine/Build/HTML5/GameX.html.template
--- a/Engine/Build/HTML5/GameX.html.template
+++ b/Engine/Build/HTML5/GameX.html.template
@@ -291,7 +291,7 @@ window.onerror  = function(msg, url, line, column, error) {
         addLog( " downloading symbols file ");
         // get non minified symbols and demangle the callstack. 
 		$.ajax({
-			  url : "%CONFIG%.js.symbols",
+			  url : "%CONFIG%.js.symbols.gz",
 			  dataType: "text",
 			  success : function (data) {	
                   addLog(" symbols file downloaded, demangling..");
@@ -376,6 +376,12 @@ var Module = {
 		infoPrinted: false, 
 		lastcurrentDownloadedSize: 0,
 
+		// file location. 
+		locateFile : function (name){ 
+		      return name + ".gz";
+		},
+
+
 		// packaged data download, sets this information. 
 		setStatus: function(text) {
 		 	var matches = text.match(/([^(]+)\((\d+(\.\d+)?)\/(\d+)\)/);
diff --git a/Engine/Source/Programs/AutomationTool/HTML5/HTML5Platform.Automation.cs b/Engine/Source/Programs/AutomationTool/HTML5/HTML5Platform.Automation.cs
--- a/Engine/Source/Programs/AutomationTool/HTML5/HTML5Platform.Automation.cs
+++ b/Engine/Source/Programs/AutomationTool/HTML5/HTML5Platform.Automation.cs
@@ -74,12 +74,6 @@ public class HTML5Platform : Platform
             }
         }
 
-		// packaging created 2 files - .data.js and .data file. lets compress them.
-
-		CompressFile(FinalDataLocation, FinalDataLocation);
-		CompressFile(FinalDataLocation + ".js", FinalDataLocation + ".js.gz");
-		File.Delete(FinalDataLocation + ".js");
-
         // copy the "Executable" to the package directory
         string GameExe = Path.GetFileNameWithoutExtension(Params.ProjectGameExeFilename);
         if (Params.ClientConfigsToBuild[0].ToString() != "Development")
@@ -99,16 +93,16 @@ public class HTML5Platform : Platform
 
         if (Path.Combine(Path.GetDirectoryName(Params.ProjectGameExeFilename), GameExe) != Path.Combine(PackagePath, GameExe))
         {
-			// compress all javascript files in place. 
-			Log("Compressing and Copying main js executable");
-			CompressFile(FullGameExePath, Path.Combine(PackagePath, GameExe) + ".gz");
-			Log("Compressing and copying memory Initialization file");
-			CompressFile(Path.Combine(Path.GetDirectoryName(Params.ProjectGameExeFilename), GameExe) + ".mem", (Path.Combine(PackagePath, GameExe) + ".mem"));
+            File.Copy(Path.Combine(Path.GetDirectoryName(Params.ProjectGameExeFilename), GameExe), Path.Combine(PackagePath, GameExe), true);
+            File.Copy(Path.Combine(Path.GetDirectoryName(Params.ProjectGameExeFilename), GameExe) + ".mem", Path.Combine(PackagePath, GameExe) + ".mem", true);
 			File.Copy(Path.Combine(Path.GetDirectoryName(Params.ProjectGameExeFilename), GameExe) + ".symbols", Path.Combine(PackagePath, GameExe) + ".symbols", true);
         }
 
+        File.SetAttributes(Path.Combine(PackagePath, GameExe), FileAttributes.Normal);
+        File.SetAttributes(Path.Combine(PackagePath, GameExe) + ".mem", FileAttributes.Normal);
 		File.SetAttributes(Path.Combine(PackagePath, GameExe) + ".symbols", FileAttributes.Normal);
 
+
         // put the HTML file to the package directory
 
         bool UseExperimentalTemplate = false;
@@ -155,19 +149,27 @@ public class HTML5Platform : Platform
 			File.AppendAllText(DestinationFile, Data);
 		}
 
-		CompressFile(DestinationFile, DestinationFile + ".gz");
-
-		// delete uncompressed file.
-		if (File.Exists(DestinationFile))
-		{
-			File.Delete(DestinationFile);
-		}
+		// Compress all files. 
+
+		//data file.
+		CompressFile(FinalDataLocation,FinalDataLocation + ".gz");
+		// data file .js driver.
+		CompressFile(FinalDataLocation + ".js" , FinalDataLocation + ".js.gz");
+		// main js.
+		CompressFile(Path.Combine(PackagePath, GameExe), Path.Combine(PackagePath, GameExe) + ".gz");
+		// mem init file.
+		CompressFile(Path.Combine(PackagePath, GameExe) + ".mem", Path.Combine(PackagePath, GameExe) + ".mem.gz");
+		// symbols file.
+		CompressFile(Path.Combine(PackagePath, GameExe) + ".symbols", Path.Combine(PackagePath, GameExe) + ".symbols.gz");
+		// Utility 
+		CompressFile(OutDir + "/Utility.js", OutDir + "/Utility.js.gz");
 
         PrintRunTime();
 	}
 
 	void CompressFile(string Source, string Destination)
 	{
+		Log(" Compressing " + Source); 
 		bool DeleteSource = false; 
 
 		if(  Source == Destination )
@@ -299,12 +301,19 @@ public class HTML5Platform : Platform
 		// put the HTML file to the package directory
 		string OutputFile = Path.Combine(PackagePath, (Params.ClientConfigsToBuild[0].ToString() != "Development" ? (Params.ShortProjectName + "-HTML5-" + Params.ClientConfigsToBuild[0].ToString()) : Params.ShortProjectName)) + ".html";
 
-		SC.ArchiveFiles(PackagePath, Path.GetFileName(FinalDataLocation));
+		// data file 
+		SC.ArchiveFiles(PackagePath, Path.GetFileName(FinalDataLocation + ".gz"));
+		// data file js driver 
 		SC.ArchiveFiles(PackagePath, Path.GetFileName(FinalDataLocation + ".js.gz"));
-		SC.ArchiveFiles(PackagePath, Path.GetFileName(GameExe +".gz"));
-		SC.ArchiveFiles(PackagePath, Path.GetFileName(GameExe + ".mem"));
-		SC.ArchiveFiles(PackagePath, Path.GetFileName(GameExe + ".symbols"));
+		// main js file
+		SC.ArchiveFiles(PackagePath, Path.GetFileName(GameExe + ".gz"));
+		// memory init file
+		SC.ArchiveFiles(PackagePath, Path.GetFileName(GameExe + ".mem.gz"));
+		// symbols file
+		SC.ArchiveFiles(PackagePath, Path.GetFileName(GameExe + ".symbols.gz"));
+		// utilities
 		SC.ArchiveFiles(PackagePath, Path.GetFileName("Utility.js.gz"));
+		// landing page.
 		SC.ArchiveFiles(PackagePath, Path.GetFileName(OutputFile));
 
 		// Archive HTML5 Server and a Readme. 
@@ -314,13 +323,13 @@ public class HTML5Platform : Platform
 
 		if (HTMLPakAutomation.CanCreateMapPaks(Params))
 		{
-			// find all paks.
+		// find all paks.
 			string[] Files = Directory.GetFiles(Path.Combine(PackagePath, Params.ShortProjectName), "*",SearchOption.AllDirectories);
 			foreach(string PakFile in Files)
 			{
 				var DestPak = PakFile.Replace(PackagePath,"");
 				SC.ArchivedFiles.Add(PakFile, DestPak);
-			}
+			}	
 		}
 
 	}
diff --git a/Engine/Source/Programs/HTML5/HTML5LaunchHelper/Program.cs b/Engine/Source/Programs/HTML5/HTML5LaunchHelper/Program.cs
--- a/Engine/Source/Programs/HTML5/HTML5LaunchHelper/Program.cs
+++ b/Engine/Source/Programs/HTML5/HTML5LaunchHelper/Program.cs
@@ -161,9 +161,9 @@ namespace HTML5LaunchHelper
 						MimeType = MimeTypeMapping[Extention];
 					}
 					// This is the crux of serving pre-compressed files. 
-					else if (Extention == ".gz" || Extention == ".data" || Extention == ".mem")
+					if (Extention == ".gz")
 					{
-						Context.Response.AddHeader("Content-Encoding", "gzip");
+					  Context.Response.AddHeader("Content-Encoding", "gzip");
 					}
 
 					byte[] buffer = new byte[source.Length];
