removed the need for BlankProject from UAT unless you use the foreign option
UEB-73
#uat

diff --git a/Engine/Source/Programs/AutomationTool/ProjectUtils.cs b/Engine/Source/Programs/AutomationTool/ProjectUtils.cs
--- a/Engine/Source/Programs/AutomationTool/ProjectUtils.cs
+++ b/Engine/Source/Programs/AutomationTool/ProjectUtils.cs
@@ -739,10 +739,10 @@ namespace AutomationTool
             }
             var FinishTime = DateTime.Now.ToString();
             CommandUtils.PrintCSVFile(String.Format("UAT,SortProjects,{0},{1}", StartTime, FinishTime));            
-            if (String.IsNullOrEmpty(BaseEngineProject.FilePath))
+ /*           if (String.IsNullOrEmpty(BaseEngineProject.FilePath))
             {
                 throw new AutomationException("All branches must have the blank project /Samples/Sandbox/BlankProject");
-            }
+            }*/
 
             CommandUtils.Log("  Base Engine:");
             var StartBranchDump = DateTime.Now.ToString();
diff --git a/Engine/Source/Programs/AutomationTool/Scripts/BuildCookRun.Automation.cs b/Engine/Source/Programs/AutomationTool/Scripts/BuildCookRun.Automation.cs
--- a/Engine/Source/Programs/AutomationTool/Scripts/BuildCookRun.Automation.cs
+++ b/Engine/Source/Programs/AutomationTool/Scripts/BuildCookRun.Automation.cs
@@ -297,7 +297,7 @@ public class BuildCookRun : BuildCommand
 				}
 				else
 				{
-					var OriginalProjectName = ParseParamValue("project", CombinePaths("Samples", "Sandbox", "BlankProject", "BlankProject.uproject"));
+					var OriginalProjectName = ParseParamValue("project", "");
 					var ProjectName = OriginalProjectName;
 					ProjectName = ProjectName.Trim(new char[] { '\"' });
 					if (ProjectName.IndexOfAny(new char[] { '\\', '/' }) < 0)
