Manually rolled back to version checked into CL 2569306

diff --git a/Engine/Source/Programs/AutomationTool/UE4Build.cs b/Engine/Source/Programs/AutomationTool/UE4Build.cs
--- a/Engine/Source/Programs/AutomationTool/UE4Build.cs
+++ b/Engine/Source/Programs/AutomationTool/UE4Build.cs
@@ -767,7 +767,7 @@ namespace AutomationTool
 			return XGEConsoleExePath;
 		}
 
-		public bool RunXGE(List<XGEItem> Actions, string TaskFilePath, bool DoRetries, bool SpecialTestFlag)
+		public string ParallelExecutorExec()
 		{
 			if (string.IsNullOrEmpty(ParallelExecutorPath))
 			{
@@ -1189,115 +1189,7 @@ namespace AutomationTool
 					Job++;
 				}
 			}
-			if (Job == 0)
-			{
-				if (DeleteBuildProducts)
-				{
-					throw new AutomationException("Unable to find xge xmls: " + CombinePaths(CmdEnv.LogFolder, "*.xge.xml"));
-				}
-				else
-				{
-					Log("Incremental build, apparently everything was up to date, no XGE jobs produced.");
-				}
-			}
-			else
-			{
-				// Make sure all the tasks have a unique prefix
-				if(ShowProgress)
-				{
-					List<XmlElement> AllToolElements = new List<XmlElement>();
-					foreach(XmlElement EnvironmentElement in EnvironmentsElement.GetElementsByTagName("Environment"))
-					{
-						foreach(XmlElement ToolsElement in EnvironmentElement.GetElementsByTagName("Tools"))
-						{
-							foreach(XmlElement ToolElement in ToolsElement.GetElementsByTagName("Tool"))
-							{
-								AllToolElements.Add(ToolElement);
-							}
-						}
-					}
-					for(int Idx = 0; Idx < AllToolElements.Count; Idx++)
-					{
-						XmlElement ToolElement = AllToolElements[Idx];
-						if (ToolElement.HasAttribute("OutputPrefix"))
-						{
-							ToolElement.SetAttribute("OutputPrefix", ToolElement.Attributes["OutputPrefix"].Value + String.Format(" [@progress increment 1/{0}]", AllToolElements.Count));
-						}
-						else
-						{
-							ToolElement.SetAttribute("OutputPrefix", String.Format(" [@progress increment 1/{0} skipline]", AllToolElements.Count));
-						}
-					}
-				}
-
-				// Write the XGE task XML to a temporary file.
-				using (FileStream OutputFileStream = new FileStream(TaskFilePath, FileMode.Create, FileAccess.Write))
-				{
-					XGETaskDocument.Save(OutputFileStream);
-				}
-				if (!FileExists(TaskFilePath))
-				{
-					throw new AutomationException("Unable to find xge xml: " + TaskFilePath);
-				}
-
-				string Args = "\"" + TaskFilePath + "\" /Rebuild /MaxCPUS=200";
-
-				int Retries = DoRetries ? 2 : 1;
-                int ConnectionRetries = 4;
-                for (int i = 0; i < Retries; i++)
-				{
-					try
-					{
-                        while (true)
-                        {
-                            Log("Running XGE *******");
-                            PushDir(CombinePaths(CmdEnv.LocalRoot, @"\Engine\Source"));
-                            int SuccesCode;
-                            string LogFile = GetRunAndLogLogName(CmdEnv, XGEConsole);
-                            string Output = RunAndLog(XGEConsole, Args, out SuccesCode, LogFile);
-                            PopDir();
-                            if (ConnectionRetries > 0 && (SuccesCode == 4 || SuccesCode == 2) && !Output.Contains("------Project:"))
-                            {
-                                Log(System.Diagnostics.TraceEventType.Warning, "XGE failure on the local connection timeout");
-                                if (ConnectionRetries < 2)
-                                {
-                                    System.Threading.Thread.Sleep(60000);
-                                }
-                                ConnectionRetries--;
-                                continue;
-                            }
-                            else if (SuccesCode != 0)
-                            {
-                                Log("XGE did not succeed *******");
-                                throw new AutomationException(String.Format("Command failed (Result:{3}): {0} {1}. See logfile for details: '{2}' ",
-                                                                XGEConsole, Args, Path.GetFileName(LogFile), SuccesCode));
-                            }
-                            Log("XGE Done *******");
-                            break;
-                        }
-                        break;
-					}
-					catch (Exception Ex)
-					{
-						Log("XGE failed on try {0}: {1}", i + 1, Ex.ToString());
-						if (i + 1 >= Retries)
-						{
-							return false;
-						}
-						Log("Deleting build products to force a rebuild.");
-						foreach (var Item in Actions)
-						{
-							XGEDeleteBuildProducts(Item.Manifest);
-						}
-					}
-				}
-
-			}
-			foreach (var Item in Actions)
-			{
-				XGEFinishBuildWithUBT(Item);
-			}
-			return true;
+			return (Job > 0);
 		}
 
 		public void ClearExportedXGEXML()
