Fix for UAT base path being invalid if the UE4 workspace is mounted directly onto a drive letter (i.e. having N:\\Engine etc.)
 - Internal base path becomes N: in this case, which isn't a valid C# path, so attempts to set it as CWD fail. N:\\ is the correct path in this case.

diff --git a/Engine/Source/Programs/AutomationTool/BuildEnvironment.cs b/Engine/Source/Programs/AutomationTool/BuildEnvironment.cs
--- a/Engine/Source/Programs/AutomationTool/BuildEnvironment.cs
+++ b/Engine/Source/Programs/AutomationTool/BuildEnvironment.cs
@@ -105,6 +105,12 @@ namespace AutomationTool
 			{
 				SetUATSavedPath();
 			}
+
+			if (LocalRoot.EndsWith(":"))
+			{
+				LocalRoot += Path.DirectorySeparatorChar;
+			}
+
 			EngineSavedFolder = CommandUtils.GetEnvVar(EnvVarNames.EngineSavedFolder);
             CSVFile = CommandUtils.GetEnvVar(EnvVarNames.CSVFile);            
 			LogFolder = CommandUtils.GetEnvVar(EnvVarNames.LogFolder);
