Mac staging fixes

diff --git a/Engine/Source/Programs/AutomationTool/Mac/MacPlatform.Automation.cs b/Engine/Source/Programs/AutomationTool/Mac/MacPlatform.Automation.cs
--- a/Engine/Source/Programs/AutomationTool/Mac/MacPlatform.Automation.cs
+++ b/Engine/Source/Programs/AutomationTool/Mac/MacPlatform.Automation.cs
@@ -136,6 +136,13 @@ public class MacPlatform : Platform
 				WorkingFileType = StagedFileType.DebugNonUFS;
 			}
 		}
+
+		if (SC.bStageCrashReporter)
+		{
+			string CrashReportClientPath = CombinePaths("Engine/Binaries", SC.PlatformDir, "CrashReportClient.app");
+			StageAppBundle(SC, StagedFileType.NonUFS, CombinePaths(SC.LocalRoot, "Engine/Binaries", SC.PlatformDir, "CrashReportClient.app"), CrashReportClientPath);
+		}
+
 		// Copy the splash screen, Mac specific
 		SC.StageFiles(StagedFileType.NonUFS, CombinePaths(SC.ProjectRoot, "Content/Splash"), "Splash.bmp", false, null, null, true);
 
@@ -173,15 +180,17 @@ public class MacPlatform : Platform
 		string TargetPath = CombinePaths(BundlePath, "Contents", "UE4");
 		if (!SC.bIsCombiningMultiplePlatforms)
 		{
-			if (!Directory.Exists(BundlePath))
+			if (Directory.Exists(BundlePath))
 			{
-				string SourceBundlePath = CombinePaths(SC.ArchiveDirectory, SC.ShortProjectName, "Binaries", "Mac", ExeName + ".app");
-				if (!Directory.Exists(SourceBundlePath))
-				{
-					SourceBundlePath = CombinePaths(SC.ArchiveDirectory, "Engine", "Binaries", "Mac", ExeName + ".app");
-				}
-				Directory.Move(SourceBundlePath, BundlePath);
+				Directory.Delete(BundlePath, true);
+			}
+
+			string SourceBundlePath = CombinePaths(SC.ArchiveDirectory, SC.ShortProjectName, "Binaries", "Mac", ExeName + ".app");
+			if (!Directory.Exists(SourceBundlePath))
+			{
+				SourceBundlePath = CombinePaths(SC.ArchiveDirectory, "Engine", "Binaries", "Mac", ExeName + ".app");
 			}
+			Directory.Move(SourceBundlePath, BundlePath);
 
 			if (DirectoryExists(TargetPath))
 			{
