Changes to reduce log spew in stdout for UAT and GUBP
UEB-156
#uat

diff --git a/Engine/Source/Programs/AutomationTool/BuildEnvironment.cs b/Engine/Source/Programs/AutomationTool/BuildEnvironment.cs
--- a/Engine/Source/Programs/AutomationTool/BuildEnvironment.cs
+++ b/Engine/Source/Programs/AutomationTool/BuildEnvironment.cs
@@ -159,18 +159,18 @@ namespace AutomationTool
 
 		void LogSettings()
 		{
-			Log.TraceInformation("Command Environment settings:");
-			Log.TraceInformation("CmdExe={0}", CmdExe);
-			Log.TraceInformation("EngineSavedFolder={0}", EngineSavedFolder);
-			Log.TraceInformation("HasCapabilityToCompile={0}", HasCapabilityToCompile);
-			Log.TraceInformation("LocalRoot={0}", LocalRoot);
-			Log.TraceInformation("LogFolder={0}", LogFolder);			
-			Log.TraceInformation("MountExe={0}", MountExe);
-			Log.TraceInformation("MsBuildExe={0}", MsBuildExe);
-			Log.TraceInformation("MsDevExe={0}", MsDevExe);
-			Log.TraceInformation("RobocopyExe={0}", RobocopyExe);
-			Log.TraceInformation("TimestampAsString={0}", TimestampAsString);
-			Log.TraceInformation("UATExe={0}", UATExe);			
+			Log.TraceVerbose("Command Environment settings:");
+			Log.TraceVerbose("CmdExe={0}", CmdExe);
+			Log.TraceVerbose("EngineSavedFolder={0}", EngineSavedFolder);
+			Log.TraceVerbose("HasCapabilityToCompile={0}", HasCapabilityToCompile);
+			Log.TraceVerbose("LocalRoot={0}", LocalRoot);
+			Log.TraceVerbose("LogFolder={0}", LogFolder);
+			Log.TraceVerbose("MountExe={0}", MountExe);
+			Log.TraceVerbose("MsBuildExe={0}", MsBuildExe);
+			Log.TraceVerbose("MsDevExe={0}", MsDevExe);
+			Log.TraceVerbose("RobocopyExe={0}", RobocopyExe);
+			Log.TraceVerbose("TimestampAsString={0}", TimestampAsString);
+			Log.TraceVerbose("UATExe={0}", UATExe);			
 		}
 
 		#region Compiler Setup
@@ -190,7 +190,7 @@ namespace AutomationTool
 			catch (Exception)
 			{
 				// Something went wrong, we can't compile.
-				Log.WriteLine(TraceEventType.Information, "SetFrameworkVars failed. Assuming no compilation capability.");
+				Log.WriteLine(TraceEventType.Warning, "SetFrameworkVars failed. Assuming no compilation capability.");
 				HasCapabilityToCompile = false;
 			}
 
@@ -202,8 +202,8 @@ namespace AutomationTool
 				}
 				catch (Exception Ex)
 				{
-					Log.WriteLine(TraceEventType.Information, Ex.Message);
-					Log.WriteLine(TraceEventType.Information, "Assuming no compilation capability.");
+					Log.WriteLine(TraceEventType.Warning, Ex.Message);
+					Log.WriteLine(TraceEventType.Warning, "Assuming no compilation capability.");
 					HasCapabilityToCompile = false;
 					MsBuildExe = "";
 				}
@@ -217,15 +217,15 @@ namespace AutomationTool
 				}
 				catch (Exception Ex)
 				{
-					Log.WriteLine(TraceEventType.Information, Ex.Message);
-					Log.WriteLine(TraceEventType.Information, "Assuming no solution compilation capability.");
+					Log.WriteLine(TraceEventType.Warning, Ex.Message);
+					Log.WriteLine(TraceEventType.Warning, "Assuming no solution compilation capability.");
 					MsDevExe = "";
 				}
 			}
 
-			Log.TraceInformation("CompilationEvironment.HasCapabilityToCompile={0}", HasCapabilityToCompile);
-			Log.TraceInformation("CompilationEvironment.MsBuildExe={0}", MsBuildExe);
-			Log.TraceInformation("CompilationEvironment.MsDevExe={0}", MsDevExe);
+			Log.TraceVerbose("CompilationEvironment.HasCapabilityToCompile={0}", HasCapabilityToCompile);
+			Log.TraceVerbose("CompilationEvironment.MsBuildExe={0}", MsBuildExe);
+			Log.TraceVerbose("CompilationEvironment.MsDevExe={0}", MsDevExe);
 		}
 
 		#endregion
diff --git a/Engine/Source/Programs/AutomationTool/Log.cs b/Engine/Source/Programs/AutomationTool/Log.cs
--- a/Engine/Source/Programs/AutomationTool/Log.cs
+++ b/Engine/Source/Programs/AutomationTool/Log.cs
@@ -36,7 +36,7 @@ namespace AutomationTool
                 bColorConsoleOutput: true,
                 TraceListeners: new TraceListener[]
                 {
-                    new ConsoleTraceListener(),
+                    new FilteredConsoleTraceListener(),
                     // could return null, but InitLogging handles this gracefully.
                     CreateLogFileListener(out LogFilename),
                     //@todo - this is only used by GUBP nodes. Ideally we don't waste this 20MB if we are not running GUBP.
@@ -263,4 +263,34 @@ namespace AutomationTool
 	}
 
 	#endregion
+
+	#region FilteredConsoleTraceListener
+
+	/// <summary>
+	///  Filtered Console Trace Listener
+	///  </summary>
+	class FilteredConsoleTraceListener : TextWriterTraceListener
+	{
+		public FilteredConsoleTraceListener()
+			: base(Console.OpenStandardOutput())
+		{ }
+
+		#region TraceListener Interface
+
+		public override void WriteLine(string message)
+		{
+			// strip off the message code
+			int Code = Convert.ToInt32(message.Substring(0,3), 16);
+			// filter the message based on the code
+			if (Code < (int)TraceEventType.Verbose || UnrealBuildTool.Log.bIsVerbose)
+			{
+				// filter based on the message code
+				base.WriteLine(message.Substring(3));
+			}
+		}
+
+		#endregion
+	}
+	#endregion
+
 }
diff --git a/Engine/Source/Programs/AutomationTool/P4Utils.cs b/Engine/Source/Programs/AutomationTool/P4Utils.cs
--- a/Engine/Source/Programs/AutomationTool/P4Utils.cs
+++ b/Engine/Source/Programs/AutomationTool/P4Utils.cs
@@ -468,7 +468,7 @@ namespace AutomationTool
 			// Init AllowSubmit
 			// If we do not specify on the commandline if submitting is allowed or not, this is 
 			// depending on whether we run locally or on a build machine.
-			Log("Initializing AllowSubmit.");
+			LogVerbose("Initializing AllowSubmit.");
 			if (GlobalCommandLine.Submit || GlobalCommandLine.NoSubmit)
 			{
 				AllowSubmit = GlobalCommandLine.Submit;
@@ -477,10 +477,10 @@ namespace AutomationTool
 			{
 				AllowSubmit = Automation.IsBuildMachine;
 			}
-			Log("AllowSubmit={0}", AllowSubmit);
+			LogVerbose("AllowSubmit={0}", AllowSubmit);
 
 			// Init P4Enabled
-			Log("Initializing P4Enabled.");
+			LogVerbose("Initializing P4Enabled.");
 			if (Automation.IsBuildMachine)
 			{
 				P4Enabled = !GlobalCommandLine.NoP4;
@@ -495,8 +495,8 @@ namespace AutomationTool
 				P4Enabled = GlobalCommandLine.P4 || bRequireP4;
 				P4CLRequired = GlobalCommandLine.P4 || bRequireCL;
 			}
-			Log("P4Enabled={0}", P4Enabled);
-			Log("P4CLRequired={0}", P4CLRequired);
+			LogVerbose("P4Enabled={0}", P4Enabled);
+			LogVerbose("P4CLRequired={0}", P4CLRequired);
 		}
 
 		/// <summary>
@@ -517,7 +517,7 @@ namespace AutomationTool
 					var RequireP4Attributes = Command.GetCustomAttributes(typeof(RequireP4Attribute), true);	
 					if (!CommandUtils.IsNullOrEmpty(RequireP4Attributes))
 					{
-						Log("Command {0} requires P4 functionality.", Command.Name);
+						LogWarning("Command {0} requires P4 functionality.", Command.Name);
 						bRequireP4 = true;
 
 						var DoesNotNeedP4CLAttributes = Command.GetCustomAttributes(typeof(DoesNotNeedP4CLAttribute), true);
diff --git a/Engine/Source/Programs/AutomationTool/Platform.cs b/Engine/Source/Programs/AutomationTool/Platform.cs
--- a/Engine/Source/Programs/AutomationTool/Platform.cs
+++ b/Engine/Source/Programs/AutomationTool/Platform.cs
@@ -34,7 +34,7 @@ namespace AutomationTool
 				Platform ExistingInstance;
 				if (AllPlatforms.TryGetValue(TargetType, out ExistingInstance) == false)
 				{
-					Log("Creating placeholder platform for target: {0}", TargetType);
+					LogVerbose("Creating placeholder platform for target: {0}", TargetType);
 					AllPlatforms.Add(TargetType, new Platform(TargetType));
 				}
 			}
@@ -42,7 +42,7 @@ namespace AutomationTool
 
 		private static void CreatePlatformsFromAssembly(Assembly ScriptAssembly)
 		{
-			Log("Looking for platforms in {0}", ScriptAssembly.Location);
+			LogVerbose("Looking for platforms in {0}", ScriptAssembly.Location);
 			Type[] AllTypes = null;
 			try
 			{
@@ -64,7 +64,7 @@ namespace AutomationTool
 					}
 					else
 					{
-						Log("No Loader Exceptions available.");
+						LogError("No Loader Exceptions available.");
 					}
 				}
 				// Re-throw, this is still a critical error!
@@ -74,7 +74,7 @@ namespace AutomationTool
 			{
 				if (PotentialPlatformType != typeof(Platform) && typeof(Platform).IsAssignableFrom(PotentialPlatformType) && !PotentialPlatformType.IsAbstract)
 				{
-					Log("Creating platform {0} from {1}.", PotentialPlatformType.Name, ScriptAssembly.Location);
+					LogVerbose("Creating platform {0} from {1}.", PotentialPlatformType.Name, ScriptAssembly.Location);
 					var PlatformInstance = Activator.CreateInstance(PotentialPlatformType) as Platform;
 					Platform ExistingInstance;
 					if (!AllPlatforms.TryGetValue(PlatformInstance.PlatformType, out ExistingInstance))
@@ -134,7 +134,7 @@ namespace AutomationTool
         public virtual void GetConnectedDevices(ProjectParams Params, out List<string> Devices)
         {
             Devices = null;
-            Log(System.Diagnostics.TraceEventType.Warning, "{0} does not implement GetConnectedDevices", PlatformType);
+            LogWarning("{0} does not implement GetConnectedDevices", PlatformType);
         }
 
 
@@ -146,7 +146,7 @@ namespace AutomationTool
 		/// <param name="SC"></param>
 		public virtual void Deploy(ProjectParams Params, DeploymentContext SC)
 		{
-			Log(System.Diagnostics.TraceEventType.Warning, "{0} does not implement Deploy...", PlatformType);
+			LogWarning("{0} does not implement Deploy...", PlatformType);
 		}
 
 		/// <summary>
diff --git a/Engine/Source/Programs/AutomationTool/ProjectUtils.cs b/Engine/Source/Programs/AutomationTool/ProjectUtils.cs
--- a/Engine/Source/Programs/AutomationTool/ProjectUtils.cs
+++ b/Engine/Source/Programs/AutomationTool/ProjectUtils.cs
@@ -317,7 +317,7 @@ namespace AutomationTool
 			// Get all ini files
 			if (!String.IsNullOrWhiteSpace(RawProjectPath))
 			{
-				CommandUtils.Log("Loading ini files for {0}", RawProjectPath);
+				CommandUtils.LogVerbose("Loading ini files for {0}", RawProjectPath);
 
 				var EngineDirectory = CommandUtils.CombinePaths(CommandUtils.CmdEnv.LocalRoot, "Engine");
 				foreach (UnrealTargetPlatform TargetPlatformType in Enum.GetValues(typeof(UnrealTargetPlatform)))
@@ -405,13 +405,13 @@ namespace AutomationTool
 			var RulesFolder = GetRulesAssemblyFolder();
 			if (!String.IsNullOrEmpty(Properties.RawProjectPath))
 			{
-				CommandUtils.Log("Looking for targets for project {0}", Properties.RawProjectPath);
+				CommandUtils.LogVerbose("Looking for targets for project {0}", Properties.RawProjectPath);
 
 				TargetsDllFilename = CommandUtils.CombinePaths(RulesFolder, String.Format("UATRules{0}.dll", Properties.RawProjectPath.GetHashCode()));
 
 				FullProjectPath = CommandUtils.GetDirectoryName(Properties.RawProjectPath);
 				GameFolders.Add(FullProjectPath);
-				CommandUtils.Log("Searching for target rule files in {0}", FullProjectPath);
+				CommandUtils.LogVerbose("Searching for target rule files in {0}", FullProjectPath);
 			}
 			else
 			{
@@ -648,29 +648,29 @@ namespace AutomationTool
             }
             public void Dump(List<UnrealTargetPlatform> InHostPlatforms)
             {
-                CommandUtils.Log("    ShortName:    " + GameName);
-                CommandUtils.Log("      FilePath          : " + FilePath);
-                CommandUtils.Log("      bIsCodeBasedProject  : " + (Properties.bIsCodeBasedProject ? "YES" : "NO"));
-                CommandUtils.Log("      bUsesSteam  : " + (Properties.bUsesSteam ? "YES" : "NO"));
-                CommandUtils.Log("      bUsesCEF3   : " + (Properties.bUsesCEF3 ? "YES" : "NO"));
-                CommandUtils.Log("      bUsesSlate  : " + (Properties.bUsesSlate ? "YES" : "NO"));
+                CommandUtils.LogVerbose("    ShortName:    " + GameName);
+				CommandUtils.LogVerbose("      FilePath          : " + FilePath);
+				CommandUtils.LogVerbose("      bIsCodeBasedProject  : " + (Properties.bIsCodeBasedProject ? "YES" : "NO"));
+				CommandUtils.LogVerbose("      bUsesSteam  : " + (Properties.bUsesSteam ? "YES" : "NO"));
+				CommandUtils.LogVerbose("      bUsesCEF3   : " + (Properties.bUsesCEF3 ? "YES" : "NO"));
+				CommandUtils.LogVerbose("      bUsesSlate  : " + (Properties.bUsesSlate ? "YES" : "NO"));
                 foreach (var HostPlatform in InHostPlatforms)
                 {
-                    CommandUtils.Log("      For Host : " + HostPlatform.ToString());
+					CommandUtils.LogVerbose("      For Host : " + HostPlatform.ToString());
                     if (Properties.bIsCodeBasedProject)
                     {
-                        CommandUtils.Log("          Promoted  : " + (Options(HostPlatform).bIsPromotable ? "YES" : "NO"));
-                        CommandUtils.Log("          SeparatePromotable  : " + (Options(HostPlatform).bSeparateGamePromotion ? "YES" : "NO"));
-                        CommandUtils.Log("          Test With Shared  : " + (Options(HostPlatform).bTestWithShared ? "YES" : "NO"));
+						CommandUtils.LogVerbose("          Promoted  : " + (Options(HostPlatform).bIsPromotable ? "YES" : "NO"));
+						CommandUtils.LogVerbose("          SeparatePromotable  : " + (Options(HostPlatform).bSeparateGamePromotion ? "YES" : "NO"));
+						CommandUtils.LogVerbose("          Test With Shared  : " + (Options(HostPlatform).bTestWithShared ? "YES" : "NO"));
                     }
-                    CommandUtils.Log("          Targets {0}:", Properties.Targets.Count);
+					CommandUtils.LogVerbose("          Targets {0}:", Properties.Targets.Count);
                     foreach (var ThisTarget in Properties.Targets)
                     {
-                        CommandUtils.Log("            TargetName          : " + ThisTarget.Value.TargetName);
-                        CommandUtils.Log("              Type          : " + ThisTarget.Key.ToString());
-                        CommandUtils.Log("              bUsesSteam  : " + (ThisTarget.Value.Rules.bUsesSteam ? "YES" : "NO"));
-                        CommandUtils.Log("              bUsesCEF3   : " + (ThisTarget.Value.Rules.bUsesCEF3 ? "YES" : "NO"));
-                        CommandUtils.Log("              bUsesSlate  : " + (ThisTarget.Value.Rules.bUsesSlate ? "YES" : "NO"));
+						CommandUtils.LogVerbose("            TargetName          : " + ThisTarget.Value.TargetName);
+						CommandUtils.LogVerbose("              Type          : " + ThisTarget.Key.ToString());
+						CommandUtils.LogVerbose("              bUsesSteam  : " + (ThisTarget.Value.Rules.bUsesSteam ? "YES" : "NO"));
+						CommandUtils.LogVerbose("              bUsesCEF3   : " + (ThisTarget.Value.Rules.bUsesCEF3 ? "YES" : "NO"));
+						CommandUtils.LogVerbose("              bUsesSlate  : " + (ThisTarget.Value.Rules.bUsesSlate ? "YES" : "NO"));
                         if (Array.IndexOf(MonolithicKinds, ThisTarget.Key) >= 0)
                         {
                             var Platforms = ThisTarget.Value.Rules.GUBP_GetPlatforms_MonolithicOnly(HostPlatform);
@@ -681,13 +681,13 @@ namespace AutomationTool
                                 var Configs = ThisTarget.Value.Rules.GUBP_GetConfigs_MonolithicOnly(HostPlatform, Plat);
                                 foreach (var Config in Configs)
                                 {
-                                    CommandUtils.Log("                Config  : " + Plat.ToString() + " " + Config.ToString());
+									CommandUtils.LogVerbose("                Config  : " + Plat.ToString() + " " + Config.ToString());
                                 }
                             }
 
                         }
                     }
-                    CommandUtils.Log("      Programs {0}:", Properties.Programs.Count);
+					CommandUtils.LogVerbose("      Programs {0}:", Properties.Programs.Count);
                     foreach (var ThisTarget in Properties.Programs)
                     {
                         bool bInternalToolOnly;
@@ -695,12 +695,12 @@ namespace AutomationTool
 						bool CrossCompile;
                         bool Tool = ThisTarget.Rules.GUBP_AlwaysBuildWithTools(HostPlatform, out bInternalToolOnly, out SeparateNode, out CrossCompile);
 
-                        CommandUtils.Log("            TargetName                    : " + ThisTarget.TargetName);
-                        CommandUtils.Log("              Build With Editor           : " + (ThisTarget.Rules.GUBP_AlwaysBuildWithBaseEditor() ? "YES" : "NO"));
-                        CommandUtils.Log("              Build With Tools            : " + (Tool && !bInternalToolOnly ? "YES" : "NO"));
-                        CommandUtils.Log("              Build With Internal Tools   : " + (Tool && bInternalToolOnly ? "YES" : "NO"));
-                        CommandUtils.Log("              Separate Node               : " + (Tool && SeparateNode ? "YES" : "NO"));
-						CommandUtils.Log("              Cross Compile               : " + (Tool && CrossCompile ? "YES" : "NO"));
+						CommandUtils.LogVerbose("            TargetName                    : " + ThisTarget.TargetName);
+						CommandUtils.LogVerbose("              Build With Editor           : " + (ThisTarget.Rules.GUBP_AlwaysBuildWithBaseEditor() ? "YES" : "NO"));
+						CommandUtils.LogVerbose("              Build With Tools            : " + (Tool && !bInternalToolOnly ? "YES" : "NO"));
+						CommandUtils.LogVerbose("              Build With Internal Tools   : " + (Tool && bInternalToolOnly ? "YES" : "NO"));
+						CommandUtils.LogVerbose("              Separate Node               : " + (Tool && SeparateNode ? "YES" : "NO"));
+						CommandUtils.LogVerbose("              Cross Compile               : " + (Tool && CrossCompile ? "YES" : "NO"));
                     }
                 }
             }
diff --git a/Engine/Source/Programs/AutomationTool/ScriptCompiler.cs b/Engine/Source/Programs/AutomationTool/ScriptCompiler.cs
--- a/Engine/Source/Programs/AutomationTool/ScriptCompiler.cs
+++ b/Engine/Source/Programs/AutomationTool/ScriptCompiler.cs
@@ -280,7 +280,7 @@ namespace AutomationTool
 						continue;
 					}
 					// Load the assembly into our app domain
-					CommandUtils.Log("Loading script DLL: {0}", ScriptsDLLFilename);
+					CommandUtils.LogVerbose("Loading script DLL: {0}", ScriptsDLLFilename);
 					try
 					{
 						var Dll = AppDomain.CurrentDomain.Load(AssemblyName.GetAssemblyName(ScriptsDLLFilename));
diff --git a/Engine/Source/Programs/AutomationTool/Scripts/GUBP.Automation.cs b/Engine/Source/Programs/AutomationTool/Scripts/GUBP.Automation.cs
--- a/Engine/Source/Programs/AutomationTool/Scripts/GUBP.Automation.cs
+++ b/Engine/Source/Programs/AutomationTool/Scripts/GUBP.Automation.cs
@@ -4148,11 +4148,11 @@ public class GUBP : BuildCommand
         }
         if (Result)
         {
-            Log("***** GUBP Trigger Node was already triggered {0} -> {1} : {2}", GUBPNodes[NodeToDo].GetFullName(), GameNameIfAny, NodeStoreName);
+            LogVerbose("***** GUBP Trigger Node was already triggered {0} -> {1} : {2}", GUBPNodes[NodeToDo].GetFullName(), GameNameIfAny, NodeStoreName);
         }
         else
         {
-            Log("***** GUBP Trigger Node was NOT yet triggered {0} -> {1} : {2}", GUBPNodes[NodeToDo].GetFullName(), GameNameIfAny, NodeStoreName);
+            LogVerbose("***** GUBP Trigger Node was NOT yet triggered {0} -> {1} : {2}", GUBPNodes[NodeToDo].GetFullName(), GameNameIfAny, NodeStoreName);
         }
         GUBPNodesCompleted.Add(NodeToDo, Result);
         return Result;
@@ -4161,7 +4161,7 @@ public class GUBP : BuildCommand
     {
         if (ParseParam("FakeEC"))
         {
-            Log("***** Would have ran ectool {0}", Args);
+            LogWarning("***** Would have ran ectool {0}", Args);
             return "We didn't actually run ectool";
         }
         else
@@ -5347,7 +5347,7 @@ public class GUBP : BuildCommand
         }
         else if (!LocalOnly && !HaveSharedTempStorage(false))
         {
-            Log("Looks like we want to use shared temp storage, but since we don't have it, we won't use it.");
+            LogWarning("Looks like we want to use shared temp storage, but since we don't have it, we won't use it.");
             LocalOnly = true;
         }
 
@@ -5381,7 +5381,7 @@ public class GUBP : BuildCommand
                 }
                 if (!FileExists_NoExceptions(P4IndexFileLocal))
                 {
-                    Log("{0} doesn't exist, checking in a new one", P4IndexFileP4);
+                    LogVerbose("{0} doesn't exist, checking in a new one", P4IndexFileP4);
                     WriteAllText(P4IndexFileLocal, "-1 0");
                     int WorkingCL = -1;
                     try
@@ -5393,7 +5393,7 @@ public class GUBP : BuildCommand
                     }
                     catch (Exception)
                     {
-                        Log("Add of CIS counter failed, assuming it now exists.");
+                        LogWarning("Add of CIS counter failed, assuming it now exists.");
                         if (WorkingCL > 0)
                         {
 							P4.DeleteChange(WorkingCL);
@@ -5403,7 +5403,7 @@ public class GUBP : BuildCommand
 				P4.Sync("-f " + P4IndexFileP4 + "#head");
                 if (!FileExists_NoExceptions(P4IndexFileLocal))
                 {
-                    Log("{0} doesn't exist, checking in a new one", P4IndexFileP4);
+                    LogVerbose("{0} doesn't exist, checking in a new one", P4IndexFileP4);
                     WriteAllText(P4IndexFileLocal, "-1 0");
                     int WorkingCL = -1;
                     try
@@ -5415,7 +5415,7 @@ public class GUBP : BuildCommand
                     }
                     catch (Exception)
                     {
-                        Log("Add of CIS counter failed, assuming it now exists.");
+                        LogWarning("Add of CIS counter failed, assuming it now exists.");
                         if (WorkingCL > 0)
                         {
                             P4.DeleteChange(WorkingCL);
@@ -5457,7 +5457,7 @@ public class GUBP : BuildCommand
                 }
                 {
                     var Line = String.Format("{0} {1}", NewIndex, NowMinutes);
-                    Log("Attempting to write {0} with {1}", P4IndexFileP4, Line);
+                    LogVerbose("Attempting to write {0} with {1}", P4IndexFileP4, Line);
                     int WorkingCL = -1;
                     try
                     {
@@ -5471,7 +5471,7 @@ public class GUBP : BuildCommand
                     }
                     catch (Exception)
                     {
-                        Log("Edit of CIS counter failed, assuming someone else checked in, retrying.");
+                        LogWarning("Edit of CIS counter failed, assuming someone else checked in, retrying.");
                         if (WorkingCL > 0)
                         {
 							P4.DeleteChange(WorkingCL);
@@ -5487,21 +5487,21 @@ public class GUBP : BuildCommand
             Log("Setting TimeIndex to {0}", TimeIndex);
         }
 
-        Log("************************* CL:			                    {0}", CL); 
-        Log("************************* P4Enabled:			            {0}", P4Enabled);
+        LogVerbose("************************* CL:			                    {0}", CL);
+		LogVerbose("************************* P4Enabled:			            {0}", P4Enabled);
         foreach (var HostPlatform in HostPlatforms)
         {
-            Log("*************************   HostPlatform:		            {0}", HostPlatform.ToString()); 
+			LogVerbose("*************************   HostPlatform:		            {0}", HostPlatform.ToString()); 
         }
-        Log("************************* StoreName:		                {0}", StoreName.ToString());
-        Log("************************* bCleanLocalTempStorage:		    {0}", bCleanLocalTempStorage);
-        Log("************************* bSkipTriggers:		            {0}", bSkipTriggers);
-        Log("************************* bSaveSharedTempStorage:		    {0}", bSaveSharedTempStorage);
-        Log("************************* bSignBuildProducts:		        {0}", bSignBuildProducts);
-        Log("************************* bFake:           		        {0}", bFake);
-        Log("************************* bFakeEC:           		        {0}", bFakeEC);
-        Log("************************* bHistory:           		        {0}", bHistory);
-        Log("************************* TimeIndex:           		    {0}", TimeIndex);        
+		LogVerbose("************************* StoreName:		                {0}", StoreName.ToString());
+		LogVerbose("************************* bCleanLocalTempStorage:		    {0}", bCleanLocalTempStorage);
+		LogVerbose("************************* bSkipTriggers:		            {0}", bSkipTriggers);
+		LogVerbose("************************* bSaveSharedTempStorage:		    {0}", bSaveSharedTempStorage);
+		LogVerbose("************************* bSignBuildProducts:		        {0}", bSignBuildProducts);
+		LogVerbose("************************* bFake:           		        {0}", bFake);
+		LogVerbose("************************* bFakeEC:           		        {0}", bFakeEC);
+		LogVerbose("************************* bHistory:           		        {0}", bHistory);
+		LogVerbose("************************* TimeIndex:           		    {0}", TimeIndex);        
 
 
         GUBPNodes = new Dictionary<string, GUBPNode>();        
@@ -5554,7 +5554,7 @@ public class GUBP : BuildCommand
         ActivePlatforms = SupportedPlatforms;
         foreach (var Plat in ActivePlatforms)
         {
-            Log("Active Platform: {0}", Plat.ToString());
+            LogVerbose("Active Platform: {0}", Plat.ToString());
         }
 
         if (HostPlatforms.Count >= 2)
@@ -5689,7 +5689,7 @@ public class GUBP : BuildCommand
                     var Proj = Branch.FindGame(Codeless.Key);
                     if (Proj == null)
                     {
-                        Log(System.Diagnostics.TraceEventType.Information, "{0} was listed as a codeless project by GUBP_NonCodeProjects_BaseEditorTypeOnly, however it does not exist in this branch.", Codeless.Key);
+                        Log(System.Diagnostics.TraceEventType.Verbose, "{0} was listed as a codeless project by GUBP_NonCodeProjects_BaseEditorTypeOnly, however it does not exist in this branch.", Codeless.Key);
                     }
                     else if (Proj.Properties.bIsCodeBasedProject)
                     {
@@ -5728,7 +5728,7 @@ public class GUBP : BuildCommand
                     }
                     else
                     {
-                        Log(System.Diagnostics.TraceEventType.Information, "{0} was listed as a codeless formal build GUBP_GetNonCodeFormalBuilds_BaseEditorTypeOnly, however it does not exist in this branch.", Codeless.Key);
+                        Log(System.Diagnostics.TraceEventType.Verbose, "{0} was listed as a codeless formal build GUBP_GetNonCodeFormalBuilds_BaseEditorTypeOnly, however it does not exist in this branch.", Codeless.Key);
                     }
                 }
             }
@@ -6480,13 +6480,13 @@ public class GUBP : BuildCommand
                         }
                         All += Dep;
                     }
-                    Log("  {0}: {1}      {2}", Node, Note, All);
+                    LogVerbose("  {0}: {1}      {2}", Node, Note, All);
                     FullNodeList.Add(Node, Note);
                     FullNodeDirectDependencies.Add(Node, All);
                 }
                 else
                 {
-                    Log("  {0}: {1} [Aggregate]", Node, Note);
+                    LogVerbose("  {0}: {1} [Aggregate]", Node, Note);
                 }
             }
             var DependencyFinish = DateTime.Now.ToString();
@@ -6532,7 +6532,7 @@ public class GUBP : BuildCommand
             {				
                 if (NodeSpec.Equals("Noop", StringComparison.InvariantCultureIgnoreCase))
                 {
-                    Log("Request for Noop node, done.");
+                    LogVerbose("Request for Noop node, done.");
                     PrintRunTime();
                     return;
                 }
@@ -6618,7 +6618,7 @@ public class GUBP : BuildCommand
         }
         if (NodesToDo.Count == 0)
         {
-            Log("No nodes specified, adding all nodes");
+            LogVerbose("No nodes specified, adding all nodes");
             foreach (var Node in GUBPNodes)
             {
                 NodesToDo.Add(Node.Key);
@@ -6626,7 +6626,7 @@ public class GUBP : BuildCommand
         }
         else if (TimeIndex != 0)
         {
-            Log("Check to make sure we didn't ask for nodes that will be culled by time index");
+            LogVerbose("Check to make sure we didn't ask for nodes that will be culled by time index");
             foreach (var NodeToDo in NodesToDo)
             {
                 if (TimeIndex % (1 << GUBPNodes[NodeToDo].DependentCISFrequencyQuantumShift()) != 0)
@@ -6636,10 +6636,10 @@ public class GUBP : BuildCommand
             }
         }
 
-        Log("Desired Nodes");
+        LogVerbose("Desired Nodes");
         foreach (var NodeToDo in NodesToDo)
         {
-            Log("  {0}", NodeToDo);			
+            LogVerbose("  {0}", NodeToDo);			
         }
         // if we are doing related to, then find things that depend on the selected nodes
         if (bRelatedToNode)
@@ -6729,18 +6729,18 @@ public class GUBP : BuildCommand
 		}		
 		if (TimeIndex != 0)
 		{
-			Log("Culling based on time index");
+			LogVerbose("Culling based on time index");
 			var NewNodesToDo = new HashSet<string>();
 			foreach (var NodeToDo in NodesToDo)
 			{
 				if (TimeIndex % (1 << GUBPNodes[NodeToDo].DependentCISFrequencyQuantumShift()) == 0)
 				{
-					Log("  Keeping {0}", NodeToDo);
+					LogVerbose("  Keeping {0}", NodeToDo);
 					NewNodesToDo.Add(NodeToDo);
 				}
 				else
 				{
-					Log("  Rejecting {0}", NodeToDo);
+					LogVerbose("  Rejecting {0}", NodeToDo);
 				}
 			}
 			NodesToDo = NewNodesToDo;
@@ -6757,7 +6757,7 @@ public class GUBP : BuildCommand
 				}
 				else
 				{
-					Log(" Rejecting {0} because -NoLinux was requested", NodeToDo);
+					LogVerbose(" Rejecting {0} because -NoLinux was requested", NodeToDo);
 				}
 			}
 			NodesToDo = NewNodesToDo;
@@ -6838,19 +6838,19 @@ public class GUBP : BuildCommand
         }
 		if (bPreflightBuild)
 		{
-			Log("Culling triggers and downstream for preflight builds ");
+			LogVerbose("Culling triggers and downstream for preflight builds ");
 			var NewNodesToDo = new HashSet<string>();
 			foreach (var NodeToDo in NodesToDo)
 			{
 				var TriggerDot = GetControllingTriggerDotName(NodeToDo);
 				if (TriggerDot == "" && !GUBPNodes[NodeToDo].TriggerNode())
 				{
-					Log("  Keeping {0}", NodeToDo);
+					LogVerbose("  Keeping {0}", NodeToDo);
 					NewNodesToDo.Add(NodeToDo);
 				}
 				else
 				{
-					Log("  Rejecting {0}", NodeToDo);
+					LogVerbose("  Rejecting {0}", NodeToDo);
 				}
 			}
 			NodesToDo = NewNodesToDo;
@@ -6859,17 +6859,17 @@ public class GUBP : BuildCommand
         GUBPNodesCompleted = new Dictionary<string, bool>();
         GUBPNodesHistory = new Dictionary<string, NodeHistory>();
 
-        Log("******* Caching completion");
+        LogVerbose("******* Caching completion");
         {
             var StartTime = DateTime.UtcNow;
             foreach (var Node in NodesToDo)
             {
-                Log("** {0}", Node);
+                LogVerbose("** {0}", Node);
                 NodeIsAlreadyComplete(Node, LocalOnly); // cache these now to avoid spam later
                 GetControllingTriggerDotName(Node);
             }
             var BuildDuration = (DateTime.UtcNow - StartTime).TotalMilliseconds;
-            Log("Took {0}s to cache completion for {1} nodes", BuildDuration / 1000, NodesToDo.Count);
+			LogVerbose("Took {0}s to cache completion for {1} nodes", BuildDuration / 1000, NodesToDo.Count);
         }
         /*if (CLString != "" && StoreName.Contains(CLString) && !ParseParam("NoHistory"))
         {
@@ -6904,7 +6904,7 @@ public class GUBP : BuildCommand
             }
         }
 
-        Log("*********** Desired And Dependent Nodes, in order.");
+        LogVerbose("*********** Desired And Dependent Nodes, in order.");
         PrintNodes(this, OrdereredToDo, LocalOnly, UnfinishedTriggers);		
         //check sorting
         {
diff --git a/Engine/Source/Programs/AutomationTool/TempStorage.cs b/Engine/Source/Programs/AutomationTool/TempStorage.cs
--- a/Engine/Source/Programs/AutomationTool/TempStorage.cs
+++ b/Engine/Source/Programs/AutomationTool/TempStorage.cs
@@ -940,7 +940,7 @@ namespace AutomationTool
             var LocalManifest = LocalTempStorageManifestFilename(Env, StorageBlockName);
             if (FileExists_NoExceptions(LocalManifest))
             {
-                Log("Found local manifest {0}", LocalManifest);
+                LogVerbose("Found local manifest {0}", LocalManifest);
                 var Local = new TempStorageManifest();
                 Local.Load(LocalManifest);
                 Files = Local.GetFiles(BaseFolder);
@@ -961,7 +961,7 @@ namespace AutomationTool
             {
                 throw new AutomationException("base folder {0} should end with a separator", BlockPath);
             }
-            Log("Attempting to retrieve from {0}", BlockPath);
+            LogVerbose("Attempting to retrieve from {0}", BlockPath);
             if (!DirectoryExists_NoExceptions(BlockPath))
             {
                 throw new AutomationException("Storage Block Does Not Exists! {0}", BlockPath);
@@ -990,7 +990,7 @@ namespace AutomationTool
                     var DestFile = CombinePaths(BaseFolder, RelativeFile);
                     if (FileExists_NoExceptions(true, DestFile))
                     {
-                        Log("Dest file {0} already exists, deleting and overwriting", DestFile);
+                        LogVerbose("Dest file {0} already exists, deleting and overwriting", DestFile);
                         DeleteFile(DestFile);
                     }
                     CopyFile(Filename, DestFile, true);
@@ -1022,7 +1022,7 @@ namespace AutomationTool
                     var DestFile = CombinePaths(BaseFolder, RelativeFile);
                     if (FileExists_NoExceptions(true, DestFile))
                     {
-                        Log("Dest file {0} already exists, deleting and overwriting", DestFile);
+                        LogVerbose("Dest file {0} already exists, deleting and overwriting", DestFile);
                         DeleteFile(DestFile);
                     }
                     SrcFiles.Add(Filename);
diff --git a/Engine/Source/Programs/AutomationTool/Utils.cs b/Engine/Source/Programs/AutomationTool/Utils.cs
--- a/Engine/Source/Programs/AutomationTool/Utils.cs
+++ b/Engine/Source/Programs/AutomationTool/Utils.cs
@@ -34,7 +34,7 @@ namespace AutomationTool
 			}
 			if (!bQuiet)
 			{
-				Log.WriteLine(TraceEventType.Information, "GetEnvironmentVariable {0}={1}", VarName, Value);
+				Log.WriteLine(TraceEventType.Verbose, "GetEnvironmentVariable {0}={1}", VarName, Value);
 			}
 			return Value;
 		}
@@ -48,7 +48,7 @@ namespace AutomationTool
 		{
 			if( !bQuiet)
 			{ 
-				Log.WriteLine(TraceEventType.Information, "SafeCreateDirectory {0}", Path);
+				Log.WriteLine(TraceEventType.Verbose, "SafeCreateDirectory {0}", Path);
 			}
 
 			bool Result = true;
@@ -76,7 +76,7 @@ namespace AutomationTool
 		{
 			if (!bQuiet)
 			{
-				Log.WriteLine(TraceEventType.Information, "SafeDeleteFile {0}", Path);
+				Log.WriteLine(TraceEventType.Verbose, "SafeDeleteFile {0}", Path);
 			}
 			int MaxAttempts = bQuiet ? 1 : 10;
 			int Attempts = 0;
@@ -116,7 +116,7 @@ namespace AutomationTool
 			{
 				if (bQuiet)
 				{
-					Log.WriteLine(TraceEventType.Information, "Failed to delete file {0} in {1} attempts.", Path, MaxAttempts);
+					Log.WriteLine(TraceEventType.Verbose, "Failed to delete file {0} in {1} attempts.", Path, MaxAttempts);
 				}
 				else
 				{
@@ -137,7 +137,7 @@ namespace AutomationTool
 		{
 			if (!bQuiet)
 			{
-				Log.WriteLine(TraceEventType.Information, "RecursivelyDeleteDirectory {0}", Path);
+				Log.WriteLine(TraceEventType.Verbose, "RecursivelyDeleteDirectory {0}", Path);
 			}
 			// Delete all files. This will also delete read-only files.
 			var FilesInDirectory = Directory.EnumerateFiles(Path);
@@ -173,7 +173,7 @@ namespace AutomationTool
 		{
 			if (!bQuiet)
 			{
-				Log.WriteLine(TraceEventType.Information, "SafeDeleteEmptyDirectory {0}", Path);
+				Log.WriteLine(TraceEventType.Verbose, "SafeDeleteEmptyDirectory {0}", Path);
 			}
 			const int MaxAttempts = 10;
 			int Attempts = 0;
@@ -218,7 +218,7 @@ namespace AutomationTool
 		{
 			if (!bQuiet)
 			{
-				Log.WriteLine(TraceEventType.Information, "SafeDeleteDirectory {0}", Path);
+				Log.WriteLine(TraceEventType.Verbose, "SafeDeleteDirectory {0}", Path);
 			}
 			if (Directory.Exists(Path))
 			{
@@ -240,7 +240,7 @@ namespace AutomationTool
 		{
 			if( !bQuiet )
 			{ 
-				Log.WriteLine(TraceEventType.Information, "SafeRenameFile {0} {1}", OldName, NewName);
+				Log.WriteLine(TraceEventType.Verbose, "SafeRenameFile {0} {1}", OldName, NewName);
 			}
 			const int MaxAttempts = 10;
 			int Attempts = 0;
@@ -330,7 +330,7 @@ namespace AutomationTool
 		{
 			if (!bQuiet)
 			{
-				Log.WriteLine(TraceEventType.Information, "SafeCopyFile {0} {1}", SourceName, TargetName);
+				Log.WriteLine(TraceEventType.Verbose, "SafeCopyFile {0} {1}", SourceName, TargetName);
 			}
 			const int MaxAttempts = 10;
 			int Attempts = 0;
@@ -409,7 +409,7 @@ namespace AutomationTool
 		/// <returns>An array containing all lines read from the file or null if the file could not be read.</returns>
 		public static string[] SafeReadAllLines(string Filename)
 		{
-			Log.WriteLine(TraceEventType.Information, "SafeReadAllLines {0}", Filename);
+			Log.WriteLine(TraceEventType.Verbose, "SafeReadAllLines {0}", Filename);
 			string[] Result = null;
 			try
 			{
@@ -430,7 +430,7 @@ namespace AutomationTool
 		/// <returns>String containing all text read from the file or null if the file could not be read.</returns>
 		public static string SafeReadAllText(string Filename)
 		{
-			Log.WriteLine(TraceEventType.Information, "SafeReadAllLines {0}", Filename);
+			Log.WriteLine(TraceEventType.Verbose, "SafeReadAllLines {0}", Filename);
 			string Result = null;
 			try
 			{
@@ -455,7 +455,7 @@ namespace AutomationTool
 		{
 			if (!bQuiet)
 			{
-				Log.WriteLine(TraceEventType.Information, "FindFiles {0} {1} {2}", Path, SearchPattern, Recursive);
+				Log.WriteLine(TraceEventType.Verbose, "FindFiles {0} {1} {2}", Path, SearchPattern, Recursive);
 			}
 
 			// On Linux, filter out symlinks since we (usually) create them to fix mispelled case-sensitive filenames in content, and if they aren't filtered, 
@@ -500,7 +500,7 @@ namespace AutomationTool
 		{
 			if (!bQuiet)
 			{
-				Log.WriteLine(TraceEventType.Information, "FindDirectories {0} {1} {2}", Path, SearchPattern, Recursive);
+				Log.WriteLine(TraceEventType.Verbose, "FindDirectories {0} {1} {2}", Path, SearchPattern, Recursive);
 			}
 			return Directory.GetDirectories(Path, SearchPattern, Recursive ? SearchOption.AllDirectories : SearchOption.TopDirectoryOnly);
 		}
@@ -516,7 +516,7 @@ namespace AutomationTool
 		{
 			if (!bQuiet)
 			{
-				Log.WriteLine(TraceEventType.Information, "SafeFindFiles {0} {1} {2}", Path, SearchPattern, Recursive);
+				Log.WriteLine(TraceEventType.Verbose, "SafeFindFiles {0} {1} {2}", Path, SearchPattern, Recursive);
 			}
 			string[] Files = null;
 			try
@@ -542,7 +542,7 @@ namespace AutomationTool
 		{
 			if (!bQuiet)
 			{
-				Log.WriteLine(TraceEventType.Information, "SafeFindDirectories {0} {1} {2}", Path, SearchPattern, Recursive);
+				Log.WriteLine(TraceEventType.Verbose, "SafeFindDirectories {0} {1} {2}", Path, SearchPattern, Recursive);
 			}
 			string[] Directories = null;
 			try
@@ -571,7 +571,7 @@ namespace AutomationTool
 				Result = File.Exists(Path);
 				if (!bQuiet)
 				{
-					Log.WriteLine(TraceEventType.Information, "SafeFileExists {0}={1}", Path, Result);
+					Log.WriteLine(TraceEventType.Verbose, "SafeFileExists {0}={1}", Path, Result);
 				}
 			}
 			catch (Exception Ex)
@@ -597,7 +597,7 @@ namespace AutomationTool
 				Result = Directory.Exists(Path);
 				if (!bQuiet)
 				{
-					Log.WriteLine(TraceEventType.Information, "SafeDirectoryExists {0}={1}", Path, Result);
+					Log.WriteLine(TraceEventType.Verbose, "SafeDirectoryExists {0}={1}", Path, Result);
 				}
 			}
 			catch (Exception Ex)
@@ -616,7 +616,7 @@ namespace AutomationTool
 		/// <returns>True if the operation was successful, false otherwise.</returns>
 		public static bool SafeWriteAllLines(string Path, string[] Text)
 		{
-			Log.WriteLine(TraceEventType.Information, "SafeWriteAllLines {0}", Path);
+			Log.WriteLine(TraceEventType.Verbose, "SafeWriteAllLines {0}", Path);
 			bool Result = false;
 			try
 			{
@@ -639,7 +639,7 @@ namespace AutomationTool
 		/// <returns>True if the operation was successful, false otherwise.</returns>
 		public static bool SafeWriteAllText(string Path, string Text)
 		{
-			Log.WriteLine(TraceEventType.Information, "SafeWriteAllText {0}", Path);
+			Log.WriteLine(TraceEventType.Verbose, "SafeWriteAllText {0}", Path);
 			bool Result = false;
 			try
 			{
@@ -662,7 +662,7 @@ namespace AutomationTool
 		/// <returns>True if the operation was successful, false otherwise.</returns>
 		public static bool SafeWriteAllBytes(string Path, byte[] Bytes)
 		{
-			Log.WriteLine(TraceEventType.Information, "SafeWriteAllBytes {0}", Path);
+			Log.WriteLine(TraceEventType.Verbose, "SafeWriteAllBytes {0}", Path);
 			bool Result = false;
 			try
 			{
diff --git a/Engine/Source/Programs/UnrealBuildTool/Linux/LinuxToolChain.cs b/Engine/Source/Programs/UnrealBuildTool/Linux/LinuxToolChain.cs
--- a/Engine/Source/Programs/UnrealBuildTool/Linux/LinuxToolChain.cs
+++ b/Engine/Source/Programs/UnrealBuildTool/Linux/LinuxToolChain.cs
@@ -1146,7 +1146,7 @@ namespace UnrealBuildTool
 
 			ProcessStartInfo StartInfo = new ProcessStartInfo();
 			StartInfo.FileName = GetStripPath(UEBuildPlatform.GetBuildPlatform(UnrealTargetPlatform.Linux).GetActiveArchitecture());
-			StartInfo.Arguments = "--strip-debug " + TargetFileName;
+			StartInfo.Arguments = "--strip-debug \"" + TargetFileName + "\"";
 			StartInfo.UseShellExecute = false;
 			StartInfo.CreateNoWindow = true;
 			Utils.RunLocalProcessAndLogOutput(StartInfo);
diff --git a/Engine/Source/Programs/UnrealBuildTool/System/UnrealBuildTool.cs b/Engine/Source/Programs/UnrealBuildTool/System/UnrealBuildTool.cs
--- a/Engine/Source/Programs/UnrealBuildTool/System/UnrealBuildTool.cs
+++ b/Engine/Source/Programs/UnrealBuildTool/System/UnrealBuildTool.cs
@@ -857,7 +857,7 @@ namespace UnrealBuildTool
                 bColorConsoleOutput: true,
                 TraceListeners: new[] 
                 {
-                    new ConsoleTraceListener(),
+                    new UnfilteredConsoleTraceListener(),
                     !string.IsNullOrEmpty(BuildConfiguration.LogFilename) ? new TextWriterTraceListener(new StreamWriter(new FileStream(BuildConfiguration.LogFilename, FileMode.Create, FileAccess.ReadWrite, FileShare.Read)) { AutoFlush = true }) : null,
                 });
         }
@@ -2755,4 +2755,34 @@ namespace UnrealBuildTool
 			}
 		}
     }
+
+	#region UnfilteredConsoleTraceListener
+
+	/// <summary>
+	///  Unfiltered Console Trace Listener
+	///  </summary>
+	class UnfilteredConsoleTraceListener : TextWriterTraceListener
+	{
+		public UnfilteredConsoleTraceListener()
+			: base(Console.OpenStandardOutput())
+		{ }
+
+		#region TraceListener Interface
+
+		public override void WriteLine(string message)
+		{
+			// strip off the message code
+			int Code = Convert.ToInt32(message.Substring(0, 3), 16);
+
+			// filter based on the message code
+			if (Code < (int)TraceEventType.Verbose || Log.bIsVerbose)
+			{
+				base.WriteLine(message.Substring(3));
+			}
+		}
+
+		#endregion
+	}
+	#endregion
+
 }
diff --git a/Engine/Source/Programs/UnrealBuildTool/System/Utils.cs b/Engine/Source/Programs/UnrealBuildTool/System/Utils.cs
--- a/Engine/Source/Programs/UnrealBuildTool/System/Utils.cs
+++ b/Engine/Source/Programs/UnrealBuildTool/System/Utils.cs
@@ -1046,6 +1046,16 @@ namespace UnrealBuildTool
             return Severity <= TraceEventType.Error ? "ERROR: " : Severity == TraceEventType.Warning ? "WARNING: " : "";
         }
 
+		/// <summary>
+		/// Converts a TraceEventType into a message code
+		/// </summary>
+		/// <param name="EventType"></param>
+		/// <returns></returns>
+		private static int GetMessageCode(TraceEventType Severity)
+		{
+			return (int)Severity;
+		}
+
         /// <summary>
         /// Formats message for logging. Enforces the configured options.
         /// </summary>
@@ -1058,12 +1068,13 @@ namespace UnrealBuildTool
         [MethodImplAttribute(MethodImplOptions.NoInlining)]
         private static string FormatMessage(int StackFramesToSkip, string CustomSource, TraceEventType Verbosity, string Format, params object[] Args)
         {
-            return string.Format("{0}{1}{2}{3}",
+            return string.Format("{4}{0}{1}{2}{3}",
                     Timer != null ? String.Format("[{0:hh\\:mm\\:ss\\.fff}] ", Timer.Elapsed) : "",
                     bLogSources ? string.Format("{0}: ", string.IsNullOrEmpty(CustomSource) ? GetSource(StackFramesToSkip) : CustomSource) : "",
                     bLogSeverity ? GetSeverityPrefix(Verbosity) : "",
                     // If there are no extra args, don't try to format the string, in case it has any format control characters in it (our LOCTEXT strings tend to).
-                    Args.Length > 0 ? string.Format(Format, Args) : Format);
+                    Args.Length > 0 ? string.Format(Format, Args) : Format,
+					GetMessageCode(Verbosity).ToString("X3"));
         }
 
         /// <summary>
@@ -1095,7 +1106,7 @@ namespace UnrealBuildTool
 				WriteOnceSet.Add(Formatted);
 			}
 
-            if (Verbosity < TraceEventType.Verbose || bLogVerbose)
+//            if (Verbosity < TraceEventType.Verbose || bLogVerbose)
             {
                 // Do console color highlighting here.
                 ConsoleColor DefaultColor = ConsoleColor.Gray;
