- Fixed UAT to not error out when building for the unified Desktop platform (recent change with receipts and XP had to be fixed up a little)

diff --git a/Engine/Source/Programs/AutomationTool/Scripts/CopyBuildToStagingDirectory.Automation.cs b/Engine/Source/Programs/AutomationTool/Scripts/CopyBuildToStagingDirectory.Automation.cs
--- a/Engine/Source/Programs/AutomationTool/Scripts/CopyBuildToStagingDirectory.Automation.cs
+++ b/Engine/Source/Programs/AutomationTool/Scripts/CopyBuildToStagingDirectory.Automation.cs
@@ -1467,7 +1467,16 @@ public partial class Project : CommandUtils
 
 					foreach (UnrealTargetPlatform ReceiptPlatform in SubPlatformsToStage)
 					{
-						string Architecture = !String.IsNullOrEmpty(Params.SpecifiedArchitecture) ? Params.SpecifiedArchitecture : UEBuildPlatform.GetBuildPlatform(ReceiptPlatform).GetActiveArchitecture();
+                        string Architecture = Params.SpecifiedArchitecture;
+                        if (string.IsNullOrEmpty(Architecture))
+                        {
+                            Architecture = "";
+                            var BuildPlatform = UEBuildPlatform.GetBuildPlatform(ReceiptPlatform, true);
+                            if (BuildPlatform != null)
+                            {
+                                Architecture = BuildPlatform.GetActiveArchitecture();
+                            }
+                        }
 						string ReceiptFileName = BuildReceipt.GetDefaultPath(ReceiptBaseDir, Target, ReceiptPlatform, Config, Architecture);
 						if(!File.Exists(ReceiptFileName))
 						{
