Unshelved from pending changelist '2589055':

PR #1262: AutomationTool unused variables cleanup (Contributed by [REDACTED_PERSON])

diff --git a/Engine/Source/Programs/AutomationTool/Scripts/GUBP.Automation.cs b/Engine/Source/Programs/AutomationTool/Scripts/GUBP.Automation.cs
--- a/Engine/Source/Programs/AutomationTool/Scripts/GUBP.Automation.cs
+++ b/Engine/Source/Programs/AutomationTool/Scripts/GUBP.Automation.cs
@@ -4924,7 +4924,6 @@ public class GUBP : BuildCommand
     }
 	void GetFailureEmails(string NodeToDo, string CLString, bool OnlyLateUpdates = false)
 	{
-		var StartTime = DateTime.UtcNow;
         string EMails;
         string FailCauserEMails = "";
         string EMailNote = "";
@@ -5112,8 +5111,6 @@ public class GUBP : BuildCommand
     }
     List<string> GetECPropsForNode(string NodeToDo, string CLString, out string EMails, bool OnlyLateUpdates = false)
     {
-        var StartTime = DateTime.UtcNow;
-
         var ECProps = new List<string>();
         EMails = "";		
 		var AdditonalEmails = "";
@@ -5139,7 +5136,6 @@ public class GUBP : BuildCommand
 			ECProps.Add(string.Format("JobStepPath/{0}={1}", NodeToDo, GetJobStepPath(NodeToDo)));
 		}
 		
-        var BuildDuration = (DateTime.UtcNow - StartTime).TotalMilliseconds;
         return ECProps;
     }
 
