Adding information to warning message to help diagnose cause of issue. This won't add unneccesary noise to normal log output, but in the event we are unable to write to the CSV file it will provide the name of the CSV file and the exception caught, appended to the existing warning message.

diff --git a/Engine/Source/Programs/AutomationTool/CommandUtils.cs b/Engine/Source/Programs/AutomationTool/CommandUtils.cs
--- a/Engine/Source/Programs/AutomationTool/CommandUtils.cs
+++ b/Engine/Source/Programs/AutomationTool/CommandUtils.cs
@@ -1907,9 +1907,9 @@ namespace AutomationTool
                 {
                     File.AppendAllText(CmdEnv.CSVFile, CSVBuilder.ToString());
                 }
-                catch
+                catch (Exception e)
                 {
-                    Log(System.Diagnostics.TraceEventType.Warning, "Could not append to csv file.");
+                    Log(System.Diagnostics.TraceEventType.Warning, "Could not append to csv file ({0}) : {1}", CmdEnv.CSVFile, e.ToString());
                 }
             }
         }
