fix for crash when the message to WriteLine doesn't include a filter code

diff --git a/Engine/Source/Programs/AutomationTool/Log.cs b/Engine/Source/Programs/AutomationTool/Log.cs
--- a/Engine/Source/Programs/AutomationTool/Log.cs
+++ b/Engine/Source/Programs/AutomationTool/Log.cs
@@ -280,12 +280,19 @@ namespace AutomationTool
 		public override void WriteLine(string message)
 		{
 			// strip off the message code
-			int Code = Convert.ToInt32(message.Substring(0,3), 16);
-			// filter the message based on the code
-			if (Code < (int)TraceEventType.Verbose || UnrealBuildTool.Log.bIsVerbose)
+			int Code = 0;
+			if (Int32.TryParse(message.Substring(0, 3), System.Globalization.NumberStyles.HexNumber, null, out Code))
 			{
-				// filter based on the message code
-				base.WriteLine(message.Substring(3));
+				// filter the message based on the code
+				if (Code < (int)TraceEventType.Verbose || UnrealBuildTool.Log.bIsVerbose)
+				{
+					// filter based on the message code
+					base.WriteLine(message.Substring(3));
+				}
+			}
+			else
+			{
+				base.WriteLine(message);
 			}
 		}
 
diff --git a/Engine/Source/Programs/UnrealBuildTool/System/UnrealBuildTool.cs b/Engine/Source/Programs/UnrealBuildTool/System/UnrealBuildTool.cs
--- a/Engine/Source/Programs/UnrealBuildTool/System/UnrealBuildTool.cs
+++ b/Engine/Source/Programs/UnrealBuildTool/System/UnrealBuildTool.cs
@@ -2772,12 +2772,18 @@ namespace UnrealBuildTool
 		public override void WriteLine(string message)
 		{
 			// strip off the message code
-			int Code = Convert.ToInt32(message.Substring(0, 3), 16);
-
-			// filter based on the message code
-			if (Code < (int)TraceEventType.Verbose || Log.bIsVerbose)
+			int Code = 0;
+			if (Int32.TryParse(message.Substring(0, 3), System.Globalization.NumberStyles.HexNumber, null, out Code))
+			{
+				// filter based on the message code
+				if (Code < (int)TraceEventType.Verbose || Log.bIsVerbose)
+				{
+					base.WriteLine(message.Substring(3));
+				}
+			}
+			else
 			{
-				base.WriteLine(message.Substring(3));
+				base.WriteLine(message);
 			}
 		}
 
