Move GUBP into its own project.

diff --git a/Engine/Source/Programs/AutomationTool/GUBP/GUBP.Automation.csproj b/Engine/Source/Programs/AutomationTool/GUBP/GUBP.Automation.csproj
new file mode 100644
--- /dev/null
+++ b/Engine/Source/Programs/AutomationTool/GUBP/GUBP.Automation.csproj
@@ -0,0 +1,75 @@
+﻿<?xml version="1.0" encoding="utf-8"?>
+<Project ToolsVersion="4.0" DefaultTargets="Build" xmlns="http://schemas.microsoft.com/developer/msbuild/2003">
+  <Import Project="$(MSBuildExtensionsPath)\$(MSBuildToolsVersion)\Microsoft.Common.props" Condition="Exists('$(MSBuildExtensionsPath)\$(MSBuildToolsVersion)\Microsoft.Common.props')" />
+  <PropertyGroup>
+    <Configuration Condition=" '$(Configuration)' == '' ">Debug</Configuration>
+    <Platform Condition=" '$(Platform)' == '' ">AnyCPU</Platform>
+    <ProjectGuid>{9A1A118B-37FC-4F33-BBF4-B431B6E3B3AC}</ProjectGuid>
+    <OutputType>Library</OutputType>
+    <AppDesignerFolder>Properties</AppDesignerFolder>
+    <RootNamespace>GUBP.Automation</RootNamespace>
+    <AssemblyName>GUBP.Automation</AssemblyName>
+    <TargetFrameworkVersion>v4.0</TargetFrameworkVersion>
+    <FileAlignment>512</FileAlignment>
+  </PropertyGroup>
+  <PropertyGroup Condition=" '$(Configuration)|$(Platform)' == 'Debug|AnyCPU' ">
+    <DebugSymbols>true</DebugSymbols>
+    <DebugType>full</DebugType>
+    <Optimize>false</Optimize>
+    <OutputPath>..\..\..\..\Binaries\DotNET\AutomationScripts\</OutputPath>
+    <DefineConstants>DEBUG;TRACE</DefineConstants>
+    <ErrorReport>prompt</ErrorReport>
+    <WarningLevel>4</WarningLevel>
+    <PlatformTarget>AnyCPU</PlatformTarget>
+    <TreatWarningsAsErrors>true</TreatWarningsAsErrors>
+  </PropertyGroup>
+  <PropertyGroup Condition=" '$(Configuration)|$(Platform)' == 'Development|AnyCPU' ">
+    <DebugType>pdbonly</DebugType>
+    <Optimize>true</Optimize>
+    <OutputPath>..\..\..\..\Binaries\DotNET\AutomationScripts\</OutputPath>
+    <DefineConstants>TRACE</DefineConstants>
+    <ErrorReport>prompt</ErrorReport>
+    <WarningLevel>4</WarningLevel>
+    <PlatformTarget>AnyCPU</PlatformTarget>
+    <TreatWarningsAsErrors>true</TreatWarningsAsErrors>
+  </PropertyGroup>
+  <ItemGroup>
+    <Reference Include="Ionic.Zip.Reduced">
+      <HintPath>..\..\..\..\Binaries\DotNET\Ionic.Zip.Reduced.dll</HintPath>
+      <Private>False</Private>
+    </Reference>
+    <Reference Include="OneSky, Version=1.0.0.0, Culture=neutral, processorArchitecture=MSIL">
+      <HintPath>..\..\..\..\Binaries\DotNET\OneSky.dll</HintPath>
+      <SpecificVersion>False</SpecificVersion>
+      <Private>False</Private>
+    </Reference>
+    <Reference Include="System" />
+    <Reference Include="System.Core" />
+    <Reference Include="System.Runtime.Serialization" />
+    <Reference Include="System.Web" />
+    <Reference Include="System.Web.Extensions" />
+    <Reference Include="System.Xml.Linq" />
+    <Reference Include="System.Data.DataSetExtensions" />
+    <Reference Include="Microsoft.CSharp" />
+    <Reference Include="System.Data" />
+    <Reference Include="System.Xml" />
+  </ItemGroup>
+  <ItemGroup>
+    <Compile Include="..\..\DotNETCommon\MetaData.cs">
+      <Link>Properties\MetaData.cs</Link>
+    </Compile>
+    <Compile Include="GUBP.Automation.cs" />
+  </ItemGroup>
+  <ItemGroup>
+    <ProjectReference Include="..\..\UnrealBuildTool\UnrealBuildTool.csproj">
+      <Project>{fd7c5e1a-cfe4-4fd5-a525-1eb1599a39ac}</Project>
+      <Name>UnrealBuildTool</Name>
+      <Private>False</Private>
+    </ProjectReference>
+    <ProjectReference Include="..\AutomationTool.csproj">
+      <Project>{a35b0261-d868-475e-a098-eb14ce6dc89f}</Project>
+      <Name>AutomationTool</Name>
+    </ProjectReference>
+  </ItemGroup>
+  <Import Project="$(MSBuildToolsPath)\Microsoft.CSharp.targets" />
+</Project>
\ No newline at end of file
diff --git a/Engine/Source/Programs/AutomationTool/Scripts/AutomationScripts.Automation.csproj b/Engine/Source/Programs/AutomationTool/Scripts/AutomationScripts.Automation.csproj
--- a/Engine/Source/Programs/AutomationTool/Scripts/AutomationScripts.Automation.csproj
+++ b/Engine/Source/Programs/AutomationTool/Scripts/AutomationScripts.Automation.csproj
@@ -73,7 +73,6 @@
     <Compile Include="DeployCommand.Automation.cs" />
     <Compile Include="FixupRedirects.Automation.cs" />
     <Compile Include="GenerateDSYM.Automation.cs" />
-    <Compile Include="GUBP.Automation.cs" />
     <Compile Include="RocketBuild.Automation.cs" />
     <Compile Include="IPhonePackager.Automation.cs" />
     <Compile Include="LauncherLocalization.Automation.cs" />
@@ -102,6 +101,11 @@
       <Name>UnrealBuildTool</Name>
       <Private>False</Private>
     </ProjectReference>
+    <ProjectReference Include="..\GUBP\GUBP.Automation.csproj">
+      <Project>{9a1a118b-37fc-4f33-bbf4-b431b6e3b3ac}</Project>
+      <Name>GUBP.Automation</Name>
+      <Private>False</Private>
+    </ProjectReference>
   </ItemGroup>
   <Import Project="$(MSBuildToolsPath)\Microsoft.CSharp.targets" />
 </Project>
\ No newline at end of file
