BUGFIX: BranchHacker Exclusions remove Tool nodes, but don't remove NonUnityTool nodes.
* Before, if you, say, remove "EpicGamesLauncher" it will remove the default launcher nodes, but will still add the NonUnityTool nodes (GUBP.Automation.cs Line 5617 and 5641).
* This would cause the GUBP verification checks to fail because the NonUnity nodes for the launcher depend on the unity nodes. The idea is to Exclude ALL of that program's nodes from the graph, otherwise there's no way to exclude NonUnity nodes at all!
* Now we check the branch hacker before removing both the Tools node and the NonUnityTool node.

diff --git a/Engine/Source/Programs/AutomationTool/GUBP/GUBP.Automation.cs b/Engine/Source/Programs/AutomationTool/GUBP/GUBP.Automation.cs
--- a/Engine/Source/Programs/AutomationTool/GUBP/GUBP.Automation.cs
+++ b/Engine/Source/Programs/AutomationTool/GUBP/GUBP.Automation.cs
@@ -5599,49 +5599,55 @@ public class GUBP : BuildCommand
 				}
 				foreach (var ProgramTarget in Branch.BaseEngineProject.Properties.Programs)
 				{
-					bool bInternalOnly;
-					bool SeparateNode;
-					bool CrossCompile;
+                    if (!BranchOptions.ExcludeNodes.Contains(ProgramTarget.TargetName))
+                    {
+                        bool bInternalOnly;
+                        bool SeparateNode;
+                        bool CrossCompile;
 
-					if (ProgramTarget.Rules.GUBP_AlwaysBuildWithTools(HostPlatform, out bInternalOnly, out SeparateNode, out CrossCompile) && ProgramTarget.Rules.SupportsPlatform(HostPlatform) && SeparateNode && !BranchOptions.ExcludeNodes.Contains(ProgramTarget.TargetName))
-					{
-						if (bInternalOnly)
-						{
-							AddNode(new SingleInternalToolsNode(HostPlatform, ProgramTarget));						
-						}
-						else
-						{
-							AddNode(new SingleToolsNode(HostPlatform, ProgramTarget));
-						}				
-					}
-					if (ProgramTarget.Rules.GUBP_IncludeNonUnityToolTest())
-					{
-						AddNode(new NonUnityToolNode(HostPlatform, ProgramTarget));
-					}
+                        if (ProgramTarget.Rules.GUBP_AlwaysBuildWithTools(HostPlatform, out bInternalOnly, out SeparateNode, out CrossCompile) && ProgramTarget.Rules.SupportsPlatform(HostPlatform) && SeparateNode)
+                        {
+                            if (bInternalOnly)
+                            {
+                                AddNode(new SingleInternalToolsNode(HostPlatform, ProgramTarget));
+                            }
+                            else
+                            {
+                                AddNode(new SingleToolsNode(HostPlatform, ProgramTarget));
+                            }
+                        }
+                        if (ProgramTarget.Rules.GUBP_IncludeNonUnityToolTest())
+                        {
+                            AddNode(new NonUnityToolNode(HostPlatform, ProgramTarget));
+                        }
+                    }
 				}
 				foreach(var CodeProj in Branch.CodeProjects)
 				{
 					foreach(var ProgramTarget in CodeProj.Properties.Programs)
 					{
-						bool bInternalNodeOnly;
-						bool SeparateNode;
-						bool CrossCompile;
+                        if (!BranchOptions.ExcludeNodes.Contains(ProgramTarget.TargetName))
+                        {
+                            bool bInternalNodeOnly;
+                            bool SeparateNode;
+                            bool CrossCompile;
 
-						if(ProgramTarget.Rules.GUBP_AlwaysBuildWithTools(HostPlatform, out bInternalNodeOnly, out SeparateNode, out CrossCompile) && ProgramTarget.Rules.SupportsPlatform(HostPlatform) && SeparateNode && !BranchOptions.ExcludeNodes.Contains(ProgramTarget.TargetName))
-						{
-							if(bInternalNodeOnly)
-							{
-								AddNode(new SingleInternalToolsNode(HostPlatform, ProgramTarget));
-							}
-							else
-							{
-								AddNode(new SingleToolsNode(HostPlatform, ProgramTarget));
-							}
-						}
-						if(ProgramTarget.Rules.GUBP_IncludeNonUnityToolTest())
-						{
-							AddNode(new NonUnityToolNode(HostPlatform, ProgramTarget));
-						}
+                            if (ProgramTarget.Rules.GUBP_AlwaysBuildWithTools(HostPlatform, out bInternalNodeOnly, out SeparateNode, out CrossCompile) && ProgramTarget.Rules.SupportsPlatform(HostPlatform) && SeparateNode)
+                            {
+                                if (bInternalNodeOnly)
+                                {
+                                    AddNode(new SingleInternalToolsNode(HostPlatform, ProgramTarget));
+                                }
+                                else
+                                {
+                                    AddNode(new SingleToolsNode(HostPlatform, ProgramTarget));
+                                }
+                            }
+                            if (ProgramTarget.Rules.GUBP_IncludeNonUnityToolTest())
+                            {
+                                AddNode(new NonUnityToolNode(HostPlatform, ProgramTarget));
+                            }
+                        }
 					}
 				}
 
