Disable XGE for non-unity compiles.

diff --git a/Engine/Source/Programs/AutomationTool/GUBP/GUBP.Automation.cs b/Engine/Source/Programs/AutomationTool/GUBP/GUBP.Automation.cs
--- a/Engine/Source/Programs/AutomationTool/GUBP/GUBP.Automation.cs
+++ b/Engine/Source/Programs/AutomationTool/GUBP/GUBP.Automation.cs
@@ -3467,7 +3467,7 @@ public class GUBP : BuildCommand
 			Agenda.AddTargets(new string[] { "UnrealHeaderTool" }, HostPlatform, UnrealTargetConfiguration.Development);
 			Agenda.AddTargets(new string[] { ProgramTarget.TargetName }, HostPlatform, UnrealTargetConfiguration.Development, InAddArgs: "-skipnonhostplatforms");			
 
-			Build.Build(Agenda, InDeleteBuildProducts: true, InUpdateVersionFiles: false, InForceNonUnity: true);
+			Build.Build(Agenda, InDeleteBuildProducts: true, InUpdateVersionFiles: false, InForceNonUnity: true, InForceNoXGE: true);
 
 			UE4Build.CheckBuildProducts(Build.BuildProductFiles);
 			SaveRecordOfSuccessAndAddToBuildProducts();
@@ -3526,7 +3526,7 @@ public class GUBP : BuildCommand
                 }
             }
 
-            Build.Build(Agenda, InDeleteBuildProducts: true, InUpdateVersionFiles: false, InForceNonUnity: true);
+            Build.Build(Agenda, InDeleteBuildProducts: true, InUpdateVersionFiles: false, InForceNonUnity: true, InForceNoXGE: true);
 
             UE4Build.CheckBuildProducts(Build.BuildProductFiles);
             SaveRecordOfSuccessAndAddToBuildProducts();
