Separate temp storage into its own class, rather than being part of CommandUtils.

diff --git a/Engine/Source/Programs/AutomationTool/CommandUtils.cs b/Engine/Source/Programs/AutomationTool/CommandUtils.cs
--- a/Engine/Source/Programs/AutomationTool/CommandUtils.cs
+++ b/Engine/Source/Programs/AutomationTool/CommandUtils.cs
@@ -2009,38 +2009,6 @@ namespace AutomationTool
 			return false;
 		}
 
-        static Dictionary<string, string> ResolveCache = new Dictionary<string, string>();
-        public static string ResolveSharedBuildDirectory(string GameFolder)
-        {
-            if (ResolveCache.ContainsKey(GameFolder))
-            {
-                return ResolveCache[GameFolder];
-            }
-            string Root = RootSharedTempStorageDirectory();
-            string Result = CombinePaths(Root, GameFolder);
-            if (String.IsNullOrEmpty(GameFolder) || !Robust_DirectoryExistsAndIsWritable_NoExceptions(Result))
-            {
-                string GameStr = "Game";
-                bool HadGame = false;
-                if (GameFolder.EndsWith(GameStr, StringComparison.InvariantCultureIgnoreCase))
-                {
-                    string ShortFolder = GameFolder.Substring(0, GameFolder.Length - GameStr.Length);
-                    Result = CombinePaths(Root, ShortFolder);
-                    HadGame = true;
-                }
-                if (!HadGame || !Robust_DirectoryExistsAndIsWritable_NoExceptions(Result))
-                {
-                    Result = CombinePaths(Root, "UE4");
-                    if (!Robust_DirectoryExistsAndIsWritable_NoExceptions(Result))
-                    {
-                        throw new AutomationException("Could not find an appropriate shared temp folder {0}", Result);
-                    }
-                }
-            }
-            ResolveCache.Add(GameFolder, Result);
-            return Result;
-        }
-
         public static void CleanFormalBuilds(string DirectoryForThisBuild, string CLString = "", int MaximumDaysToKeepTempStorage = 4)
         {
             if (CLString == "" && (!IsBuildMachine || !DirectoryForThisBuild.StartsWith(RootSharedTempStorageDirectory()) || !P4Enabled))
diff --git a/Engine/Source/Programs/AutomationTool/GUBP/GUBP.Automation.cs b/Engine/Source/Programs/AutomationTool/GUBP/GUBP.Automation.cs
--- a/Engine/Source/Programs/AutomationTool/GUBP/GUBP.Automation.cs
+++ b/Engine/Source/Programs/AutomationTool/GUBP/GUBP.Automation.cs
@@ -3142,7 +3142,7 @@ public class GUBP : BuildCommand
 		}
         public static string GetArchiveDirectory(BranchInfo.BranchUProject InGameProj, UnrealTargetPlatform InHostPlatform, List<UnrealTargetPlatform> InClientTargetPlatforms = null, List<UnrealTargetConfiguration> InClientConfigs = null, List<UnrealTargetPlatform> InServerTargetPlatforms = null, List<UnrealTargetConfiguration> InServerConfigs = null, bool InClientNotGame = false)
         {
-            string BaseDir = ResolveSharedBuildDirectory(InGameProj.GameName);
+            string BaseDir = TempStorage.ResolveSharedBuildDirectory(InGameProj.GameName);
             string NodeName = StaticGetFullName(InGameProj, InHostPlatform, InClientTargetPlatforms, InClientConfigs, InServerTargetPlatforms, InServerConfigs, InClientNotGame);
             string Inner = P4Env.BuildRootEscaped + "-CL-" + P4Env.ChangelistString;
             if (GUBP.bPreflightBuild)
@@ -3829,7 +3829,7 @@ public class GUBP : BuildCommand
                 var StartTime = DateTime.UtcNow;
                 foreach (var NodeToDo in bp.GUBPNodes)
                 {
-                    CleanSharedTempStorageDirectory(NodeToDo.Value.GameNameIfAnyForTempStorage());
+                    TempStorage.CleanSharedTempStorageDirectory(NodeToDo.Value.GameNameIfAnyForTempStorage());
                 }
                 var BuildDuration = (DateTime.UtcNow - StartTime).TotalMilliseconds;
                 Log("Took {0}s to clear temp storage of old files.", BuildDuration / 1000);
@@ -4136,14 +4136,14 @@ public class GUBP : BuildCommand
         bool Result;
         if (LocalOnly)
         {
-            Result = LocalTempStorageExists(CmdEnv, NodeStoreName, bQuiet : true);
+            Result = TempStorage.LocalTempStorageExists(CmdEnv, NodeStoreName, bQuiet : true);
         }
         else
         {
-            Result = TempStorageExists(CmdEnv, NodeStoreName, GameNameIfAny, bQuiet: true);
+            Result = TempStorage.TempStorageExists(CmdEnv, NodeStoreName, GameNameIfAny, bQuiet: true);
 			if(GameNameIfAny != "" && Result == false)
 			{
-				Result = TempStorageExists(CmdEnv, NodeStoreName, "", bQuiet: true);
+				Result = TempStorage.TempStorageExists(CmdEnv, NodeStoreName, "", bQuiet: true);
 			}
         }
         if (Result)
@@ -4582,7 +4582,7 @@ public class GUBP : BuildCommand
         string RecordOfSuccess = CombinePaths(CommandUtils.CmdEnv.LocalRoot, "Engine", "Saved", "Logs", NodeToDo + Suffix +".log");
         CreateDirectory(Path.GetDirectoryName(RecordOfSuccess));
         WriteAllText(RecordOfSuccess, Contents);		
-		StoreToTempStorage(CmdEnv, NodeStoreName + Suffix, new List<string> { RecordOfSuccess }, !bSaveSharedTempStorage, GameNameIfAny);		
+		TempStorage.StoreToTempStorage(CmdEnv, NodeStoreName + Suffix, new List<string> { RecordOfSuccess }, !bSaveSharedTempStorage, GameNameIfAny);		
     }
 	string GetPropertyFromStep(string PropertyPath)
 	{
@@ -4886,9 +4886,9 @@ public class GUBP : BuildCommand
             string NodeStoreWildCard = StoreName.Replace(CLString, "*") + "-" + GUBPNodes[Node].GetFullName();
             var History = new NodeHistory();
 
-            History.AllStarted = ConvertCLToIntList(FindTempStorageManifests(CmdEnv, NodeStoreWildCard + StartedTempStorageSuffix, false, true, GameNameIfAny));
-            History.AllSucceeded = ConvertCLToIntList(FindTempStorageManifests(CmdEnv, NodeStoreWildCard + SucceededTempStorageSuffix, false, true, GameNameIfAny));
-            History.AllFailed = ConvertCLToIntList(FindTempStorageManifests(CmdEnv, NodeStoreWildCard + FailedTempStorageSuffix, false, true, GameNameIfAny));
+            History.AllStarted = ConvertCLToIntList(TempStorage.FindTempStorageManifests(CmdEnv, NodeStoreWildCard + StartedTempStorageSuffix, false, true, GameNameIfAny));
+            History.AllSucceeded = ConvertCLToIntList(TempStorage.FindTempStorageManifests(CmdEnv, NodeStoreWildCard + SucceededTempStorageSuffix, false, true, GameNameIfAny));
+            History.AllFailed = ConvertCLToIntList(TempStorage.FindTempStorageManifests(CmdEnv, NodeStoreWildCard + FailedTempStorageSuffix, false, true, GameNameIfAny));
 
             if (History.AllFailed.Count > 0)
             {
@@ -5073,13 +5073,13 @@ public class GUBP : BuildCommand
                 break;
             }
             var ThisNodeStore = NodeStore.Replace(CLString, String.Format("{0}", CL));
-            DeleteLocalTempStorage(CmdEnv, ThisNodeStore, true); // these all clash locally, which is fine we just retrieve them from shared
+            TempStorage.DeleteLocalTempStorage(CmdEnv, ThisNodeStore, true); // these all clash locally, which is fine we just retrieve them from shared
 
             List<string> Files = null;
             try
             {
                 bool WasLocal;
-                Files = RetrieveFromTempStorage(CmdEnv, ThisNodeStore, out WasLocal, GameNameIfAny); // this will fail on our CL if we didn't fail or we are just setting up the branch
+                Files = TempStorage.RetrieveFromTempStorage(CmdEnv, ThisNodeStore, out WasLocal, GameNameIfAny); // this will fail on our CL if we didn't fail or we are just setting up the branch
             }
             catch (Exception)
             {
@@ -5335,13 +5335,13 @@ public class GUBP : BuildCommand
         } 
         if (bSaveSharedTempStorage)
         {
-            if (!HaveSharedTempStorage(true))
+            if (!TempStorage.HaveSharedTempStorage(true))
             {
-                throw new AutomationException("Request to save to temp storage, but {0} is unavailable.", UE4TempStorageDirectory());
+                throw new AutomationException("Request to save to temp storage, but {0} is unavailable.", TempStorage.UE4TempStorageDirectory());
             }
             bSignBuildProducts = true;
         }
-        else if (!LocalOnly && !HaveSharedTempStorage(false))
+        else if (!LocalOnly && !TempStorage.HaveSharedTempStorage(false))
         {
             LogWarning("Looks like we want to use shared temp storage, but since we don't have it, we won't use it.");
             LocalOnly = true;
@@ -6357,7 +6357,7 @@ public class GUBP : BuildCommand
 
         if (bCleanLocalTempStorage)  // shared temp storage can never be wiped
         {
-            DeleteLocalTempStorageManifests(CmdEnv);
+            TempStorage.DeleteLocalTempStorageManifests(CmdEnv);
         }
 
         GUBPNodesControllingTrigger = new Dictionary<string, string>();
@@ -7540,13 +7540,13 @@ public class GUBP : BuildCommand
                     bool WasLocal;
 					try
 					{
-						GUBPNodes[NodeToDo].BuildProducts = RetrieveFromTempStorage(CmdEnv, NodeStoreName, out WasLocal, GameNameIfAny, StorageRootIfAny);
+						GUBPNodes[NodeToDo].BuildProducts = TempStorage.RetrieveFromTempStorage(CmdEnv, NodeStoreName, out WasLocal, GameNameIfAny, StorageRootIfAny);
 					}
 					catch
 					{
 						if(GameNameIfAny != "")
 						{
-							GUBPNodes[NodeToDo].BuildProducts = RetrieveFromTempStorage(CmdEnv, NodeStoreName, out WasLocal, "", StorageRootIfAny);
+							GUBPNodes[NodeToDo].BuildProducts = TempStorage.RetrieveFromTempStorage(CmdEnv, NodeStoreName, out WasLocal, "", StorageRootIfAny);
 						}
 						else
 						{
@@ -7596,7 +7596,7 @@ public class GUBP : BuildCommand
 						var StoreDuration = 0.0;
 						var StartTime = DateTime.UtcNow;
                         var StartStore = DateTime.Now.ToString();
-                        StoreToTempStorage(CmdEnv, NodeStoreName, GUBPNodes[NodeToDo].BuildProducts, !bSaveSharedTempStorage, GameNameIfAny, StorageRootIfAny);
+                        TempStorage.StoreToTempStorage(CmdEnv, NodeStoreName, GUBPNodes[NodeToDo].BuildProducts, !bSaveSharedTempStorage, GameNameIfAny, StorageRootIfAny);
 						StoreDuration = (DateTime.UtcNow - StartTime).TotalMilliseconds / 1000;
                         var FinishStore = DateTime.Now.ToString();                        
 						Log("Took {0} seconds to store build products", StoreDuration);
@@ -7613,7 +7613,7 @@ public class GUBP : BuildCommand
                                 {
                                     bool WasLocal;
                                     var StartRetrieve = DateTime.Now.ToString();
-									RetrieveFromTempStorage(CmdEnv, NodeStoreName, out WasLocal, GameNameIfAny, StorageRootIfAny);
+									TempStorage.RetrieveFromTempStorage(CmdEnv, NodeStoreName, out WasLocal, GameNameIfAny, StorageRootIfAny);
                                     var FinishRetrieve = DateTime.Now.ToString();
                                     if (IsBuildMachine)
                                     {
diff --git a/Engine/Source/Programs/AutomationTool/Scripts/RocketBuild.Automation.cs b/Engine/Source/Programs/AutomationTool/Scripts/RocketBuild.Automation.cs
--- a/Engine/Source/Programs/AutomationTool/Scripts/RocketBuild.Automation.cs
+++ b/Engine/Source/Programs/AutomationTool/Scripts/RocketBuild.Automation.cs
@@ -853,8 +853,8 @@ namespace Rocket
 
 			// Copy the files to their final location
 			CommandUtils.Log("Copying files to {0}", PublishDir);
-			CommandUtils.Robust_CopyFile(FullZipFileName, Path.Combine(PublishDir, Path.GetFileName(FullZipFileName)));
-			CommandUtils.Robust_CopyFile(EditorZipFileName, Path.Combine(PublishDir, Path.GetFileName(EditorZipFileName)));
+			TempStorage.Robust_CopyFile(FullZipFileName, Path.Combine(PublishDir, Path.GetFileName(FullZipFileName)));
+			TempStorage.Robust_CopyFile(EditorZipFileName, Path.Combine(PublishDir, Path.GetFileName(EditorZipFileName)));
 			CommandUtils.DeleteFile(FullZipFileName);
 			CommandUtils.DeleteFile(EditorZipFileName);
 
@@ -911,7 +911,7 @@ namespace Rocket
 					if(DebugExtensions.Contains(Extension) || Extension == ".exe" || Extension == ".dll") // Need all windows build products for crash reporter
 					{
 						string OutputFileName = CommandUtils.MakeRerootedFilePath(InputFileName, CommandUtils.CmdEnv.LocalRoot, SymbolsOutputDir);
-						CommandUtils.Robust_CopyFile(InputFileName, OutputFileName);
+						TempStorage.Robust_CopyFile(InputFileName, OutputFileName);
 					}
 				}
 			}
diff --git a/Engine/Source/Programs/AutomationTool/Scripts/Tests.Automation.cs b/Engine/Source/Programs/AutomationTool/Scripts/Tests.Automation.cs
--- a/Engine/Source/Programs/AutomationTool/Scripts/Tests.Automation.cs
+++ b/Engine/Source/Programs/AutomationTool/Scripts/Tests.Automation.cs
@@ -476,266 +476,6 @@ class TestMacZip : BuildCommand
 	}
 }
 
-[Help("Tests the temp storage operations.")]
-class TestTempStorage : BuildCommand
-{
-	public override void ExecuteBuild()
-	{
-		Log("TestTempStorage********");
-
-		DeleteLocalTempStorageManifests(CmdEnv);
-		DeleteSharedTempStorageManifests(CmdEnv, "Test");
-		if (TempStorageExists(CmdEnv, "Test"))
-		{
-			throw new AutomationException("storage should not exist");
-		}
-
-		string UnitTestFile = CombinePaths(CmdEnv.LocalRoot, "Engine", "Build", "Batchfiles", "TestFile.Txt");
-		Log("Test file {0}", UnitTestFile);
-
-		if (FileExists(UnitTestFile))
-		{
-			DeleteFile(UnitTestFile);
-		}
-		string[] LinesToWrite = new string[]
-		{
-			"This is a test",
-			"Of writing to file " + UnitTestFile
-		};
-		foreach (string Line in LinesToWrite)
-		{
-			WriteToFile(UnitTestFile, Line);
-		}
-
-		string UnitTestFile2 = CombinePaths(CmdEnv.LocalRoot, "TestFile2.Txt");
-		Log("Test file {0}", UnitTestFile2);
-
-		if (FileExists(UnitTestFile2))
-		{
-			DeleteFile(UnitTestFile2);
-		}
-		string[] LinesToWrite2 = new string[]
-		{
-			"This is a test",
-			"Of writing to file " + UnitTestFile2
-		};
-		foreach (string Line in LinesToWrite2)
-		{
-			WriteToFile(UnitTestFile2, Line);
-		}
-
-		string UnitTestFile3 = CombinePaths(CmdEnv.LocalRoot, "engine", "plugins", "TestFile3.Txt");
-		Log("Test file {0}", UnitTestFile3);
-
-		if (FileExists(UnitTestFile3))
-		{
-			DeleteFile(UnitTestFile3);
-		}
-		string[] LinesToWrite3 = new string[]
-		{
-			"This is a test",
-			"Of writing to file " + UnitTestFile3
-		};
-		foreach (string Line in LinesToWrite3)
-		{
-			WriteToFile(UnitTestFile3, Line);
-		}
-
-		{
-			string[] LinesRead = ReadAllLines(UnitTestFile);
-			if (LinesRead == null || LinesRead.Length != LinesToWrite.Length)
-			{
-				throw new AutomationException("Contents of the file created is different to the file read.");
-			}
-			for (int LineIndex = 0; LineIndex < LinesRead.Length; ++LineIndex)
-			{
-				if (LinesRead[LineIndex] != LinesToWrite[LineIndex])
-				{
-					throw new AutomationException("Contents of the file created is different to the file read.");
-				}
-			}
-
-			string[] LinesRead2 = ReadAllLines(UnitTestFile2);
-			if (LinesRead2 == null || LinesRead2.Length != LinesToWrite2.Length)
-			{
-				throw new AutomationException("Contents of the file created is different to the file read.");
-			}
-			for (int LineIndex = 0; LineIndex < LinesRead2.Length; ++LineIndex)
-			{
-				if (LinesRead2[LineIndex] != LinesToWrite2[LineIndex])
-				{
-					throw new AutomationException("Contents of the file created is different to the file read.");
-				}
-			}
-
-			string[] LinesRead3 = ReadAllLines(UnitTestFile3);
-			if (LinesRead3 == null || LinesRead3.Length != LinesToWrite3.Length)
-			{
-				throw new AutomationException("Contents of the file created is different to the file read.");
-			}
-			for (int LineIndex = 0; LineIndex < LinesRead3.Length; ++LineIndex)
-			{
-				if (LinesRead3[LineIndex] != LinesToWrite3[LineIndex])
-				{
-					throw new AutomationException("Contents of the file created is different to the file read.");
-				}
-			}
-		}
-		var TestFiles = new List<string> { UnitTestFile, UnitTestFile2, UnitTestFile3 };
-
-		StoreToTempStorage(CmdEnv, "Test", TestFiles);
-
-		if (!LocalTempStorageExists(CmdEnv, "Test"))
-		{
-			throw new AutomationException("local storage should exist");
-		}
-		if (!SharedTempStorageExists(CmdEnv, "Test"))
-		{
-			throw new AutomationException("shared storage should exist");
-		}
-		DeleteLocalTempStorageManifests(CmdEnv);
-		if (LocalTempStorageExists(CmdEnv, "Test"))
-		{
-			throw new AutomationException("local storage should not exist");
-		}
-		if (!TempStorageExists(CmdEnv, "Test"))
-		{
-			throw new AutomationException("some storage should exist");
-		}
-		DeleteFile(UnitTestFile);
-		DeleteFile(UnitTestFile2);
-		DeleteFile(UnitTestFile3);
-
-		bool WasLocal;
-		RetrieveFromTempStorage(CmdEnv, "Test", out WasLocal);
-		if (!LocalTempStorageExists(CmdEnv, "Test"))
-		{
-			throw new AutomationException("local storage should exist");
-		}
-		if (!SharedTempStorageExists(CmdEnv, "Test"))
-		{
-			throw new AutomationException("shared storage should exist");
-		}
-		{
-			string[] LinesRead = ReadAllLines(UnitTestFile);
-			if (LinesRead == null || LinesRead.Length != LinesToWrite.Length)
-			{
-				throw new AutomationException("Contents of the file created is different to the file read.");
-			}
-			for (int LineIndex = 0; LineIndex < LinesRead.Length; ++LineIndex)
-			{
-				if (LinesRead[LineIndex] != LinesToWrite[LineIndex])
-				{
-					throw new AutomationException("Contents of the file created is different to the file read.");
-				}
-			}
-
-			string[] LinesRead2 = ReadAllLines(UnitTestFile2);
-			if (LinesRead2 == null || LinesRead2.Length != LinesToWrite2.Length)
-			{
-				throw new AutomationException("Contents of the file created is different to the file read.");
-			}
-			for (int LineIndex = 0; LineIndex < LinesRead2.Length; ++LineIndex)
-			{
-				if (LinesRead2[LineIndex] != LinesToWrite2[LineIndex])
-				{
-					throw new AutomationException("Contents of the file created is different to the file read.");
-				}
-			}
-
-			string[] LinesRead3 = ReadAllLines(UnitTestFile3);
-			if (LinesRead3 == null || LinesRead3.Length != LinesToWrite3.Length)
-			{
-				throw new AutomationException("Contents of the file created is different to the file read.");
-			}
-			for (int LineIndex = 0; LineIndex < LinesRead3.Length; ++LineIndex)
-			{
-				if (LinesRead3[LineIndex] != LinesToWrite3[LineIndex])
-				{
-					throw new AutomationException("Contents of the file created is different to the file read.");
-				}
-			}
-		}
-		DeleteSharedTempStorageManifests(CmdEnv, "Test");
-		if (SharedTempStorageExists(CmdEnv, "Test"))
-		{
-			throw new AutomationException("shared storage should not exist");
-		}
-		RetrieveFromTempStorage(CmdEnv, "Test", out WasLocal); // this should just rely on the local
-		if (!WasLocal || !LocalTempStorageExists(CmdEnv, "Test"))
-		{
-			throw new AutomationException("local storage should exist");
-		}
-
-		// and now lets test tampering
-		DeleteLocalTempStorageManifests(CmdEnv);
-		{
-			bool bFailedProperly = false;
-			var MissingFile = new List<string>(TestFiles);
-			MissingFile.Add(CombinePaths(CmdEnv.LocalRoot, "Engine", "SomeFileThatDoesntExist.txt"));
-			try
-			{
-				StoreToTempStorage(CmdEnv, "Test", MissingFile);
-			}
-			catch (AutomationException)
-			{
-				bFailedProperly = true;
-			}
-			if (!bFailedProperly)
-			{
-				throw new AutomationException("Missing file did not fail.");
-			}
-		}
-		DeleteSharedTempStorageManifests(CmdEnv, "Test"); // this ends up being junk
-		StoreToTempStorage(CmdEnv, "Test", TestFiles);
-		DeleteLocalTempStorageManifests(CmdEnv); // force a load from shared
-		var SharedFile = CombinePaths(SharedTempStorageDirectory("Test"), "Engine", "Build", "Batchfiles", "TestFile.Txt");
-		if (!FileExists_NoExceptions(SharedFile))
-		{
-			throw new AutomationException("Shared file {0} did not exist", SharedFile);
-		}
-		DeleteFile(SharedFile);
-		{
-			bool bFailedProperly = false;
-			try
-			{
-				RetrieveFromTempStorage(CmdEnv, "Test", out WasLocal);
-			}
-			catch (AutomationException)
-			{
-				bFailedProperly = true;
-			}
-			if (!bFailedProperly)
-			{
-				throw new AutomationException("Did not fail to load from missing file.");
-			}
-		}
-		DeleteSharedTempStorageManifests(CmdEnv, "Test");
-		StoreToTempStorage(CmdEnv, "Test", TestFiles);
-		DeleteFile(UnitTestFile);
-		{
-			bool bFailedProperly = false;
-			try
-			{
-				RetrieveFromTempStorage(CmdEnv, "Test", out WasLocal);
-			}
-			catch (AutomationException)
-			{
-				bFailedProperly = true;
-			}
-			if (!bFailedProperly)
-			{
-				throw new AutomationException("Did not fail to load from missing local file.");
-			}
-		}
-		DeleteSharedTempStorageManifests(CmdEnv, "Test");
-		DeleteLocalTempStorageManifests(CmdEnv);
-		DeleteFile(UnitTestFile);
-		DeleteFile(UnitTestFile2);
-		DeleteFile(UnitTestFile3);
-	}
-}
-
 [Help("Reads the build time from build.properties")]
 class TestBuildTime : BuildCommand
 {
diff --git a/Engine/Source/Programs/AutomationTool/TempStorage.cs b/Engine/Source/Programs/AutomationTool/TempStorage.cs
--- a/Engine/Source/Programs/AutomationTool/TempStorage.cs
+++ b/Engine/Source/Programs/AutomationTool/TempStorage.cs
@@ -9,7 +9,7 @@ using System.Xml;
 namespace AutomationTool
 {
 
-    public partial class CommandUtils
+    public class TempStorage
     {
         public class TempStorageFileInfo
         {
@@ -34,7 +34,7 @@ namespace AutomationTool
                 bool bOk = true;
                 if (!Name.Equals(Other.Name, StringComparison.InvariantCultureIgnoreCase))
                 {
-                    Log(System.Diagnostics.TraceEventType.Error, "File name mismatch {0} {1}", Name, Other.Name);
+                    CommandUtils.Log(System.Diagnostics.TraceEventType.Error, "File name mismatch {0} {1}", Name, Other.Name);
                     bOk = false;
                 }
                 else
@@ -123,7 +123,7 @@ namespace AutomationTool
         }
         public static void Robust_FileExists_NoExceptions(bool bQuiet, string Filename, string Message)
         {
-            if (!FileExists_NoExceptions(bQuiet, Filename))
+            if (!CommandUtils.FileExists_NoExceptions(bQuiet, Filename))
             {
                 bool bFound = false;
                 // mac is terrible on shares, this isn't a solution, but a stop gap
@@ -134,7 +134,7 @@ namespace AutomationTool
                     {
                         CommandUtils.Log(System.Diagnostics.TraceEventType.Information, "*** Mac temp storage retry {0}", Filename);
                         System.Threading.Thread.Sleep(10000);
-                        bFound = FileExists_NoExceptions(bQuiet, Filename);
+                        bFound = CommandUtils.FileExists_NoExceptions(bQuiet, Filename);
 						Retry++;
                     }
                 }
@@ -148,7 +148,7 @@ namespace AutomationTool
 		public static bool Robust_DirectoryExists_NoExceptions(string Directoryname, string Message)
 		{
 			bool bFound = false;
-			if (!DirectoryExists_NoExceptions(Directoryname))
+			if (!CommandUtils.DirectoryExists_NoExceptions(Directoryname))
 			{				
 				// mac is terrible on shares, this isn't a solution, but a stop gap
 				if (Directoryname.StartsWith("/Volumes/"))
@@ -158,7 +158,7 @@ namespace AutomationTool
 					{
 						CommandUtils.Log(System.Diagnostics.TraceEventType.Information, "*** Mac temp storage retry {0}", Directoryname);
 						System.Threading.Thread.Sleep(10000);
-						bFound = DirectoryExists_NoExceptions(Directoryname);
+						bFound = CommandUtils.DirectoryExists_NoExceptions(Directoryname);
 						Retry++;
 					}
 				}				
@@ -172,7 +172,7 @@ namespace AutomationTool
         public static bool Robust_DirectoryExistsAndIsWritable_NoExceptions(string Directoryname)
         {
             bool bFound = false;
-            if (!DirectoryExistsAndIsWritable_NoExceptions(Directoryname))
+            if (!CommandUtils.DirectoryExistsAndIsWritable_NoExceptions(Directoryname))
             {
                 // mac is terrible on shares, this isn't a solution, but a stop gap
                 if (Directoryname.StartsWith("/Volumes/"))
@@ -187,7 +187,7 @@ namespace AutomationTool
                     {
                         CommandUtils.Log("*** Mac temp storage retry {0}", Directoryname);
                         System.Threading.Thread.Sleep(1000);
-                        bFound = DirectoryExistsAndIsWritable_NoExceptions(Directoryname);
+                        bFound = CommandUtils.DirectoryExistsAndIsWritable_NoExceptions(Directoryname);
                         Retry++;
                     }
                 }
@@ -212,7 +212,7 @@ namespace AutomationTool
 
             public void Create(List<string> InFiles, string BaseFolder)
             {
-                BaseFolder = CombinePaths(BaseFolder, "/");
+                BaseFolder = CommandUtils.CombinePaths(BaseFolder, "/");
                 if (!BaseFolder.EndsWith("/")&& !BaseFolder.EndsWith("\\"))
                 {
                     throw new AutomationException("base folder {0} should end with a separator", BaseFolder);
@@ -220,7 +220,7 @@ namespace AutomationTool
 
                 foreach (string InFilename in InFiles)
                 {
-                    var Filename = CombinePaths(InFilename);
+                    var Filename = CommandUtils.CombinePaths(InFilename);
                     Robust_FileExists_NoExceptions(true, Filename, "Could not add {0} to manifest because it does not exist");
 
                     FileInfo Info = new FileInfo(Filename);
@@ -233,7 +233,7 @@ namespace AutomationTool
                     }
                     var RelativeFile = Filename.Substring(BaseFolder.Length);
                     List<TempStorageFileInfo> ManifestDirectory;
-                    string DirectoryName = CombinePaths(Path.GetDirectoryName(Filename), "/");
+                    string DirectoryName = CommandUtils.CombinePaths(Path.GetDirectoryName(Filename), "/");
                     if (!DirectoryName.StartsWith(BaseFolder, StringComparison.InvariantCultureIgnoreCase))
                     {
                         throw new AutomationException("Could not add directory {0} to manifest because it does not start with the base folder {1}", DirectoryName, BaseFolder);
@@ -254,20 +254,20 @@ namespace AutomationTool
             public void Stats(string Description)
             {
 
-                Log("{0}: Directories={1} Files={2} Size={3}", Description, Directories.Count, GetFileCount(), GetTotalSize());
+                CommandUtils.Log("{0}: Directories={1} Files={2} Size={3}", Description, Directories.Count, GetFileCount(), GetTotalSize());
             }
             
             public bool Compare(TempStorageManifest Other)
             {
                 if (Directories.Count != Other.Directories.Count)
                 {
-                    Log(System.Diagnostics.TraceEventType.Error, "Directory count mismatch {0} {1}", Directories.Count, Other.Directories.Count);
+                    CommandUtils.Log(System.Diagnostics.TraceEventType.Error, "Directory count mismatch {0} {1}", Directories.Count, Other.Directories.Count);
                     foreach (KeyValuePair<string, List<TempStorageFileInfo>> Directory in Directories)
                     {
                         List<TempStorageFileInfo> OtherDirectory;
                         if (Other.Directories.TryGetValue(Directory.Key, out OtherDirectory) == false)
                         {
-                            Log(System.Diagnostics.TraceEventType.Error, "Missing Directory {0}", Directory.Key);
+                            CommandUtils.Log(System.Diagnostics.TraceEventType.Error, "Missing Directory {0}", Directory.Key);
                             return false;
                         }
                     }
@@ -276,7 +276,7 @@ namespace AutomationTool
                         List<TempStorageFileInfo> OtherDirectory;
                         if (Directories.TryGetValue(Directory.Key, out OtherDirectory) == false)
                         {
-                            Log(System.Diagnostics.TraceEventType.Error, "Missing Other Directory {0}", Directory.Key);
+                            CommandUtils.Log(System.Diagnostics.TraceEventType.Error, "Missing Other Directory {0}", Directory.Key);
                             return false;
                         }
                     }
@@ -288,19 +288,19 @@ namespace AutomationTool
                     List<TempStorageFileInfo> OtherDirectory;
                     if (Other.Directories.TryGetValue(Directory.Key, out OtherDirectory) == false)
                     {
-                        Log(System.Diagnostics.TraceEventType.Error, "Missing Directory {0}", Directory.Key); 
+                        CommandUtils.Log(System.Diagnostics.TraceEventType.Error, "Missing Directory {0}", Directory.Key); 
                         return false;
                     }
                     if (OtherDirectory.Count != Directory.Value.Count)
                     {
-                        Log(System.Diagnostics.TraceEventType.Error, "File count mismatch {0} {1} {2}", Directory.Key, OtherDirectory.Count, Directory.Value.Count);
+                        CommandUtils.Log(System.Diagnostics.TraceEventType.Error, "File count mismatch {0} {1} {2}", Directory.Key, OtherDirectory.Count, Directory.Value.Count);
                         for (int FileIndex = 0; FileIndex < Directory.Value.Count; ++FileIndex)
                         {
-                            Log("Manifest1: {0}", Directory.Value[FileIndex].Name);
+                            CommandUtils.Log("Manifest1: {0}", Directory.Value[FileIndex].Name);
                         }
                         for (int FileIndex = 0; FileIndex < OtherDirectory.Count; ++FileIndex)
                         {
-                            Log("Manifest2: {0}", OtherDirectory[FileIndex].Name);
+                            CommandUtils.Log("Manifest2: {0}", OtherDirectory[FileIndex].Name);
                         }
                         return false;
                     }
@@ -322,18 +322,18 @@ namespace AutomationTool
 
             public TempStorageFileInfo FindFile(string LocalPath, string BaseFolder)
             {
-                BaseFolder = CombinePaths(BaseFolder, "/");
+                BaseFolder = CommandUtils.CombinePaths(BaseFolder, "/");
                 if (!BaseFolder.EndsWith("/")&& !BaseFolder.EndsWith("\\"))
                 {
                     throw new AutomationException("base folder {0} should end with a separator", BaseFolder);
                 }
-                var Filename = CombinePaths(LocalPath);
+                var Filename = CommandUtils.CombinePaths(LocalPath);
                 if (!Filename.StartsWith(BaseFolder, StringComparison.InvariantCultureIgnoreCase))
                 {
                     throw new AutomationException("Could not add {0} to manifest because it does not start with the base folder {1}", Filename, BaseFolder);
                 }
                 var RelativeFile = Filename.Substring(BaseFolder.Length);
-                string DirectoryName = CombinePaths(Path.GetDirectoryName(Filename), "/");
+                string DirectoryName = CommandUtils.CombinePaths(Path.GetDirectoryName(Filename), "/");
                 if (!DirectoryName.StartsWith(BaseFolder, StringComparison.InvariantCultureIgnoreCase))
                 {
                     throw new AutomationException("Could not add directory {0} to manifest because it does not start with the base folder {1}", DirectoryName, BaseFolder);
@@ -553,7 +553,7 @@ namespace AutomationTool
                 {
                     foreach (var ThisFileInfo in Dir.Value)
                     {
-                        var NewFile = CombinePaths(NewBaseDir, ThisFileInfo.Name);
+                        var NewFile = CommandUtils.CombinePaths(NewBaseDir, ThisFileInfo.Name);
 
                         FileInfo Info = new FileInfo(NewFile);
 
@@ -565,7 +565,7 @@ namespace AutomationTool
                     throw new AutomationException("No files in attempt to GetFiles().");
                 }
 
-				Log("Checking that first file in manifest exists...");
+				CommandUtils.Log("Checking that first file in manifest exists...");
 				Robust_FileExists_NoExceptions(false, Result[0], "Rebased manifest file does not exist {0}");
 
                 return Result;
@@ -574,7 +574,7 @@ namespace AutomationTool
 
         public static string LocalTempStorageManifestDirectory(CommandEnvironment Env)
         {
-            return CombinePaths(Env.LocalRoot, "Engine", "Saved", "GUBP");
+            return CommandUtils.CombinePaths(Env.LocalRoot, "Engine", "Saved", "GUBP");
         }
         static HashSet<string> TopDirectoryTestedForClean = new HashSet<string>();
         static void CleanSharedTempStorage(string TopDirectory)
@@ -591,11 +591,11 @@ namespace AutomationTool
             var TopLevelDirs = DirInfo.GetDirectories();
             {
                 var BuildDuration = (DateTime.UtcNow - StartTimeDir).TotalMilliseconds;
-                Log("Took {0}s to enumerate {1} directories.", BuildDuration / 1000, TopLevelDirs.Length);
+                CommandUtils.Log("Took {0}s to enumerate {1} directories.", BuildDuration / 1000, TopLevelDirs.Length);
             }
             foreach (var TopLevelDir in TopLevelDirs)
             {
-                if (DirectoryExists_NoExceptions(TopLevelDir.FullName))
+                if (CommandUtils.DirectoryExists_NoExceptions(TopLevelDir.FullName))
                 {
                     bool bOld = false;
                     foreach (var ThisFile in CommandUtils.FindFiles_NoExceptions(true, "*.TempManifest", false, TopLevelDir.FullName))
@@ -609,7 +609,7 @@ namespace AutomationTool
                     }
                     if (bOld)
                     {
-                        Log("Deleteing temp storage directory {0}, because it is more than {1} days old.", TopLevelDir.FullName, MaximumDaysToKeepTempStorage);
+                        CommandUtils.Log("Deleteing temp storage directory {0}, because it is more than {1} days old.", TopLevelDir.FullName, MaximumDaysToKeepTempStorage);
                         var StartTime = DateTime.UtcNow;
                         try
                         {
@@ -622,9 +622,9 @@ namespace AutomationTool
                         catch (Exception)
                         {
                         }
-                        DeleteDirectory_NoExceptions(true, TopLevelDir.FullName);
+                        CommandUtils.DeleteDirectory_NoExceptions(true, TopLevelDir.FullName);
                         var BuildDuration = (DateTime.UtcNow - StartTime).TotalMilliseconds;
-                        Log("Took {0}s to delete {1}.", BuildDuration / 1000, TopLevelDir.FullName);
+                        CommandUtils.Log("Took {0}s to delete {1}.", BuildDuration / 1000, TopLevelDir.FullName);
                     }
                 }
             }
@@ -637,8 +637,8 @@ namespace AutomationTool
         public static string UE4TempStorageDirectory()
         {
             string SharedSubdir = SharedTempStorageDirectory();
-            string Root = RootSharedTempStorageDirectory();
-            return CombinePaths(Root, "UE4", SharedSubdir);
+            string Root = CommandUtils.RootSharedTempStorageDirectory();
+            return CommandUtils.CombinePaths(Root, "UE4", SharedSubdir);
         }
         public static bool HaveSharedTempStorage(bool ForSaving)
         {
@@ -652,7 +652,7 @@ namespace AutomationTool
                     {
                         return true;
                     }
-                    FindDirectories_NoExceptions(false, "*", false, Dir); // there is some internet evidence to suggest this might perk up the mac share
+                    CommandUtils.FindDirectories_NoExceptions(false, "*", false, Dir); // there is some internet evidence to suggest this might perk up the mac share
                     System.Threading.Thread.Sleep(5000);
 					Retries++;
                 }
@@ -664,14 +664,46 @@ namespace AutomationTool
             return false;
         }
 
+        static Dictionary<string, string> ResolveCache = new Dictionary<string, string>();
+        public static string ResolveSharedBuildDirectory(string GameFolder)
+        {
+            if (ResolveCache.ContainsKey(GameFolder))
+            {
+                return ResolveCache[GameFolder];
+            }
+            string Root = CommandUtils.RootSharedTempStorageDirectory();
+            string Result = CommandUtils.CombinePaths(Root, GameFolder);
+            if (String.IsNullOrEmpty(GameFolder) || !Robust_DirectoryExistsAndIsWritable_NoExceptions(Result))
+            {
+                string GameStr = "Game";
+                bool HadGame = false;
+                if (GameFolder.EndsWith(GameStr, StringComparison.InvariantCultureIgnoreCase))
+                {
+                    string ShortFolder = GameFolder.Substring(0, GameFolder.Length - GameStr.Length);
+                    Result = CommandUtils.CombinePaths(Root, ShortFolder);
+                    HadGame = true;
+                }
+                if (!HadGame || !Robust_DirectoryExistsAndIsWritable_NoExceptions(Result))
+                {
+                    Result = CommandUtils.CombinePaths(Root, "UE4");
+                    if (!Robust_DirectoryExistsAndIsWritable_NoExceptions(Result))
+                    {
+                        throw new AutomationException("Could not find an appropriate shared temp folder {0}", Result);
+                    }
+                }
+            }
+            ResolveCache.Add(GameFolder, Result);
+            return Result;
+        }
+
         public static string ResolveSharedTempStorageDirectory(string GameFolder)
         {
             string SharedSubdir = SharedTempStorageDirectory();
-            string Result = CombinePaths(ResolveSharedBuildDirectory(GameFolder), SharedSubdir);
+            string Result = CommandUtils.CombinePaths(ResolveSharedBuildDirectory(GameFolder), SharedSubdir);
 
             if (!Robust_DirectoryExists_NoExceptions(Result, "Could not find {0}"))
             {
-                CreateDirectory_NoExceptions(Result);
+                CommandUtils.CreateDirectory_NoExceptions(Result);
             }
             if (!Robust_DirectoryExists_NoExceptions(Result, "Could not find {0}"))
             {
@@ -684,7 +716,7 @@ namespace AutomationTool
         public static void CleanSharedTempStorageDirectory(string GameFolder)
         {
 
-            if (!IsBuildMachine || TestedForClean.Contains(GameFolder) ||
+            if (!CommandUtils.IsBuildMachine || TestedForClean.Contains(GameFolder) ||
                 UnrealBuildTool.Utils.IsRunningOnMono)  // saw a hang on this, anyway it isn't necessary to clean with macs, they are slow anyway
             {
                 return;
@@ -696,13 +728,13 @@ namespace AutomationTool
             }
             catch (Exception Ex)
             {
-                Log(System.Diagnostics.TraceEventType.Warning, "Unable to Clean Directory with GameName {0}", GameFolder);
-                Log(System.Diagnostics.TraceEventType.Warning, " Exception was {0}", LogUtils.FormatException(Ex));
+                CommandUtils.Log(System.Diagnostics.TraceEventType.Warning, "Unable to Clean Directory with GameName {0}", GameFolder);
+                CommandUtils.Log(System.Diagnostics.TraceEventType.Warning, " Exception was {0}", LogUtils.FormatException(Ex));
             }
         }
         public static string SharedTempStorageDirectory(string StorageBlock, string GameFolder = "")
         {
-            return CombinePaths(ResolveSharedTempStorageDirectory(GameFolder), StorageBlock);
+            return CommandUtils.CombinePaths(ResolveSharedTempStorageDirectory(GameFolder), StorageBlock);
         }
 
         public static TempStorageManifest SaveTempStorageManifest(string RootDir, string FinalFilename, List<string> Files)
@@ -714,11 +746,11 @@ namespace AutomationTool
                 throw new AutomationException("Saver manifest differs has wrong number of files {0} != {1}", Saver.GetFileCount(), Files.Count);
             }
             var TempFilename = FinalFilename + ".temp";
-            if (FileExists_NoExceptions(true, TempFilename))
+            if (CommandUtils.FileExists_NoExceptions(true, TempFilename))
             {
                 throw new AutomationException("Temp manifest file already exists {0}", TempFilename);
             }
-            CreateDirectory(true, Path.GetDirectoryName(FinalFilename));
+            CommandUtils.CreateDirectory(true, Path.GetDirectoryName(FinalFilename));
             Saver.Save(TempFilename);
 
             var Tester = new TempStorageManifest();
@@ -729,8 +761,8 @@ namespace AutomationTool
                 throw new AutomationException("Temp manifest differs {0}", TempFilename);
             }
 
-            RenameFile(TempFilename, FinalFilename, true);
-            if (FileExists_NoExceptions(true, TempFilename))
+            CommandUtils.RenameFile(TempFilename, FinalFilename, true);
+            if (CommandUtils.FileExists_NoExceptions(true, TempFilename))
             {
                 throw new AutomationException("Temp manifest didn't go away {0}", TempFilename);
             }
@@ -741,12 +773,12 @@ namespace AutomationTool
             {
                 throw new AutomationException("Final manifest differs {0}", TempFilename);
             }
-            Log("Saved {0} with {1} files and total size {2}", FinalFilename, Saver.GetFileCount(), Saver.GetTotalSize());
+            CommandUtils.Log("Saved {0} with {1} files and total size {2}", FinalFilename, Saver.GetFileCount(), Saver.GetTotalSize());
             return Saver;
         }
         public static string LocalTempStorageManifestFilename(CommandEnvironment Env, string StorageBlockName)
         {
-            return CombinePaths(LocalTempStorageManifestDirectory(Env), StorageBlockName + ".TempManifest");
+            return CommandUtils.CombinePaths(LocalTempStorageManifestDirectory(Env), StorageBlockName + ".TempManifest");
         }
         public static TempStorageManifest SaveLocalTempStorageManifest(CommandEnvironment Env, string BaseFolder, string StorageBlockName, List<string> Files)
         {
@@ -754,7 +786,7 @@ namespace AutomationTool
         }
         public static string SharedTempStorageManifestFilename(CommandEnvironment Env, string StorageBlockName, string GameFolder)
         {
-            return CombinePaths(SharedTempStorageDirectory(StorageBlockName, GameFolder), StorageBlockName + ".TempManifest");
+            return CommandUtils.CombinePaths(SharedTempStorageDirectory(StorageBlockName, GameFolder), StorageBlockName + ".TempManifest");
         }
         public static TempStorageManifest SaveSharedTempStorageManifest(CommandEnvironment Env, string StorageBlockName, string GameFolder, List<string> Files)
         {
@@ -762,16 +794,16 @@ namespace AutomationTool
         }
         public static void DeleteLocalTempStorageManifests(CommandEnvironment Env)
         {
-            DeleteDirectory(true, LocalTempStorageManifestDirectory(Env));
+            CommandUtils.DeleteDirectory(true, LocalTempStorageManifestDirectory(Env));
         }
         public static void DeleteSharedTempStorageManifests(CommandEnvironment Env, string StorageBlockName, string GameFolder = "")
         {
-            DeleteDirectory(true, SharedTempStorageDirectory(StorageBlockName, GameFolder));
+            CommandUtils.DeleteDirectory(true, SharedTempStorageDirectory(StorageBlockName, GameFolder));
         }
         public static bool LocalTempStorageExists(CommandEnvironment Env, string StorageBlockName, bool bQuiet = false)
         {
             var LocalManifest = LocalTempStorageManifestFilename(Env, StorageBlockName);
-            if (FileExists_NoExceptions(bQuiet, LocalManifest))
+            if (CommandUtils.FileExists_NoExceptions(bQuiet, LocalManifest))
             {
                 return true;
             }
@@ -780,12 +812,12 @@ namespace AutomationTool
         public static void DeleteLocalTempStorage(CommandEnvironment Env, string StorageBlockName, bool bQuiet = false)
         {
             var LocalManifest = LocalTempStorageManifestFilename(Env, StorageBlockName);
-            DeleteFile(bQuiet, LocalManifest);
+            CommandUtils.DeleteFile(bQuiet, LocalManifest);
         }
         public static bool SharedTempStorageExists(CommandEnvironment Env, string StorageBlockName, string GameFolder = "", bool bQuiet = false)
         {
             var SharedManifest = SharedTempStorageManifestFilename(Env, StorageBlockName, GameFolder);
-            if (FileExists_NoExceptions(bQuiet, SharedManifest))
+            if (CommandUtils.FileExists_NoExceptions(bQuiet, SharedManifest))
             {
                 return true;
             }
@@ -808,7 +840,7 @@ namespace AutomationTool
                 BaseFolder = Env.LocalRoot;
             }
 
-            BaseFolder = CombinePaths(BaseFolder, "/"); 
+            BaseFolder = CommandUtils.CombinePaths(BaseFolder, "/"); 
             if (!BaseFolder.EndsWith("/") && !BaseFolder.EndsWith("\\"))
             {
                 throw new AutomationException("base folder {0} should end with a separator", BaseFolder);
@@ -820,13 +852,13 @@ namespace AutomationTool
                 var StartTime = DateTime.UtcNow;
 
                 var BlockPath = SharedTempStorageDirectory(StorageBlockName, GameFolder);
-                Log("Storing to {0}", BlockPath);
-                if (DirectoryExists_NoExceptions(BlockPath))
+                CommandUtils.Log("Storing to {0}", BlockPath);
+                if (CommandUtils.DirectoryExists_NoExceptions(BlockPath))
                 {
                     throw new AutomationException("Storage Block Already Exists! {0}", BlockPath);
                 }
-                CreateDirectory(true, BlockPath);
-                if (!DirectoryExists_NoExceptions(BlockPath))
+                CommandUtils.CreateDirectory(true, BlockPath);
+                if (!CommandUtils.DirectoryExists_NoExceptions(BlockPath))
                 {
                     throw new AutomationException("Storage Block Could Not Be Created! {0}", BlockPath);
                 }
@@ -836,7 +868,7 @@ namespace AutomationTool
                 {
                     foreach (string InFilename in Files)
                     {
-                        var Filename = CombinePaths(InFilename);
+                        var Filename = CommandUtils.CombinePaths(InFilename);
                         Robust_FileExists_NoExceptions(false, Filename, "Could not add {0} to manifest because it does not exist");
 
                         if (!Filename.StartsWith(BaseFolder, StringComparison.InvariantCultureIgnoreCase))
@@ -844,12 +876,12 @@ namespace AutomationTool
                             throw new AutomationException("Could not add {0} to manifest because it does not start with the base folder {1}", Filename, BaseFolder);
                         }
                         var RelativeFile = Filename.Substring(BaseFolder.Length);
-                        var DestFile = CombinePaths(BlockPath, RelativeFile);
-                        if (FileExists_NoExceptions(true, DestFile))
+                        var DestFile = CommandUtils.CombinePaths(BlockPath, RelativeFile);
+                        if (CommandUtils.FileExists_NoExceptions(true, DestFile))
                         {
                             throw new AutomationException("Dest file {0} already exists.", DestFile);
                         }
-                        CopyFile(Filename, DestFile, true);
+                        CommandUtils.CopyFile(Filename, DestFile, true);
                         Robust_FileExists_NoExceptions(true, DestFile, "Could not copy to {0}");
 
                         DestFiles.Add(DestFile);
@@ -860,7 +892,7 @@ namespace AutomationTool
                     var SrcFiles = new List<string>();
                     foreach (string InFilename in Files)
                     {
-                        var Filename = CombinePaths(InFilename);
+                        var Filename = CommandUtils.CombinePaths(InFilename);
                         Robust_FileExists_NoExceptions(false, Filename, "Could not add {0} to manifest because it does not exist");
 
                         if (!Filename.StartsWith(BaseFolder, StringComparison.InvariantCultureIgnoreCase))
@@ -868,15 +900,15 @@ namespace AutomationTool
                             throw new AutomationException("Could not add {0} to manifest because it does not start with the base folder {1}", Filename, BaseFolder);
                         }
                         var RelativeFile = Filename.Substring(BaseFolder.Length);
-                        var DestFile = CombinePaths(BlockPath, RelativeFile);
-                        if (FileExists_NoExceptions(true, DestFile))
+                        var DestFile = CommandUtils.CombinePaths(BlockPath, RelativeFile);
+                        if (CommandUtils.FileExists_NoExceptions(true, DestFile))
                         {
                             throw new AutomationException("Dest file {0} already exists.", DestFile);
                         }
                         SrcFiles.Add(Filename);
                         DestFiles.Add(DestFile);
                     }
-                    ThreadedCopyFiles(SrcFiles.ToArray(), DestFiles.ToArray(), ThreadsToCopyWith());
+                    CommandUtils.ThreadedCopyFiles(SrcFiles.ToArray(), DestFiles.ToArray(), ThreadsToCopyWith());
                     foreach (string DestFile in DestFiles)
                     {
                         Robust_FileExists_NoExceptions(true, DestFile, "Could not copy to {0}");
@@ -886,14 +918,14 @@ namespace AutomationTool
                 if (!Local.Compare(Shared))
                 {
                     // we will rename this so it can't be used, but leave it around for inspection
-                    RenameFile_NoExceptions(SharedTempStorageManifestFilename(Env, StorageBlockName, GameFolder), SharedTempStorageManifestFilename(Env, StorageBlockName, GameFolder) + ".broken");
+                    CommandUtils.RenameFile_NoExceptions(SharedTempStorageManifestFilename(Env, StorageBlockName, GameFolder), SharedTempStorageManifestFilename(Env, StorageBlockName, GameFolder) + ".broken");
                     throw new AutomationException("Shared and Local manifest mismatch.");
                 }
                 float BuildDuration = (float)((DateTime.UtcNow - StartTime).TotalSeconds);
                 if (BuildDuration > 60.0f && Shared.GetTotalSize() > 0)
                 {
                     var MBSec = (((float)(Shared.GetTotalSize())) / (1024.0f * 1024.0f)) / BuildDuration;
-                    Log("Wrote to shared temp storage at {0} MB/s    {1}B {2}s", MBSec, Shared.GetTotalSize(), BuildDuration);
+                    CommandUtils.Log("Wrote to shared temp storage at {0} MB/s    {1}B {2}s", MBSec, Shared.GetTotalSize(), BuildDuration);
                 }
             }
 
@@ -901,21 +933,21 @@ namespace AutomationTool
 		public static void RetrieveFromPermanentStorage(CommandEnvironment Env, string NetworkRoot, string Build, string Platform)
 		{
 			var LocalManifest = LocalTempStorageManifestFilename(Env, Build);
-			if (!FileExists_NoExceptions(LocalManifest))
+			if (!CommandUtils.FileExists_NoExceptions(LocalManifest))
 			{
-				string BaseFolder = CombinePaths(Env.LocalRoot, "Rocket", "TempInst", Platform);
-				string Root = RootSharedTempStorageDirectory();
-				string SrcFolder = CombinePaths(Root, NetworkRoot, Build, Platform);
-				var SourceFiles = FindFiles_NoExceptions("*", true, SrcFolder);
+				string BaseFolder = CommandUtils.CombinePaths(Env.LocalRoot, "Rocket", "TempInst", Platform);
+				string Root = CommandUtils.RootSharedTempStorageDirectory();
+				string SrcFolder = CommandUtils.CombinePaths(Root, NetworkRoot, Build, Platform);
+				var SourceFiles = CommandUtils.FindFiles_NoExceptions("*", true, SrcFolder);
 				foreach (var File in SourceFiles)
 				{
 					var DestFile = File.Replace(SrcFolder, "");
-					DestFile = CombinePaths(BaseFolder, DestFile);
-					CopyFile_NoExceptions(File, DestFile);
+					DestFile = CommandUtils.CombinePaths(BaseFolder, DestFile);
+					CommandUtils.CopyFile_NoExceptions(File, DestFile);
 				}
-				if(!DirectoryExists(Path.GetDirectoryName(LocalManifest)))
+				if(!CommandUtils.DirectoryExists(Path.GetDirectoryName(LocalManifest)))
 				{
-					CreateDirectory(Path.GetDirectoryName(LocalManifest));
+					CommandUtils.CreateDirectory(Path.GetDirectoryName(LocalManifest));
 				}
 				XmlDocument Local = new XmlDocument();
 				XmlElement RootElement = Local.CreateElement("permstorage");
@@ -930,7 +962,7 @@ namespace AutomationTool
                 BaseFolder = Env.LocalRoot;
             }
 
-            BaseFolder = CombinePaths(BaseFolder, "/");
+            BaseFolder = CommandUtils.CombinePaths(BaseFolder, "/");
             if (!BaseFolder.EndsWith("/") && !BaseFolder.EndsWith("\\"))
             {
                 throw new AutomationException("base folder {0} should end with a separator", BaseFolder);
@@ -938,9 +970,9 @@ namespace AutomationTool
 
             var Files = new List<string>();
             var LocalManifest = LocalTempStorageManifestFilename(Env, StorageBlockName);
-            if (FileExists_NoExceptions(LocalManifest))
+            if (CommandUtils.FileExists_NoExceptions(LocalManifest))
             {
-                LogVerbose("Found local manifest {0}", LocalManifest);
+                CommandUtils.LogVerbose("Found local manifest {0}", LocalManifest);
                 var Local = new TempStorageManifest();
                 Local.Load(LocalManifest);
                 Files = Local.GetFiles(BaseFolder);
@@ -956,13 +988,13 @@ namespace AutomationTool
             WasLocal = false;
             var StartTime = DateTime.UtcNow;
 
-            var BlockPath = CombinePaths(SharedTempStorageDirectory(StorageBlockName, GameFolder), "/");
+            var BlockPath = CommandUtils.CombinePaths(SharedTempStorageDirectory(StorageBlockName, GameFolder), "/");
             if (!BlockPath.EndsWith("/") && !BlockPath.EndsWith("\\"))
             {
                 throw new AutomationException("base folder {0} should end with a separator", BlockPath);
             }
-            LogVerbose("Attempting to retrieve from {0}", BlockPath);
-            if (!DirectoryExists_NoExceptions(BlockPath))
+            CommandUtils.LogVerbose("Attempting to retrieve from {0}", BlockPath);
+            if (!CommandUtils.DirectoryExists_NoExceptions(BlockPath))
             {
                 throw new AutomationException("Storage Block Does Not Exists! {0}", BlockPath);
             }
@@ -979,7 +1011,7 @@ namespace AutomationTool
             {
                 foreach (string InFilename in SharedFiles)
                 {
-                    var Filename = CombinePaths(InFilename);
+                    var Filename = CommandUtils.CombinePaths(InFilename);
                     Robust_FileExists_NoExceptions(true, Filename, "Could not add {0} to manifest because it does not exist");
 
                     if (!Filename.StartsWith(BlockPath, StringComparison.InvariantCultureIgnoreCase))
@@ -987,19 +1019,19 @@ namespace AutomationTool
                         throw new AutomationException("Could not add {0} to manifest because it does not start with the base folder {1}", Filename, BlockPath);
                     }
                     var RelativeFile = Filename.Substring(BlockPath.Length);
-                    var DestFile = CombinePaths(BaseFolder, RelativeFile);
-                    if (FileExists_NoExceptions(true, DestFile))
+                    var DestFile = CommandUtils.CombinePaths(BaseFolder, RelativeFile);
+                    if (CommandUtils.FileExists_NoExceptions(true, DestFile))
                     {
-                        LogVerbose("Dest file {0} already exists, deleting and overwriting", DestFile);
-                        DeleteFile(DestFile);
+                        CommandUtils.LogVerbose("Dest file {0} already exists, deleting and overwriting", DestFile);
+                        CommandUtils.DeleteFile(DestFile);
                     }
-                    CopyFile(Filename, DestFile, true);
+                    CommandUtils.CopyFile(Filename, DestFile, true);
 
                     Robust_FileExists_NoExceptions(true, DestFile, "Could not copy to {0}");
 
                     if (UnrealBuildTool.Utils.IsRunningOnMono)
                     {
-						FixUnixFilePermissions(DestFile);
+						CommandUtils.FixUnixFilePermissions(DestFile);
                     }
 
                     FileInfo Info = new FileInfo(DestFile);
@@ -1011,7 +1043,7 @@ namespace AutomationTool
                 var SrcFiles = new List<string>();
                 foreach (string InFilename in SharedFiles)
                 {
-                    var Filename = CombinePaths(InFilename);
+                    var Filename = CommandUtils.CombinePaths(InFilename);
                     //Robust_FileExists_NoExceptions(true, Filename, "Could not add {0} to manifest because it does not exist");
 
                     if (!Filename.StartsWith(BlockPath, StringComparison.InvariantCultureIgnoreCase))
@@ -1019,23 +1051,23 @@ namespace AutomationTool
                         throw new AutomationException("Could not add {0} to manifest because it does not start with the base folder {1}", Filename, BlockPath);
                     }
                     var RelativeFile = Filename.Substring(BlockPath.Length);
-                    var DestFile = CombinePaths(BaseFolder, RelativeFile);
-                    if (FileExists_NoExceptions(true, DestFile))
+                    var DestFile = CommandUtils.CombinePaths(BaseFolder, RelativeFile);
+                    if (CommandUtils.FileExists_NoExceptions(true, DestFile))
                     {
-                        LogVerbose("Dest file {0} already exists, deleting and overwriting", DestFile);
-                        DeleteFile(DestFile);
+                        CommandUtils.LogVerbose("Dest file {0} already exists, deleting and overwriting", DestFile);
+                        CommandUtils.DeleteFile(DestFile);
                     }
                     SrcFiles.Add(Filename);
                     DestFiles.Add(DestFile);
                 }
-                ThreadedCopyFiles(SrcFiles.ToArray(), DestFiles.ToArray(), ThreadsToCopyWith());
+                CommandUtils.ThreadedCopyFiles(SrcFiles.ToArray(), DestFiles.ToArray(), ThreadsToCopyWith());
                 var NewDestFiles = new List<string>();
                 foreach (string DestFile in DestFiles)
                 {
                     Robust_FileExists_NoExceptions(true, DestFile, "Could not copy to {0}");
                     if (UnrealBuildTool.Utils.IsRunningOnMono)
                     {
-						FixUnixFilePermissions(DestFile);
+						CommandUtils.FixUnixFilePermissions(DestFile);
                     }
                     FileInfo Info = new FileInfo(DestFile);
                     NewDestFiles.Add(Info.FullName);
@@ -1046,14 +1078,14 @@ namespace AutomationTool
             if (!NewLocal.Compare(Shared))
             {
                 // we will rename this so it can't be used, but leave it around for inspection
-                RenameFile_NoExceptions(LocalManifest, LocalManifest + ".broken");
+                CommandUtils.RenameFile_NoExceptions(LocalManifest, LocalManifest + ".broken");
                 throw new AutomationException("Shared and Local manifest mismatch.");
             }
             float BuildDuration = (float)((DateTime.UtcNow - StartTime).TotalSeconds);
             if (BuildDuration > 60.0f && Shared.GetTotalSize() > 0)
             {
                 var MBSec = (((float)(Shared.GetTotalSize())) / (1024.0f * 1024.0f)) / BuildDuration;
-                Log("Read from shared temp storage at {0} MB/s    {1}B {2}s", MBSec, Shared.GetTotalSize(), BuildDuration);
+                CommandUtils.Log("Read from shared temp storage at {0} MB/s    {1}B {2}s", MBSec, Shared.GetTotalSize(), BuildDuration);
             }
             return DestFiles;
         }
@@ -1074,11 +1106,11 @@ namespace AutomationTool
             string PreStarWildcard = WildCard.Substring(0, IndexOfStar);
             string PostStarWildcard = Path.GetFileNameWithoutExtension(WildCard.Substring(IndexOfStar + 1));
 
-            if (!SharedOnly && DirectoryExists_NoExceptions(LocalParent))
+            if (!SharedOnly && CommandUtils.DirectoryExists_NoExceptions(LocalParent))
             {
                 foreach (var ThisFile in CommandUtils.FindFiles_NoExceptions(WildCard, true, LocalParent))
                 {
-                    Log("  Found local file {0}", ThisFile);
+                    CommandUtils.Log("  Found local file {0}", ThisFile);
                     int IndexOfWildcard = ThisFile.IndexOf(PreStarWildcard);
                     if (IndexOfWildcard < 0)
                     {
@@ -1106,7 +1138,7 @@ namespace AutomationTool
                 var SharedFiles = SharedTempStorageManifestFilename(Env, StorageBlockName, GameFolder);
                 var SharedParent = Path.GetDirectoryName(Path.GetDirectoryName(SharedFiles));
 
-                if (DirectoryExists_NoExceptions(SharedParent))
+                if (CommandUtils.DirectoryExists_NoExceptions(SharedParent))
                 {
                     string[] Dirs = null;
 
@@ -1116,8 +1148,8 @@ namespace AutomationTool
                     }
                     catch (Exception Ex)
                     {
-                        Log("Unable to Find Directories in {0} with wildcard {1}", SharedParent, Path.GetFileNameWithoutExtension(SharedFiles));
-                        Log(" Exception was {0}", LogUtils.FormatException(Ex));
+                        CommandUtils.Log("Unable to Find Directories in {0} with wildcard {1}", SharedParent, Path.GetFileNameWithoutExtension(SharedFiles));
+                        CommandUtils.Log(" Exception was {0}", LogUtils.FormatException(Ex));
                     }
                     if (Dirs != null)
                     {
@@ -1165,4 +1197,264 @@ namespace AutomationTool
             return OutFiles;
         }
     }
+
+	[Help("Tests the temp storage operations.")]
+	class TestTempStorage : BuildCommand
+	{
+		public override void ExecuteBuild()
+		{
+			Log("TestTempStorage********");
+
+			TempStorage.DeleteLocalTempStorageManifests(CmdEnv);
+			TempStorage.DeleteSharedTempStorageManifests(CmdEnv, "Test");
+			if (TempStorage.TempStorageExists(CmdEnv, "Test"))
+			{
+				throw new AutomationException("storage should not exist");
+			}
+
+			string UnitTestFile = CombinePaths(CmdEnv.LocalRoot, "Engine", "Build", "Batchfiles", "TestFile.Txt");
+			Log("Test file {0}", UnitTestFile);
+
+			if (FileExists(UnitTestFile))
+			{
+				DeleteFile(UnitTestFile);
+			}
+			string[] LinesToWrite = new string[]
+			{
+				"This is a test",
+				"Of writing to file " + UnitTestFile
+			};
+			foreach (string Line in LinesToWrite)
+			{
+				WriteToFile(UnitTestFile, Line);
+			}
+
+			string UnitTestFile2 = CombinePaths(CmdEnv.LocalRoot, "TestFile2.Txt");
+			Log("Test file {0}", UnitTestFile2);
+
+			if (FileExists(UnitTestFile2))
+			{
+				DeleteFile(UnitTestFile2);
+			}
+			string[] LinesToWrite2 = new string[]
+			{
+				"This is a test",
+				"Of writing to file " + UnitTestFile2
+			};
+			foreach (string Line in LinesToWrite2)
+			{
+				WriteToFile(UnitTestFile2, Line);
+			}
+
+			string UnitTestFile3 = CombinePaths(CmdEnv.LocalRoot, "engine", "plugins", "TestFile3.Txt");
+			Log("Test file {0}", UnitTestFile3);
+
+			if (FileExists(UnitTestFile3))
+			{
+				DeleteFile(UnitTestFile3);
+			}
+			string[] LinesToWrite3 = new string[]
+			{
+				"This is a test",
+				"Of writing to file " + UnitTestFile3
+			};
+			foreach (string Line in LinesToWrite3)
+			{
+				WriteToFile(UnitTestFile3, Line);
+			}
+
+			{
+				string[] LinesRead = ReadAllLines(UnitTestFile);
+				if (LinesRead == null || LinesRead.Length != LinesToWrite.Length)
+				{
+					throw new AutomationException("Contents of the file created is different to the file read.");
+				}
+				for (int LineIndex = 0; LineIndex < LinesRead.Length; ++LineIndex)
+				{
+					if (LinesRead[LineIndex] != LinesToWrite[LineIndex])
+					{
+						throw new AutomationException("Contents of the file created is different to the file read.");
+					}
+				}
+
+				string[] LinesRead2 = ReadAllLines(UnitTestFile2);
+				if (LinesRead2 == null || LinesRead2.Length != LinesToWrite2.Length)
+				{
+					throw new AutomationException("Contents of the file created is different to the file read.");
+				}
+				for (int LineIndex = 0; LineIndex < LinesRead2.Length; ++LineIndex)
+				{
+					if (LinesRead2[LineIndex] != LinesToWrite2[LineIndex])
+					{
+						throw new AutomationException("Contents of the file created is different to the file read.");
+					}
+				}
+
+				string[] LinesRead3 = ReadAllLines(UnitTestFile3);
+				if (LinesRead3 == null || LinesRead3.Length != LinesToWrite3.Length)
+				{
+					throw new AutomationException("Contents of the file created is different to the file read.");
+				}
+				for (int LineIndex = 0; LineIndex < LinesRead3.Length; ++LineIndex)
+				{
+					if (LinesRead3[LineIndex] != LinesToWrite3[LineIndex])
+					{
+						throw new AutomationException("Contents of the file created is different to the file read.");
+					}
+				}
+			}
+			var TestFiles = new List<string> { UnitTestFile, UnitTestFile2, UnitTestFile3 };
+
+			TempStorage.StoreToTempStorage(CmdEnv, "Test", TestFiles);
+
+			if (!TempStorage.LocalTempStorageExists(CmdEnv, "Test"))
+			{
+				throw new AutomationException("local storage should exist");
+			}
+			if (!TempStorage.SharedTempStorageExists(CmdEnv, "Test"))
+			{
+				throw new AutomationException("shared storage should exist");
+			}
+			TempStorage.DeleteLocalTempStorageManifests(CmdEnv);
+			if (TempStorage.LocalTempStorageExists(CmdEnv, "Test"))
+			{
+				throw new AutomationException("local storage should not exist");
+			}
+			if (!TempStorage.TempStorageExists(CmdEnv, "Test"))
+			{
+				throw new AutomationException("some storage should exist");
+			}
+			DeleteFile(UnitTestFile);
+			DeleteFile(UnitTestFile2);
+			DeleteFile(UnitTestFile3);
+
+			bool WasLocal;
+			TempStorage.RetrieveFromTempStorage(CmdEnv, "Test", out WasLocal);
+			if (!TempStorage.LocalTempStorageExists(CmdEnv, "Test"))
+			{
+				throw new AutomationException("local storage should exist");
+			}
+			if (!TempStorage.SharedTempStorageExists(CmdEnv, "Test"))
+			{
+				throw new AutomationException("shared storage should exist");
+			}
+			{
+				string[] LinesRead = ReadAllLines(UnitTestFile);
+				if (LinesRead == null || LinesRead.Length != LinesToWrite.Length)
+				{
+					throw new AutomationException("Contents of the file created is different to the file read.");
+				}
+				for (int LineIndex = 0; LineIndex < LinesRead.Length; ++LineIndex)
+				{
+					if (LinesRead[LineIndex] != LinesToWrite[LineIndex])
+					{
+						throw new AutomationException("Contents of the file created is different to the file read.");
+					}
+				}
+
+				string[] LinesRead2 = ReadAllLines(UnitTestFile2);
+				if (LinesRead2 == null || LinesRead2.Length != LinesToWrite2.Length)
+				{
+					throw new AutomationException("Contents of the file created is different to the file read.");
+				}
+				for (int LineIndex = 0; LineIndex < LinesRead2.Length; ++LineIndex)
+				{
+					if (LinesRead2[LineIndex] != LinesToWrite2[LineIndex])
+					{
+						throw new AutomationException("Contents of the file created is different to the file read.");
+					}
+				}
+
+				string[] LinesRead3 = ReadAllLines(UnitTestFile3);
+				if (LinesRead3 == null || LinesRead3.Length != LinesToWrite3.Length)
+				{
+					throw new AutomationException("Contents of the file created is different to the file read.");
+				}
+				for (int LineIndex = 0; LineIndex < LinesRead3.Length; ++LineIndex)
+				{
+					if (LinesRead3[LineIndex] != LinesToWrite3[LineIndex])
+					{
+						throw new AutomationException("Contents of the file created is different to the file read.");
+					}
+				}
+			}
+			TempStorage.DeleteSharedTempStorageManifests(CmdEnv, "Test");
+			if (TempStorage.SharedTempStorageExists(CmdEnv, "Test"))
+			{
+				throw new AutomationException("shared storage should not exist");
+			}
+			TempStorage.RetrieveFromTempStorage(CmdEnv, "Test", out WasLocal); // this should just rely on the local
+			if (!WasLocal || !TempStorage.LocalTempStorageExists(CmdEnv, "Test"))
+			{
+				throw new AutomationException("local storage should exist");
+			}
+
+			// and now lets test tampering
+			TempStorage.DeleteLocalTempStorageManifests(CmdEnv);
+			{
+				bool bFailedProperly = false;
+				var MissingFile = new List<string>(TestFiles);
+				MissingFile.Add(CombinePaths(CmdEnv.LocalRoot, "Engine", "SomeFileThatDoesntExist.txt"));
+				try
+				{
+					TempStorage.StoreToTempStorage(CmdEnv, "Test", MissingFile);
+				}
+				catch (AutomationException)
+				{
+					bFailedProperly = true;
+				}
+				if (!bFailedProperly)
+				{
+					throw new AutomationException("Missing file did not fail.");
+				}
+			}
+			TempStorage.DeleteSharedTempStorageManifests(CmdEnv, "Test"); // this ends up being junk
+			TempStorage.StoreToTempStorage(CmdEnv, "Test", TestFiles);
+			TempStorage.DeleteLocalTempStorageManifests(CmdEnv); // force a load from shared
+			var SharedFile = CombinePaths(TempStorage.SharedTempStorageDirectory("Test"), "Engine", "Build", "Batchfiles", "TestFile.Txt");
+			if (!FileExists_NoExceptions(SharedFile))
+			{
+				throw new AutomationException("Shared file {0} did not exist", SharedFile);
+			}
+			DeleteFile(SharedFile);
+			{
+				bool bFailedProperly = false;
+				try
+				{
+					TempStorage.RetrieveFromTempStorage(CmdEnv, "Test", out WasLocal);
+				}
+				catch (AutomationException)
+				{
+					bFailedProperly = true;
+				}
+				if (!bFailedProperly)
+				{
+					throw new AutomationException("Did not fail to load from missing file.");
+				}
+			}
+			TempStorage.DeleteSharedTempStorageManifests(CmdEnv, "Test");
+			TempStorage.StoreToTempStorage(CmdEnv, "Test", TestFiles);
+			DeleteFile(UnitTestFile);
+			{
+				bool bFailedProperly = false;
+				try
+				{
+					TempStorage.RetrieveFromTempStorage(CmdEnv, "Test", out WasLocal);
+				}
+				catch (AutomationException)
+				{
+					bFailedProperly = true;
+				}
+				if (!bFailedProperly)
+				{
+					throw new AutomationException("Did not fail to load from missing local file.");
+				}
+			}
+			TempStorage.DeleteSharedTempStorageManifests(CmdEnv, "Test");
+			TempStorage.DeleteLocalTempStorageManifests(CmdEnv);
+			DeleteFile(UnitTestFile);
+			DeleteFile(UnitTestFile2);
+			DeleteFile(UnitTestFile3);
+		}
+	}
 }
