Move TempStorage into the GUBP project.

diff --git a/Engine/Source/Programs/AutomationTool/AutomationTool.csproj b/Engine/Source/Programs/AutomationTool/AutomationTool.csproj
--- a/Engine/Source/Programs/AutomationTool/AutomationTool.csproj
+++ b/Engine/Source/Programs/AutomationTool/AutomationTool.csproj
@@ -125,7 +125,6 @@
     <Compile Include="ScriptCompiler.cs" />
     <Compile Include="CommandUtils.cs" />
     <Compile Include="SharedUtils.cs" />
-    <Compile Include="TempStorage.cs" />
     <Compile Include="UBTUtils.cs" />
     <Compile Include="UE4Build.cs" />
     <Compile Include="Utils.cs" />
diff --git a/Engine/Source/Programs/AutomationTool/GUBP/GUBP.Automation.csproj b/Engine/Source/Programs/AutomationTool/GUBP/GUBP.Automation.csproj
--- a/Engine/Source/Programs/AutomationTool/GUBP/GUBP.Automation.csproj
+++ b/Engine/Source/Programs/AutomationTool/GUBP/GUBP.Automation.csproj
@@ -59,6 +59,9 @@
       <Link>Properties\MetaData.cs</Link>
     </Compile>
     <Compile Include="GUBP.Automation.cs" />
+    <Compile Include="TempStorage.cs">
+      <SubType>Code</SubType>
+    </Compile>
   </ItemGroup>
   <ItemGroup>
     <ProjectReference Include="..\..\UnrealBuildTool\UnrealBuildTool.csproj">
diff --git a/Engine/Source/Programs/AutomationTool/GUBP/TempStorage.cs b/Engine/Source/Programs/AutomationTool/GUBP/TempStorage.cs
new file mode 100644
--- /dev/null
+++ b/Engine/Source/Programs/AutomationTool/GUBP/TempStorage.cs
@@ -0,0 +1,1460 @@
+﻿// Copyright 1998-2015 Epic Games, Inc. All Rights Reserved.
+using System;
+using System.Collections.Generic;
+using System.Text;
+using System.Linq;
+using System.IO;
+using System.Xml;
+
+namespace AutomationTool
+{
+
+    public class TempStorage
+    {
+        public class TempStorageFileInfo
+        {
+            public string Name;
+            public DateTime Timestamp;
+            public long Size;
+
+            public TempStorageFileInfo()
+            {
+            }
+
+            public TempStorageFileInfo(string Filename, string RelativeFile)
+            {
+                FileInfo Info = new FileInfo(Filename);
+                Name = RelativeFile.Replace("\\", "/"); 
+                Timestamp = Info.LastWriteTimeUtc;
+                Size = Info.Length;
+            }
+
+            public bool Compare(TempStorageFileInfo Other)
+            {
+                bool bOk = true;
+                if (!Name.Equals(Other.Name, StringComparison.InvariantCultureIgnoreCase))
+                {
+                    CommandUtils.Log(System.Diagnostics.TraceEventType.Error, "File name mismatch {0} {1}", Name, Other.Name);
+                    bOk = false;
+                }
+                else
+                {
+                    // this is a bit of a hack, but UAT itself creates these, so we need to allow them to be 
+                    bool bOkToBeDifferent = Name.Contains("Engine/Binaries/DotNET/");
+                    // this is a problem with mac compiles
+                    bOkToBeDifferent = bOkToBeDifferent || Name.EndsWith("MacOS/libogg.dylib");
+                    bOkToBeDifferent = bOkToBeDifferent || Name.EndsWith("MacOS/libvorbis.dylib");
+                    bOkToBeDifferent = bOkToBeDifferent || Name.EndsWith("Contents/MacOS/UE4Editor");
+
+                    //temp hack until the mac build products work correctly
+                    bOkToBeDifferent = bOkToBeDifferent || Name.Contains("Engine/Binaries/Mac/UE4Editor.app/Contents/MacOS/");
+
+
+                    // DotNETUtilities.dll is built by a tons of other things
+                    bool bSilentOkToBeDifferent = (Name == "Engine/Binaries/DotNET/DotNETUtilities.dll");
+                    bSilentOkToBeDifferent = bSilentOkToBeDifferent || (Name == "Engine/Binaries/DotNET/DotNETUtilities.pdb");
+                    // RPCUtility is build by IPP and maybe other things
+                    bSilentOkToBeDifferent = bSilentOkToBeDifferent || (Name == "Engine/Binaries/DotNET/RPCUtility.exe");
+                    bSilentOkToBeDifferent = bSilentOkToBeDifferent || (Name == "Engine/Binaries/DotNET/RPCUtility.pdb");
+                    bSilentOkToBeDifferent = bSilentOkToBeDifferent || (Name == "Engine/Binaries/DotNET/AutomationTool.exe");
+                    bSilentOkToBeDifferent = bSilentOkToBeDifferent || (Name == "Engine/Binaries/DotNET/UnrealBuildTool.exe");
+                    bSilentOkToBeDifferent = bSilentOkToBeDifferent || (Name == "Engine/Binaries/DotNET/UnrealBuildTool.exe.config");					
+                    bSilentOkToBeDifferent = bSilentOkToBeDifferent || (Name == "Engine/Binaries/DotNET/EnvVarsToXML.exe");
+                    bSilentOkToBeDifferent = bSilentOkToBeDifferent || (Name == "Engine/Binaries/DotNET/EnvVarsToXML.exe.config");					
+
+                    // Lets just allow all mac warnings to be silent
+                    bSilentOkToBeDifferent = bSilentOkToBeDifferent || Name.Contains("Engine/Binaries/Mac");
+
+                    System.Diagnostics.TraceEventType LogType = bOkToBeDifferent ? System.Diagnostics.TraceEventType.Warning : System.Diagnostics.TraceEventType.Error;
+                    if (bSilentOkToBeDifferent)
+                    {
+                        LogType = System.Diagnostics.TraceEventType.Information;
+                    }
+
+                    // on FAT filesystems writetime has a two seconds resolution
+                    // cf. http://msdn.microsoft.com/en-us/library/windows/desktop/ms724290%28v=vs.85%29.aspx
+                    if (!((Timestamp - Other.Timestamp).TotalSeconds < 2 && (Timestamp - Other.Timestamp).TotalSeconds > -2))
+                    {
+                        CommandUtils.Log(LogType, "File date mismatch {0} {1} {2} {3}", Name, Timestamp.ToString(), Other.Name, Other.Timestamp.ToString());
+                        bOk = bOkToBeDifferent || bSilentOkToBeDifferent;
+                    }
+                    if (!(Size == Other.Size))
+                    {
+                        CommandUtils.Log(LogType, "File size mismatch {0} {1} {2} {3}", Name, Size, Other.Name, Other.Size);
+                        bOk = bOkToBeDifferent;
+                    }
+                }
+                return bOk;
+            }
+
+            public override string ToString()
+            {
+                return String.IsNullOrEmpty(Name) ? "" : Name;
+            }
+        }
+
+		public static void Robust_CopyFile(string InputFileName, string OutputFileName)
+		{
+            if (OutputFileName.StartsWith("/Volumes/"))
+            {
+                int Retry = 0;
+                int NumRetries = 60;
+                bool bCopied = false;
+                while (!bCopied && Retry < NumRetries)
+                {
+                    if (Retry > 0)
+                    {
+                        CommandUtils.Log("*** Mac temp storage retry {0}", OutputFileName);
+                        System.Threading.Thread.Sleep(1000);
+                    }
+                    bCopied = CommandUtils.CopyFile_NoExceptions(InputFileName, OutputFileName, true);
+                    Retry++;
+                }
+            }
+            else
+            {
+                CommandUtils.CopyFile(InputFileName, OutputFileName);
+            }
+		}
+
+        public static void Robust_FileExists_NoExceptions(string Filename, string Message)
+        {
+            Robust_FileExists_NoExceptions(false, Filename, Message);
+        }
+        public static void Robust_FileExists_NoExceptions(bool bQuiet, string Filename, string Message)
+        {
+            if (!CommandUtils.FileExists_NoExceptions(bQuiet, Filename))
+            {
+                bool bFound = false;
+                // mac is terrible on shares, this isn't a solution, but a stop gap
+                if (Filename.StartsWith("/Volumes/"))
+                {
+                    int Retry = 0;
+                    while (!bFound && Retry < 60)
+                    {
+                        CommandUtils.Log(System.Diagnostics.TraceEventType.Information, "*** Mac temp storage retry {0}", Filename);
+                        System.Threading.Thread.Sleep(10000);
+                        bFound = CommandUtils.FileExists_NoExceptions(bQuiet, Filename);
+						Retry++;
+                    }
+                }
+                if (!bFound)
+                {
+                    throw new AutomationException(Message, Filename);
+                }
+            }
+        }
+
+		public static bool Robust_DirectoryExists_NoExceptions(string Directoryname, string Message)
+		{
+			bool bFound = false;
+			if (!CommandUtils.DirectoryExists_NoExceptions(Directoryname))
+			{				
+				// mac is terrible on shares, this isn't a solution, but a stop gap
+				if (Directoryname.StartsWith("/Volumes/"))
+				{
+					int Retry = 0;
+					while (!bFound && Retry < 60)
+					{
+						CommandUtils.Log(System.Diagnostics.TraceEventType.Information, "*** Mac temp storage retry {0}", Directoryname);
+						System.Threading.Thread.Sleep(10000);
+						bFound = CommandUtils.DirectoryExists_NoExceptions(Directoryname);
+						Retry++;
+					}
+				}				
+			}
+			else
+			{
+				bFound = true;
+			}
+			return bFound;
+		}
+        public static bool Robust_DirectoryExistsAndIsWritable_NoExceptions(string Directoryname)
+        {
+            bool bFound = false;
+            if (!CommandUtils.DirectoryExistsAndIsWritable_NoExceptions(Directoryname))
+            {
+                // mac is terrible on shares, this isn't a solution, but a stop gap
+                if (Directoryname.StartsWith("/Volumes/"))
+                {
+                    int Retry = 0;
+                    int NumRetries = 60;
+                    if(!Directoryname.Contains("UE4"))
+                    {
+                        NumRetries = 2;
+                    }
+                    while (!bFound && Retry < NumRetries)
+                    {
+                        CommandUtils.Log("*** Mac temp storage retry {0}", Directoryname);
+                        System.Threading.Thread.Sleep(1000);
+                        bFound = CommandUtils.DirectoryExistsAndIsWritable_NoExceptions(Directoryname);
+                        Retry++;
+                    }
+                }
+            }
+            else
+            {
+                bFound = true;
+            }
+            return bFound;
+        }
+
+        public class TempStorageManifest
+        {
+            private static readonly string RootElementName = "tempstorage";
+            private static readonly string DirectoryElementName = "directory";
+            private static readonly string FileElementName = "file";
+            private static readonly string NameAttributeName = "name";
+            private static readonly string TimestampAttributeName = "timestamp";
+            private static readonly string SizeAttributeName = "size";
+
+            private Dictionary<string, List<TempStorageFileInfo>> Directories = new Dictionary<string, List<TempStorageFileInfo>>();
+
+            public void Create(List<string> InFiles, string BaseFolder)
+            {
+                BaseFolder = CommandUtils.CombinePaths(BaseFolder, "/");
+                if (!BaseFolder.EndsWith("/")&& !BaseFolder.EndsWith("\\"))
+                {
+                    throw new AutomationException("base folder {0} should end with a separator", BaseFolder);
+                }
+
+                foreach (string InFilename in InFiles)
+                {
+                    var Filename = CommandUtils.CombinePaths(InFilename);
+                    Robust_FileExists_NoExceptions(true, Filename, "Could not add {0} to manifest because it does not exist");
+
+                    FileInfo Info = new FileInfo(Filename);
+                    Filename = Info.FullName;
+                    Robust_FileExists_NoExceptions(true, Filename, "Could not add {0} to manifest because it does not exist2");
+
+                    if (!Filename.StartsWith(BaseFolder, StringComparison.InvariantCultureIgnoreCase))
+                    {
+                        throw new AutomationException("Could not add {0} to manifest because it does not start with the base folder {1}", Filename, BaseFolder);
+                    }
+                    var RelativeFile = Filename.Substring(BaseFolder.Length);
+                    List<TempStorageFileInfo> ManifestDirectory;
+                    string DirectoryName = CommandUtils.CombinePaths(Path.GetDirectoryName(Filename), "/");
+                    if (!DirectoryName.StartsWith(BaseFolder, StringComparison.InvariantCultureIgnoreCase))
+                    {
+                        throw new AutomationException("Could not add directory {0} to manifest because it does not start with the base folder {1}", DirectoryName, BaseFolder);
+                    }
+                    var RelativeDir = DirectoryName.Substring(BaseFolder.Length).Replace("\\", "/");
+                    if (Directories.TryGetValue(RelativeDir, out ManifestDirectory) == false)
+                    {
+                        ManifestDirectory = new List<TempStorageFileInfo>();
+                        Directories.Add(RelativeDir, ManifestDirectory);
+                    }
+                    ManifestDirectory.Add(new TempStorageFileInfo(Filename, RelativeFile));
+                }
+
+                Stats("Created manifest");
+
+            }
+
+            public void Stats(string Description)
+            {
+
+                CommandUtils.Log("{0}: Directories={1} Files={2} Size={3}", Description, Directories.Count, GetFileCount(), GetTotalSize());
+            }
+            
+            public bool Compare(TempStorageManifest Other)
+            {
+                if (Directories.Count != Other.Directories.Count)
+                {
+                    CommandUtils.Log(System.Diagnostics.TraceEventType.Error, "Directory count mismatch {0} {1}", Directories.Count, Other.Directories.Count);
+                    foreach (KeyValuePair<string, List<TempStorageFileInfo>> Directory in Directories)
+                    {
+                        List<TempStorageFileInfo> OtherDirectory;
+                        if (Other.Directories.TryGetValue(Directory.Key, out OtherDirectory) == false)
+                        {
+                            CommandUtils.Log(System.Diagnostics.TraceEventType.Error, "Missing Directory {0}", Directory.Key);
+                            return false;
+                        }
+                    }
+                    foreach (KeyValuePair<string, List<TempStorageFileInfo>> Directory in Other.Directories)
+                    {
+                        List<TempStorageFileInfo> OtherDirectory;
+                        if (Directories.TryGetValue(Directory.Key, out OtherDirectory) == false)
+                        {
+                            CommandUtils.Log(System.Diagnostics.TraceEventType.Error, "Missing Other Directory {0}", Directory.Key);
+                            return false;
+                        }
+                    }
+                    return false;
+                }
+
+                foreach (KeyValuePair<string, List<TempStorageFileInfo>> Directory in Directories)
+                {
+                    List<TempStorageFileInfo> OtherDirectory;
+                    if (Other.Directories.TryGetValue(Directory.Key, out OtherDirectory) == false)
+                    {
+                        CommandUtils.Log(System.Diagnostics.TraceEventType.Error, "Missing Directory {0}", Directory.Key); 
+                        return false;
+                    }
+                    if (OtherDirectory.Count != Directory.Value.Count)
+                    {
+                        CommandUtils.Log(System.Diagnostics.TraceEventType.Error, "File count mismatch {0} {1} {2}", Directory.Key, OtherDirectory.Count, Directory.Value.Count);
+                        for (int FileIndex = 0; FileIndex < Directory.Value.Count; ++FileIndex)
+                        {
+                            CommandUtils.Log("Manifest1: {0}", Directory.Value[FileIndex].Name);
+                        }
+                        for (int FileIndex = 0; FileIndex < OtherDirectory.Count; ++FileIndex)
+                        {
+                            CommandUtils.Log("Manifest2: {0}", OtherDirectory[FileIndex].Name);
+                        }
+                        return false;
+                    }
+                    bool bResult = true;
+                    for (int FileIndex = 0; FileIndex < Directory.Value.Count; ++FileIndex)
+                    {
+                        TempStorageFileInfo File = Directory.Value[FileIndex];
+                        TempStorageFileInfo OtherFile = OtherDirectory[FileIndex];
+                        if (File.Compare(OtherFile) == false)
+                        {
+                            bResult = false;
+                        }
+                    }
+                    return bResult;
+                }
+
+                return true;
+            }
+
+            public TempStorageFileInfo FindFile(string LocalPath, string BaseFolder)
+            {
+                BaseFolder = CommandUtils.CombinePaths(BaseFolder, "/");
+                if (!BaseFolder.EndsWith("/")&& !BaseFolder.EndsWith("\\"))
+                {
+                    throw new AutomationException("base folder {0} should end with a separator", BaseFolder);
+                }
+                var Filename = CommandUtils.CombinePaths(LocalPath);
+                if (!Filename.StartsWith(BaseFolder, StringComparison.InvariantCultureIgnoreCase))
+                {
+                    throw new AutomationException("Could not add {0} to manifest because it does not start with the base folder {1}", Filename, BaseFolder);
+                }
+                var RelativeFile = Filename.Substring(BaseFolder.Length);
+                string DirectoryName = CommandUtils.CombinePaths(Path.GetDirectoryName(Filename), "/");
+                if (!DirectoryName.StartsWith(BaseFolder, StringComparison.InvariantCultureIgnoreCase))
+                {
+                    throw new AutomationException("Could not add directory {0} to manifest because it does not start with the base folder {1}", DirectoryName, BaseFolder);
+                }
+                var RelativeDir = DirectoryName.Substring(BaseFolder.Length);               
+
+                List<TempStorageFileInfo> Directory;
+                if (Directories.TryGetValue(RelativeDir, out Directory))
+                {
+                    foreach (var SyncedFile in Directory)
+                    {
+                        if (SyncedFile.Name.Equals(RelativeFile.Replace("\\", "/"), StringComparison.InvariantCultureIgnoreCase))
+                        {
+                            return SyncedFile;
+                        }
+                    }
+                }
+                return null;
+            }
+
+            public int GetFileCount()
+            {
+                int FileCount = 0;
+                foreach (var Directory in Directories)
+                {
+                    FileCount += Directory.Value.Count;
+                }
+                return FileCount;
+            }
+
+
+            public void Load(string Filename, bool bQuiet = false)
+            {
+                XmlDocument File = new XmlDocument();
+                try
+                {
+                    File.Load(Filename);
+                }
+                catch (Exception Ex)
+                {
+                    throw new AutomationException("Unable to load manifest file {0}: {1}", Filename, Ex.Message);
+                }
+
+                XmlNode RootNode = File.FirstChild;
+                if (RootNode != null && RootNode.Name == RootElementName)
+                {
+                    for (XmlNode ChildNode = RootNode.FirstChild; ChildNode != null; ChildNode = ChildNode.NextSibling)
+                    {
+                        if (ChildNode.Name == DirectoryElementName)
+                        {
+                            LoadDirectory(ChildNode);
+                        }
+                        else
+                        {
+                            throw new AutomationException("Unexpected node {0} while reading manifest files.", ChildNode.Name);
+                        }
+                    }
+                }
+                else
+                {
+                    throw new AutomationException("Bad root node ({0}).", RootNode != null ? RootNode.Name : "null");
+                }
+
+                if (!bQuiet)
+                {
+                    Stats(String.Format("Loaded manifest {0}", Filename));
+                }
+                if (GetFileCount() <= 0 || GetTotalSize() <= 0)
+                {
+                    throw new AutomationException("Attempt to load empty manifest.");
+                }
+            }
+
+            bool TryGetAttribute(XmlNode Node, string AttributeName, out string Value)
+            {
+                bool Result = false;
+                Value = "";
+                if (Node.Attributes != null && Node.Attributes.Count > 0)
+                {
+                    for (int Index = 0; Result == false && Index < Node.Attributes.Count; ++Index)
+                    {
+                        if (Node.Attributes[Index].Name == AttributeName)
+                        {
+                            Value = Node.Attributes[Index].Value;
+                            Result = true;
+                        }
+                    }
+                }
+                return Result;
+            }
+
+            void LoadDirectory(XmlNode DirectoryNode)
+            {
+                string DirectoryName;
+                if (TryGetAttribute(DirectoryNode, NameAttributeName, out DirectoryName) == false)
+                {
+                    throw new AutomationException("Unable to get directory name attribute.");
+                }
+
+                List<TempStorageFileInfo> Files = new List<TempStorageFileInfo>();
+                for (XmlNode ChildNode = DirectoryNode.FirstChild; ChildNode != null; ChildNode = ChildNode.NextSibling)
+                {
+                    if (ChildNode.Name == FileElementName)
+                    {
+                        TempStorageFileInfo Info = LoadFile(ChildNode);
+                        if (Info != null)
+                        {
+                            Files.Add(Info);
+                        }
+                        else
+                        {
+                            throw new AutomationException("Error reading file entry in the manifest file ({0})", ChildNode.InnerXml);
+                        }
+                    }
+                    else
+                    {
+                        throw new AutomationException("Unexpected node {0} while reading file nodes.", ChildNode.Name);
+                    }
+                }
+                Directories.Add(DirectoryName, Files);
+            }
+
+            TempStorageFileInfo LoadFile(XmlNode FileNode)
+            {
+                TempStorageFileInfo Info = new TempStorageFileInfo();
+
+                long Timestamp;
+                string TimestampAttribute;
+                if (TryGetAttribute(FileNode, TimestampAttributeName, out TimestampAttribute) == false)
+                {
+                   throw new AutomationException("Unable to get timestamp attribute.");
+                }
+                if (Int64.TryParse(TimestampAttribute, out Timestamp) == false)
+                {
+                    throw new AutomationException("Unable to parse timestamp attribute ({0}).", TimestampAttribute);
+                }
+                Info.Timestamp = new DateTime(Timestamp);
+
+                string SizeAttribute;
+                if (TryGetAttribute(FileNode, SizeAttributeName, out SizeAttribute) == false)
+                {
+                    throw new AutomationException("Unable to get size attribute.");
+                }
+                if (Int64.TryParse(SizeAttribute, out Info.Size) == false)
+                {
+                    throw new AutomationException("Unable to parse size attribute ({0}).", SizeAttribute);
+                }
+
+                Info.Name = FileNode.InnerText;
+
+                return Info;
+            }
+
+            public void Save(string Filename)
+            {
+                if (GetFileCount() <= 0 || GetTotalSize() <= 0)
+                {
+                    throw new AutomationException("Attempt to save empty manifest.");
+                }
+                XmlDocument File = new XmlDocument();
+
+                XmlElement RootElement = File.CreateElement(RootElementName);
+                File.AppendChild(RootElement);
+
+                SaveDirectories(File, RootElement);
+
+                File.Save(Filename);
+            }
+
+            void SaveDirectories(XmlDocument File, XmlElement Root)
+            {
+                foreach (KeyValuePair<string, List<TempStorageFileInfo>> Directory in Directories)
+                {
+                    XmlElement DirectoryElement = File.CreateElement(DirectoryElementName);
+                    Root.AppendChild(DirectoryElement);
+                    DirectoryElement.SetAttribute(NameAttributeName, Directory.Key);
+                    SaveFiles(File, DirectoryElement, Directory.Value);
+                }
+            }
+
+            void SaveFiles(XmlDocument File, XmlElement Directory, List<TempStorageFileInfo> Files)
+            {
+                foreach (TempStorageFileInfo Info in Files)
+                {
+                    XmlElement FileElement = File.CreateElement(FileElementName);
+                    FileElement.SetAttribute(TimestampAttributeName, Info.Timestamp.Ticks.ToString());
+                    FileElement.SetAttribute(SizeAttributeName, Info.Size.ToString());
+                    FileElement.InnerText = Info.Name;
+                    Directory.AppendChild(FileElement);
+                }
+            }
+
+            public Dictionary<string, List<TempStorageFileInfo>> DirectoryList
+            {
+                get { return Directories; }
+            }
+
+            public long GetTotalSize()
+            {
+
+                long TotalSize = 0;
+
+                foreach (var Dir in Directories)
+                {
+                    foreach (var FileInfo in Dir.Value)
+                    {
+                        TotalSize += FileInfo.Size;
+                    }
+                }
+
+                return TotalSize;
+            }
+            public List<string> GetFiles(string NewBaseDir)
+            {
+                var Result = new List<string>();
+                foreach (var Dir in Directories)
+                {
+                    foreach (var ThisFileInfo in Dir.Value)
+                    {
+                        var NewFile = CommandUtils.CombinePaths(NewBaseDir, ThisFileInfo.Name);
+
+                        FileInfo Info = new FileInfo(NewFile);
+
+                        Result.Add(Info.FullName);
+                    }
+                }
+                if (Result.Count < 1)
+                {
+                    throw new AutomationException("No files in attempt to GetFiles().");
+                }
+
+				CommandUtils.Log("Checking that first file in manifest exists...");
+				Robust_FileExists_NoExceptions(false, Result[0], "Rebased manifest file does not exist {0}");
+
+                return Result;
+            }
+        }
+
+        public static string LocalTempStorageManifestDirectory(CommandEnvironment Env)
+        {
+            return CommandUtils.CombinePaths(Env.LocalRoot, "Engine", "Saved", "GUBP");
+        }
+        static HashSet<string> TopDirectoryTestedForClean = new HashSet<string>();
+        static void CleanSharedTempStorage(string TopDirectory)
+        {
+            if (TopDirectoryTestedForClean.Contains(TopDirectory))
+            {
+                return;
+            }
+            TopDirectoryTestedForClean.Add(TopDirectory);
+
+            const int MaximumDaysToKeepTempStorage = 2;
+            var StartTimeDir = DateTime.UtcNow;
+            DirectoryInfo DirInfo = new DirectoryInfo(TopDirectory);
+            var TopLevelDirs = DirInfo.GetDirectories();
+            {
+                var BuildDuration = (DateTime.UtcNow - StartTimeDir).TotalMilliseconds;
+                CommandUtils.Log("Took {0}s to enumerate {1} directories.", BuildDuration / 1000, TopLevelDirs.Length);
+            }
+            foreach (var TopLevelDir in TopLevelDirs)
+            {
+                if (CommandUtils.DirectoryExists_NoExceptions(TopLevelDir.FullName))
+                {
+                    bool bOld = false;
+                    foreach (var ThisFile in CommandUtils.FindFiles_NoExceptions(true, "*.TempManifest", false, TopLevelDir.FullName))
+                    {
+                        FileInfo Info = new FileInfo(ThisFile);
+
+                        if ((DateTime.UtcNow - Info.LastWriteTimeUtc).TotalDays > MaximumDaysToKeepTempStorage)
+                        {
+                            bOld = true;
+                        }
+                    }
+                    if (bOld)
+                    {
+                        CommandUtils.Log("Deleteing temp storage directory {0}, because it is more than {1} days old.", TopLevelDir.FullName, MaximumDaysToKeepTempStorage);
+                        var StartTime = DateTime.UtcNow;
+                        try
+                        {
+                            if (Directory.Exists(TopLevelDir.FullName))
+                            {
+                                // try the direct approach first
+                                Directory.Delete(TopLevelDir.FullName, true);
+                            }
+                        }
+                        catch (Exception)
+                        {
+                        }
+                        CommandUtils.DeleteDirectory_NoExceptions(true, TopLevelDir.FullName);
+                        var BuildDuration = (DateTime.UtcNow - StartTime).TotalMilliseconds;
+                        CommandUtils.Log("Took {0}s to delete {1}.", BuildDuration / 1000, TopLevelDir.FullName);
+                    }
+                }
+            }
+        }
+        
+        public static string SharedTempStorageDirectory()
+        {
+            return "GUBP";
+        }
+        public static string UE4TempStorageDirectory()
+        {
+            string SharedSubdir = SharedTempStorageDirectory();
+            string Root = CommandUtils.RootSharedTempStorageDirectory();
+            return CommandUtils.CombinePaths(Root, "UE4", SharedSubdir);
+        }
+        public static bool HaveSharedTempStorage(bool ForSaving)
+        {
+            var Dir = UE4TempStorageDirectory();
+            if (ForSaving)
+            {
+                int Retries = 0;
+                while (Retries < 24)
+                {
+                    if (Robust_DirectoryExistsAndIsWritable_NoExceptions(Dir))
+                    {
+                        return true;
+                    }
+                    CommandUtils.FindDirectories_NoExceptions(false, "*", false, Dir); // there is some internet evidence to suggest this might perk up the mac share
+                    System.Threading.Thread.Sleep(5000);
+					Retries++;
+                }
+            }
+            else if (Robust_DirectoryExists_NoExceptions(Dir, "Could not find {0}"))
+            {
+                return true;
+            }
+            return false;
+        }
+
+        static Dictionary<string, string> ResolveCache = new Dictionary<string, string>();
+        public static string ResolveSharedBuildDirectory(string GameFolder)
+        {
+            if (ResolveCache.ContainsKey(GameFolder))
+            {
+                return ResolveCache[GameFolder];
+            }
+            string Root = CommandUtils.RootSharedTempStorageDirectory();
+            string Result = CommandUtils.CombinePaths(Root, GameFolder);
+            if (String.IsNullOrEmpty(GameFolder) || !Robust_DirectoryExistsAndIsWritable_NoExceptions(Result))
+            {
+                string GameStr = "Game";
+                bool HadGame = false;
+                if (GameFolder.EndsWith(GameStr, StringComparison.InvariantCultureIgnoreCase))
+                {
+                    string ShortFolder = GameFolder.Substring(0, GameFolder.Length - GameStr.Length);
+                    Result = CommandUtils.CombinePaths(Root, ShortFolder);
+                    HadGame = true;
+                }
+                if (!HadGame || !Robust_DirectoryExistsAndIsWritable_NoExceptions(Result))
+                {
+                    Result = CommandUtils.CombinePaths(Root, "UE4");
+                    if (!Robust_DirectoryExistsAndIsWritable_NoExceptions(Result))
+                    {
+                        throw new AutomationException("Could not find an appropriate shared temp folder {0}", Result);
+                    }
+                }
+            }
+            ResolveCache.Add(GameFolder, Result);
+            return Result;
+        }
+
+        public static string ResolveSharedTempStorageDirectory(string GameFolder)
+        {
+            string SharedSubdir = SharedTempStorageDirectory();
+            string Result = CommandUtils.CombinePaths(ResolveSharedBuildDirectory(GameFolder), SharedSubdir);
+
+            if (!Robust_DirectoryExists_NoExceptions(Result, "Could not find {0}"))
+            {
+                CommandUtils.CreateDirectory_NoExceptions(Result);
+            }
+            if (!Robust_DirectoryExists_NoExceptions(Result, "Could not find {0}"))
+            {
+                throw new AutomationException("Could not create an appropriate shared temp folder {0}", Result);
+            }
+            return Result;
+        }
+
+        static HashSet<string> TestedForClean = new HashSet<string>();
+        public static void CleanSharedTempStorageDirectory(string GameFolder)
+        {
+
+            if (!CommandUtils.IsBuildMachine || TestedForClean.Contains(GameFolder) ||
+                UnrealBuildTool.Utils.IsRunningOnMono)  // saw a hang on this, anyway it isn't necessary to clean with macs, they are slow anyway
+            {
+                return;
+            }
+            TestedForClean.Add(GameFolder);
+            try
+            {
+                CleanSharedTempStorage(ResolveSharedTempStorageDirectory(GameFolder));
+            }
+            catch (Exception Ex)
+            {
+                CommandUtils.Log(System.Diagnostics.TraceEventType.Warning, "Unable to Clean Directory with GameName {0}", GameFolder);
+                CommandUtils.Log(System.Diagnostics.TraceEventType.Warning, " Exception was {0}", LogUtils.FormatException(Ex));
+            }
+        }
+        public static string SharedTempStorageDirectory(string StorageBlock, string GameFolder = "")
+        {
+            return CommandUtils.CombinePaths(ResolveSharedTempStorageDirectory(GameFolder), StorageBlock);
+        }
+
+        public static TempStorageManifest SaveTempStorageManifest(string RootDir, string FinalFilename, List<string> Files)
+        {
+            var Saver = new TempStorageManifest();
+            Saver.Create(Files, RootDir);
+            if (Saver.GetFileCount() != Files.Count)
+            {
+                throw new AutomationException("Saver manifest differs has wrong number of files {0} != {1}", Saver.GetFileCount(), Files.Count);
+            }
+            var TempFilename = FinalFilename + ".temp";
+            if (CommandUtils.FileExists_NoExceptions(true, TempFilename))
+            {
+                throw new AutomationException("Temp manifest file already exists {0}", TempFilename);
+            }
+            CommandUtils.CreateDirectory(true, Path.GetDirectoryName(FinalFilename));
+            Saver.Save(TempFilename);
+
+            var Tester = new TempStorageManifest();
+            Tester.Load(TempFilename, true);
+
+            if (!Saver.Compare(Tester))
+            {
+                throw new AutomationException("Temp manifest differs {0}", TempFilename);
+            }
+
+            CommandUtils.RenameFile(TempFilename, FinalFilename, true);
+            if (CommandUtils.FileExists_NoExceptions(true, TempFilename))
+            {
+                throw new AutomationException("Temp manifest didn't go away {0}", TempFilename);
+            }
+            var FinalTester = new TempStorageManifest();
+            FinalTester.Load(FinalFilename, true);
+
+            if (!Saver.Compare(FinalTester))
+            {
+                throw new AutomationException("Final manifest differs {0}", TempFilename);
+            }
+            CommandUtils.Log("Saved {0} with {1} files and total size {2}", FinalFilename, Saver.GetFileCount(), Saver.GetTotalSize());
+            return Saver;
+        }
+        public static string LocalTempStorageManifestFilename(CommandEnvironment Env, string StorageBlockName)
+        {
+            return CommandUtils.CombinePaths(LocalTempStorageManifestDirectory(Env), StorageBlockName + ".TempManifest");
+        }
+        public static TempStorageManifest SaveLocalTempStorageManifest(CommandEnvironment Env, string BaseFolder, string StorageBlockName, List<string> Files)
+        {
+            return SaveTempStorageManifest(BaseFolder, LocalTempStorageManifestFilename(Env, StorageBlockName), Files);
+        }
+        public static string SharedTempStorageManifestFilename(CommandEnvironment Env, string StorageBlockName, string GameFolder)
+        {
+            return CommandUtils.CombinePaths(SharedTempStorageDirectory(StorageBlockName, GameFolder), StorageBlockName + ".TempManifest");
+        }
+        public static TempStorageManifest SaveSharedTempStorageManifest(CommandEnvironment Env, string StorageBlockName, string GameFolder, List<string> Files)
+        {
+            return SaveTempStorageManifest(SharedTempStorageDirectory(StorageBlockName, GameFolder), SharedTempStorageManifestFilename(Env, StorageBlockName, GameFolder), Files);
+        }
+        public static void DeleteLocalTempStorageManifests(CommandEnvironment Env)
+        {
+            CommandUtils.DeleteDirectory(true, LocalTempStorageManifestDirectory(Env));
+        }
+        public static void DeleteSharedTempStorageManifests(CommandEnvironment Env, string StorageBlockName, string GameFolder = "")
+        {
+            CommandUtils.DeleteDirectory(true, SharedTempStorageDirectory(StorageBlockName, GameFolder));
+        }
+        public static bool LocalTempStorageExists(CommandEnvironment Env, string StorageBlockName, bool bQuiet = false)
+        {
+            var LocalManifest = LocalTempStorageManifestFilename(Env, StorageBlockName);
+            if (CommandUtils.FileExists_NoExceptions(bQuiet, LocalManifest))
+            {
+                return true;
+            }
+            return false;
+        }
+        public static void DeleteLocalTempStorage(CommandEnvironment Env, string StorageBlockName, bool bQuiet = false)
+        {
+            var LocalManifest = LocalTempStorageManifestFilename(Env, StorageBlockName);
+            CommandUtils.DeleteFile(bQuiet, LocalManifest);
+        }
+        public static bool SharedTempStorageExists(CommandEnvironment Env, string StorageBlockName, string GameFolder = "", bool bQuiet = false)
+        {
+            var SharedManifest = SharedTempStorageManifestFilename(Env, StorageBlockName, GameFolder);
+            if (CommandUtils.FileExists_NoExceptions(bQuiet, SharedManifest))
+            {
+                return true;
+            }
+            return false;
+        }
+        public static bool TempStorageExists(CommandEnvironment Env, string StorageBlockName, string GameFolder = "", bool bLocalOnly = false, bool bQuiet = false)
+        {
+            return LocalTempStorageExists(Env, StorageBlockName, bQuiet) || (!bLocalOnly && SharedTempStorageExists(Env, StorageBlockName, GameFolder, bQuiet));
+        }
+
+        static int ThreadsToCopyWith()
+        {
+            return 8;
+        }
+
+        public static void StoreToTempStorage(CommandEnvironment Env, string StorageBlockName, List<string> Files, bool bLocalOnly = false, string GameFolder = "", string BaseFolder = "")
+        {
+            if (String.IsNullOrEmpty(BaseFolder))
+            {
+                BaseFolder = Env.LocalRoot;
+            }
+
+            BaseFolder = CommandUtils.CombinePaths(BaseFolder, "/"); 
+            if (!BaseFolder.EndsWith("/") && !BaseFolder.EndsWith("\\"))
+            {
+                throw new AutomationException("base folder {0} should end with a separator", BaseFolder);
+            }
+
+            var Local = SaveLocalTempStorageManifest(Env, BaseFolder, StorageBlockName, Files); 
+            if (!bLocalOnly)
+            {
+                var StartTime = DateTime.UtcNow;
+
+                var BlockPath = SharedTempStorageDirectory(StorageBlockName, GameFolder);
+                CommandUtils.Log("Storing to {0}", BlockPath);
+                if (CommandUtils.DirectoryExists_NoExceptions(BlockPath))
+                {
+                    throw new AutomationException("Storage Block Already Exists! {0}", BlockPath);
+                }
+                CommandUtils.CreateDirectory(true, BlockPath);
+                if (!CommandUtils.DirectoryExists_NoExceptions(BlockPath))
+                {
+                    throw new AutomationException("Storage Block Could Not Be Created! {0}", BlockPath);
+                }
+
+                var DestFiles = new List<string>();
+                if (ThreadsToCopyWith() < 2)
+                {
+                    foreach (string InFilename in Files)
+                    {
+                        var Filename = CommandUtils.CombinePaths(InFilename);
+                        Robust_FileExists_NoExceptions(false, Filename, "Could not add {0} to manifest because it does not exist");
+
+                        if (!Filename.StartsWith(BaseFolder, StringComparison.InvariantCultureIgnoreCase))
+                        {
+                            throw new AutomationException("Could not add {0} to manifest because it does not start with the base folder {1}", Filename, BaseFolder);
+                        }
+                        var RelativeFile = Filename.Substring(BaseFolder.Length);
+                        var DestFile = CommandUtils.CombinePaths(BlockPath, RelativeFile);
+                        if (CommandUtils.FileExists_NoExceptions(true, DestFile))
+                        {
+                            throw new AutomationException("Dest file {0} already exists.", DestFile);
+                        }
+                        CommandUtils.CopyFile(Filename, DestFile, true);
+                        Robust_FileExists_NoExceptions(true, DestFile, "Could not copy to {0}");
+
+                        DestFiles.Add(DestFile);
+                    }
+                }
+                else
+                {
+                    var SrcFiles = new List<string>();
+                    foreach (string InFilename in Files)
+                    {
+                        var Filename = CommandUtils.CombinePaths(InFilename);
+                        Robust_FileExists_NoExceptions(false, Filename, "Could not add {0} to manifest because it does not exist");
+
+                        if (!Filename.StartsWith(BaseFolder, StringComparison.InvariantCultureIgnoreCase))
+                        {
+                            throw new AutomationException("Could not add {0} to manifest because it does not start with the base folder {1}", Filename, BaseFolder);
+                        }
+                        var RelativeFile = Filename.Substring(BaseFolder.Length);
+                        var DestFile = CommandUtils.CombinePaths(BlockPath, RelativeFile);
+                        if (CommandUtils.FileExists_NoExceptions(true, DestFile))
+                        {
+                            throw new AutomationException("Dest file {0} already exists.", DestFile);
+                        }
+                        SrcFiles.Add(Filename);
+                        DestFiles.Add(DestFile);
+                    }
+                    CommandUtils.ThreadedCopyFiles(SrcFiles.ToArray(), DestFiles.ToArray(), ThreadsToCopyWith());
+                    foreach (string DestFile in DestFiles)
+                    {
+                        Robust_FileExists_NoExceptions(true, DestFile, "Could not copy to {0}");
+                    }
+                }
+                var Shared = SaveSharedTempStorageManifest(Env, StorageBlockName, GameFolder, DestFiles);
+                if (!Local.Compare(Shared))
+                {
+                    // we will rename this so it can't be used, but leave it around for inspection
+                    CommandUtils.RenameFile_NoExceptions(SharedTempStorageManifestFilename(Env, StorageBlockName, GameFolder), SharedTempStorageManifestFilename(Env, StorageBlockName, GameFolder) + ".broken");
+                    throw new AutomationException("Shared and Local manifest mismatch.");
+                }
+                float BuildDuration = (float)((DateTime.UtcNow - StartTime).TotalSeconds);
+                if (BuildDuration > 60.0f && Shared.GetTotalSize() > 0)
+                {
+                    var MBSec = (((float)(Shared.GetTotalSize())) / (1024.0f * 1024.0f)) / BuildDuration;
+                    CommandUtils.Log("Wrote to shared temp storage at {0} MB/s    {1}B {2}s", MBSec, Shared.GetTotalSize(), BuildDuration);
+                }
+            }
+
+        }
+		public static void RetrieveFromPermanentStorage(CommandEnvironment Env, string NetworkRoot, string Build, string Platform)
+		{
+			var LocalManifest = LocalTempStorageManifestFilename(Env, Build);
+			if (!CommandUtils.FileExists_NoExceptions(LocalManifest))
+			{
+				string BaseFolder = CommandUtils.CombinePaths(Env.LocalRoot, "Rocket", "TempInst", Platform);
+				string Root = CommandUtils.RootSharedTempStorageDirectory();
+				string SrcFolder = CommandUtils.CombinePaths(Root, NetworkRoot, Build, Platform);
+				var SourceFiles = CommandUtils.FindFiles_NoExceptions("*", true, SrcFolder);
+				foreach (var File in SourceFiles)
+				{
+					var DestFile = File.Replace(SrcFolder, "");
+					DestFile = CommandUtils.CombinePaths(BaseFolder, DestFile);
+					CommandUtils.CopyFile_NoExceptions(File, DestFile);
+				}
+				if(!CommandUtils.DirectoryExists(Path.GetDirectoryName(LocalManifest)))
+				{
+					CommandUtils.CreateDirectory(Path.GetDirectoryName(LocalManifest));
+				}
+				XmlDocument Local = new XmlDocument();
+				XmlElement RootElement = Local.CreateElement("permstorage");
+				Local.AppendChild(RootElement);
+				Local.Save(LocalManifest);
+			}
+		}
+        public static List<string> RetrieveFromTempStorage(CommandEnvironment Env, string StorageBlockName, out bool WasLocal, string GameFolder = "", string BaseFolder = "")
+        {
+            if (String.IsNullOrEmpty(BaseFolder))
+            {
+                BaseFolder = Env.LocalRoot;
+            }
+
+            BaseFolder = CommandUtils.CombinePaths(BaseFolder, "/");
+            if (!BaseFolder.EndsWith("/") && !BaseFolder.EndsWith("\\"))
+            {
+                throw new AutomationException("base folder {0} should end with a separator", BaseFolder);
+            }
+
+            var Files = new List<string>();
+            var LocalManifest = LocalTempStorageManifestFilename(Env, StorageBlockName);
+            if (CommandUtils.FileExists_NoExceptions(LocalManifest))
+            {
+                CommandUtils.LogVerbose("Found local manifest {0}", LocalManifest);
+                var Local = new TempStorageManifest();
+                Local.Load(LocalManifest);
+                Files = Local.GetFiles(BaseFolder);
+                var LocalTest = new TempStorageManifest();
+                LocalTest.Create(Files, BaseFolder);
+                if (!Local.Compare(LocalTest))
+                {
+                    throw new AutomationException("Local files in manifest {0} were tampered with.", LocalManifest);
+                }
+                WasLocal = true;
+                return Files;
+            }
+            WasLocal = false;
+            var StartTime = DateTime.UtcNow;
+
+            var BlockPath = CommandUtils.CombinePaths(SharedTempStorageDirectory(StorageBlockName, GameFolder), "/");
+            if (!BlockPath.EndsWith("/") && !BlockPath.EndsWith("\\"))
+            {
+                throw new AutomationException("base folder {0} should end with a separator", BlockPath);
+            }
+            CommandUtils.LogVerbose("Attempting to retrieve from {0}", BlockPath);
+            if (!CommandUtils.DirectoryExists_NoExceptions(BlockPath))
+            {
+                throw new AutomationException("Storage Block Does Not Exists! {0}", BlockPath);
+            }
+            var SharedManifest = SharedTempStorageManifestFilename(Env, StorageBlockName, GameFolder);
+            Robust_FileExists_NoExceptions(SharedManifest, "Storage Block Manifest Does Not Exists! {0}");
+
+            var Shared = new TempStorageManifest();
+            Shared.Load(SharedManifest);
+
+            var SharedFiles = Shared.GetFiles(BlockPath);
+
+            var DestFiles = new List<string>();
+            if (ThreadsToCopyWith() < 2)
+            {
+                foreach (string InFilename in SharedFiles)
+                {
+                    var Filename = CommandUtils.CombinePaths(InFilename);
+                    Robust_FileExists_NoExceptions(true, Filename, "Could not add {0} to manifest because it does not exist");
+
+                    if (!Filename.StartsWith(BlockPath, StringComparison.InvariantCultureIgnoreCase))
+                    {
+                        throw new AutomationException("Could not add {0} to manifest because it does not start with the base folder {1}", Filename, BlockPath);
+                    }
+                    var RelativeFile = Filename.Substring(BlockPath.Length);
+                    var DestFile = CommandUtils.CombinePaths(BaseFolder, RelativeFile);
+                    if (CommandUtils.FileExists_NoExceptions(true, DestFile))
+                    {
+                        CommandUtils.LogVerbose("Dest file {0} already exists, deleting and overwriting", DestFile);
+                        CommandUtils.DeleteFile(DestFile);
+                    }
+                    CommandUtils.CopyFile(Filename, DestFile, true);
+
+                    Robust_FileExists_NoExceptions(true, DestFile, "Could not copy to {0}");
+
+                    if (UnrealBuildTool.Utils.IsRunningOnMono)
+                    {
+						CommandUtils.FixUnixFilePermissions(DestFile);
+                    }
+
+                    FileInfo Info = new FileInfo(DestFile);
+                    DestFiles.Add(Info.FullName);
+                }
+            }
+            else
+            {
+                var SrcFiles = new List<string>();
+                foreach (string InFilename in SharedFiles)
+                {
+                    var Filename = CommandUtils.CombinePaths(InFilename);
+                    //Robust_FileExists_NoExceptions(true, Filename, "Could not add {0} to manifest because it does not exist");
+
+                    if (!Filename.StartsWith(BlockPath, StringComparison.InvariantCultureIgnoreCase))
+                    {
+                        throw new AutomationException("Could not add {0} to manifest because it does not start with the base folder {1}", Filename, BlockPath);
+                    }
+                    var RelativeFile = Filename.Substring(BlockPath.Length);
+                    var DestFile = CommandUtils.CombinePaths(BaseFolder, RelativeFile);
+                    if (CommandUtils.FileExists_NoExceptions(true, DestFile))
+                    {
+                        CommandUtils.LogVerbose("Dest file {0} already exists, deleting and overwriting", DestFile);
+                        CommandUtils.DeleteFile(DestFile);
+                    }
+                    SrcFiles.Add(Filename);
+                    DestFiles.Add(DestFile);
+                }
+                CommandUtils.ThreadedCopyFiles(SrcFiles.ToArray(), DestFiles.ToArray(), ThreadsToCopyWith());
+                var NewDestFiles = new List<string>();
+                foreach (string DestFile in DestFiles)
+                {
+                    Robust_FileExists_NoExceptions(true, DestFile, "Could not copy to {0}");
+                    if (UnrealBuildTool.Utils.IsRunningOnMono)
+                    {
+						CommandUtils.FixUnixFilePermissions(DestFile);
+                    }
+                    FileInfo Info = new FileInfo(DestFile);
+                    NewDestFiles.Add(Info.FullName);
+                }
+                DestFiles = NewDestFiles;
+            }
+            var NewLocal = SaveLocalTempStorageManifest(Env, BaseFolder, StorageBlockName, DestFiles);
+            if (!NewLocal.Compare(Shared))
+            {
+                // we will rename this so it can't be used, but leave it around for inspection
+                CommandUtils.RenameFile_NoExceptions(LocalManifest, LocalManifest + ".broken");
+                throw new AutomationException("Shared and Local manifest mismatch.");
+            }
+            float BuildDuration = (float)((DateTime.UtcNow - StartTime).TotalSeconds);
+            if (BuildDuration > 60.0f && Shared.GetTotalSize() > 0)
+            {
+                var MBSec = (((float)(Shared.GetTotalSize())) / (1024.0f * 1024.0f)) / BuildDuration;
+                CommandUtils.Log("Read from shared temp storage at {0} MB/s    {1}B {2}s", MBSec, Shared.GetTotalSize(), BuildDuration);
+            }
+            return DestFiles;
+        }
+        public static List<string> FindTempStorageManifests(CommandEnvironment Env, string StorageBlockName, bool LocalOnly = false, bool SharedOnly = false, string GameFolder = "")
+        {
+            var Files = new List<string>();
+
+            var LocalFiles = LocalTempStorageManifestFilename(Env, StorageBlockName);
+            var LocalParent = Path.GetDirectoryName(LocalFiles);
+            var WildCard = Path.GetFileName(LocalFiles);
+
+            int IndexOfStar = WildCard.IndexOf("*");
+            if (IndexOfStar < 0 || WildCard.LastIndexOf("*") != IndexOfStar)
+            {
+                throw new AutomationException("Wildcard {0} either has no star or it has more than one.", WildCard);
+            }
+
+            string PreStarWildcard = WildCard.Substring(0, IndexOfStar);
+            string PostStarWildcard = Path.GetFileNameWithoutExtension(WildCard.Substring(IndexOfStar + 1));
+
+            if (!SharedOnly && CommandUtils.DirectoryExists_NoExceptions(LocalParent))
+            {
+                foreach (var ThisFile in CommandUtils.FindFiles_NoExceptions(WildCard, true, LocalParent))
+                {
+                    CommandUtils.Log("  Found local file {0}", ThisFile);
+                    int IndexOfWildcard = ThisFile.IndexOf(PreStarWildcard);
+                    if (IndexOfWildcard < 0)
+                    {
+                        throw new AutomationException("File {0} didn't contain {1}.", ThisFile, PreStarWildcard);
+                    }
+                    int LastIndexOfWildcardTail = ThisFile.LastIndexOf(PostStarWildcard);
+                    if (LastIndexOfWildcardTail < 0 || LastIndexOfWildcardTail < IndexOfWildcard + PreStarWildcard.Length)
+                    {
+                        throw new AutomationException("File {0} didn't contain {1} or it was before the prefix", ThisFile, PostStarWildcard);
+                    }
+                    string StarReplacement = ThisFile.Substring(IndexOfWildcard + PreStarWildcard.Length, LastIndexOfWildcardTail - IndexOfWildcard - PreStarWildcard.Length);
+                    if (StarReplacement.Length < 1)
+                    {
+                        throw new AutomationException("Dir {0} didn't have any string to fit the star in the wildcard {1}", ThisFile, WildCard);
+                    }
+                    if (!Files.Contains(StarReplacement))
+                    {
+                        Files.Add(StarReplacement);
+                    }
+                }
+            }
+
+            if (!LocalOnly)
+            {
+                var SharedFiles = SharedTempStorageManifestFilename(Env, StorageBlockName, GameFolder);
+                var SharedParent = Path.GetDirectoryName(Path.GetDirectoryName(SharedFiles));
+
+                if (CommandUtils.DirectoryExists_NoExceptions(SharedParent))
+                {
+                    string[] Dirs = null;
+
+                    try
+                    {
+                        Dirs = Directory.GetDirectories(SharedParent, Path.GetFileNameWithoutExtension(SharedFiles), SearchOption.TopDirectoryOnly);
+                    }
+                    catch (Exception Ex)
+                    {
+                        CommandUtils.Log("Unable to Find Directories in {0} with wildcard {1}", SharedParent, Path.GetFileNameWithoutExtension(SharedFiles));
+                        CommandUtils.Log(" Exception was {0}", LogUtils.FormatException(Ex));
+                    }
+                    if (Dirs != null)
+                    {
+                        foreach (var ThisSubDir in Dirs)
+                        {
+                            int IndexOfWildcard = ThisSubDir.IndexOf(PreStarWildcard);
+                            if (IndexOfWildcard < 0)
+                            {
+                                throw new AutomationException("Dir {0} didn't contain {1}.", ThisSubDir, PreStarWildcard);
+                            }
+                            int LastIndexOfWildcardTail = ThisSubDir.LastIndexOf(PostStarWildcard);
+                            if (LastIndexOfWildcardTail < 0 || LastIndexOfWildcardTail < IndexOfWildcard + PreStarWildcard.Length)
+                            {
+                                throw new AutomationException("Dir {0} didn't contain {1} or it was before the prefix", ThisSubDir, PostStarWildcard);
+                            }
+                            string StarReplacement = ThisSubDir.Substring(IndexOfWildcard + PreStarWildcard.Length, LastIndexOfWildcardTail - IndexOfWildcard - PreStarWildcard.Length);
+                            if (StarReplacement.Length < 1)
+                            {
+                                throw new AutomationException("Dir {0} didn't have any string to fit the star in the wildcard {1}", ThisSubDir, WildCard);
+                            }
+                            // these are a bunch of false positives
+                            if (StarReplacement.Contains("-"))
+                            {
+                                continue;
+                            }
+                            if (!Files.Contains(StarReplacement))
+                            {
+                                Files.Add(StarReplacement);
+                            }
+                        }
+                    }
+                }
+            }
+
+            var OutFiles = new List<string>();
+            foreach (var StarReplacement in Files)
+            {
+                var NewBlock = StorageBlockName.Replace("*", StarReplacement);
+
+                if (TempStorageExists(Env, NewBlock, GameFolder, LocalOnly, true))
+                {
+                    OutFiles.Add(StarReplacement);
+                }
+            }
+            return OutFiles;
+        }
+    }
+
+	[Help("Tests the temp storage operations.")]
+	class TestTempStorage : BuildCommand
+	{
+		public override void ExecuteBuild()
+		{
+			Log("TestTempStorage********");
+
+			TempStorage.DeleteLocalTempStorageManifests(CmdEnv);
+			TempStorage.DeleteSharedTempStorageManifests(CmdEnv, "Test");
+			if (TempStorage.TempStorageExists(CmdEnv, "Test"))
+			{
+				throw new AutomationException("storage should not exist");
+			}
+
+			string UnitTestFile = CombinePaths(CmdEnv.LocalRoot, "Engine", "Build", "Batchfiles", "TestFile.Txt");
+			Log("Test file {0}", UnitTestFile);
+
+			if (FileExists(UnitTestFile))
+			{
+				DeleteFile(UnitTestFile);
+			}
+			string[] LinesToWrite = new string[]
+			{
+				"This is a test",
+				"Of writing to file " + UnitTestFile
+			};
+			foreach (string Line in LinesToWrite)
+			{
+				WriteToFile(UnitTestFile, Line);
+			}
+
+			string UnitTestFile2 = CombinePaths(CmdEnv.LocalRoot, "TestFile2.Txt");
+			Log("Test file {0}", UnitTestFile2);
+
+			if (FileExists(UnitTestFile2))
+			{
+				DeleteFile(UnitTestFile2);
+			}
+			string[] LinesToWrite2 = new string[]
+			{
+				"This is a test",
+				"Of writing to file " + UnitTestFile2
+			};
+			foreach (string Line in LinesToWrite2)
+			{
+				WriteToFile(UnitTestFile2, Line);
+			}
+
+			string UnitTestFile3 = CombinePaths(CmdEnv.LocalRoot, "engine", "plugins", "TestFile3.Txt");
+			Log("Test file {0}", UnitTestFile3);
+
+			if (FileExists(UnitTestFile3))
+			{
+				DeleteFile(UnitTestFile3);
+			}
+			string[] LinesToWrite3 = new string[]
+			{
+				"This is a test",
+				"Of writing to file " + UnitTestFile3
+			};
+			foreach (string Line in LinesToWrite3)
+			{
+				WriteToFile(UnitTestFile3, Line);
+			}
+
+			{
+				string[] LinesRead = ReadAllLines(UnitTestFile);
+				if (LinesRead == null || LinesRead.Length != LinesToWrite.Length)
+				{
+					throw new AutomationException("Contents of the file created is different to the file read.");
+				}
+				for (int LineIndex = 0; LineIndex < LinesRead.Length; ++LineIndex)
+				{
+					if (LinesRead[LineIndex] != LinesToWrite[LineIndex])
+					{
+						throw new AutomationException("Contents of the file created is different to the file read.");
+					}
+				}
+
+				string[] LinesRead2 = ReadAllLines(UnitTestFile2);
+				if (LinesRead2 == null || LinesRead2.Length != LinesToWrite2.Length)
+				{
+					throw new AutomationException("Contents of the file created is different to the file read.");
+				}
+				for (int LineIndex = 0; LineIndex < LinesRead2.Length; ++LineIndex)
+				{
+					if (LinesRead2[LineIndex] != LinesToWrite2[LineIndex])
+					{
+						throw new AutomationException("Contents of the file created is different to the file read.");
+					}
+				}
+
+				string[] LinesRead3 = ReadAllLines(UnitTestFile3);
+				if (LinesRead3 == null || LinesRead3.Length != LinesToWrite3.Length)
+				{
+					throw new AutomationException("Contents of the file created is different to the file read.");
+				}
+				for (int LineIndex = 0; LineIndex < LinesRead3.Length; ++LineIndex)
+				{
+					if (LinesRead3[LineIndex] != LinesToWrite3[LineIndex])
+					{
+						throw new AutomationException("Contents of the file created is different to the file read.");
+					}
+				}
+			}
+			var TestFiles = new List<string> { UnitTestFile, UnitTestFile2, UnitTestFile3 };
+
+			TempStorage.StoreToTempStorage(CmdEnv, "Test", TestFiles);
+
+			if (!TempStorage.LocalTempStorageExists(CmdEnv, "Test"))
+			{
+				throw new AutomationException("local storage should exist");
+			}
+			if (!TempStorage.SharedTempStorageExists(CmdEnv, "Test"))
+			{
+				throw new AutomationException("shared storage should exist");
+			}
+			TempStorage.DeleteLocalTempStorageManifests(CmdEnv);
+			if (TempStorage.LocalTempStorageExists(CmdEnv, "Test"))
+			{
+				throw new AutomationException("local storage should not exist");
+			}
+			if (!TempStorage.TempStorageExists(CmdEnv, "Test"))
+			{
+				throw new AutomationException("some storage should exist");
+			}
+			DeleteFile(UnitTestFile);
+			DeleteFile(UnitTestFile2);
+			DeleteFile(UnitTestFile3);
+
+			bool WasLocal;
+			TempStorage.RetrieveFromTempStorage(CmdEnv, "Test", out WasLocal);
+			if (!TempStorage.LocalTempStorageExists(CmdEnv, "Test"))
+			{
+				throw new AutomationException("local storage should exist");
+			}
+			if (!TempStorage.SharedTempStorageExists(CmdEnv, "Test"))
+			{
+				throw new AutomationException("shared storage should exist");
+			}
+			{
+				string[] LinesRead = ReadAllLines(UnitTestFile);
+				if (LinesRead == null || LinesRead.Length != LinesToWrite.Length)
+				{
+					throw new AutomationException("Contents of the file created is different to the file read.");
+				}
+				for (int LineIndex = 0; LineIndex < LinesRead.Length; ++LineIndex)
+				{
+					if (LinesRead[LineIndex] != LinesToWrite[LineIndex])
+					{
+						throw new AutomationException("Contents of the file created is different to the file read.");
+					}
+				}
+
+				string[] LinesRead2 = ReadAllLines(UnitTestFile2);
+				if (LinesRead2 == null || LinesRead2.Length != LinesToWrite2.Length)
+				{
+					throw new AutomationException("Contents of the file created is different to the file read.");
+				}
+				for (int LineIndex = 0; LineIndex < LinesRead2.Length; ++LineIndex)
+				{
+					if (LinesRead2[LineIndex] != LinesToWrite2[LineIndex])
+					{
+						throw new AutomationException("Contents of the file created is different to the file read.");
+					}
+				}
+
+				string[] LinesRead3 = ReadAllLines(UnitTestFile3);
+				if (LinesRead3 == null || LinesRead3.Length != LinesToWrite3.Length)
+				{
+					throw new AutomationException("Contents of the file created is different to the file read.");
+				}
+				for (int LineIndex = 0; LineIndex < LinesRead3.Length; ++LineIndex)
+				{
+					if (LinesRead3[LineIndex] != LinesToWrite3[LineIndex])
+					{
+						throw new AutomationException("Contents of the file created is different to the file read.");
+					}
+				}
+			}
+			TempStorage.DeleteSharedTempStorageManifests(CmdEnv, "Test");
+			if (TempStorage.SharedTempStorageExists(CmdEnv, "Test"))
+			{
+				throw new AutomationException("shared storage should not exist");
+			}
+			TempStorage.RetrieveFromTempStorage(CmdEnv, "Test", out WasLocal); // this should just rely on the local
+			if (!WasLocal || !TempStorage.LocalTempStorageExists(CmdEnv, "Test"))
+			{
+				throw new AutomationException("local storage should exist");
+			}
+
+			// and now lets test tampering
+			TempStorage.DeleteLocalTempStorageManifests(CmdEnv);
+			{
+				bool bFailedProperly = false;
+				var MissingFile = new List<string>(TestFiles);
+				MissingFile.Add(CombinePaths(CmdEnv.LocalRoot, "Engine", "SomeFileThatDoesntExist.txt"));
+				try
+				{
+					TempStorage.StoreToTempStorage(CmdEnv, "Test", MissingFile);
+				}
+				catch (AutomationException)
+				{
+					bFailedProperly = true;
+				}
+				if (!bFailedProperly)
+				{
+					throw new AutomationException("Missing file did not fail.");
+				}
+			}
+			TempStorage.DeleteSharedTempStorageManifests(CmdEnv, "Test"); // this ends up being junk
+			TempStorage.StoreToTempStorage(CmdEnv, "Test", TestFiles);
+			TempStorage.DeleteLocalTempStorageManifests(CmdEnv); // force a load from shared
+			var SharedFile = CombinePaths(TempStorage.SharedTempStorageDirectory("Test"), "Engine", "Build", "Batchfiles", "TestFile.Txt");
+			if (!FileExists_NoExceptions(SharedFile))
+			{
+				throw new AutomationException("Shared file {0} did not exist", SharedFile);
+			}
+			DeleteFile(SharedFile);
+			{
+				bool bFailedProperly = false;
+				try
+				{
+					TempStorage.RetrieveFromTempStorage(CmdEnv, "Test", out WasLocal);
+				}
+				catch (AutomationException)
+				{
+					bFailedProperly = true;
+				}
+				if (!bFailedProperly)
+				{
+					throw new AutomationException("Did not fail to load from missing file.");
+				}
+			}
+			TempStorage.DeleteSharedTempStorageManifests(CmdEnv, "Test");
+			TempStorage.StoreToTempStorage(CmdEnv, "Test", TestFiles);
+			DeleteFile(UnitTestFile);
+			{
+				bool bFailedProperly = false;
+				try
+				{
+					TempStorage.RetrieveFromTempStorage(CmdEnv, "Test", out WasLocal);
+				}
+				catch (AutomationException)
+				{
+					bFailedProperly = true;
+				}
+				if (!bFailedProperly)
+				{
+					throw new AutomationException("Did not fail to load from missing local file.");
+				}
+			}
+			TempStorage.DeleteSharedTempStorageManifests(CmdEnv, "Test");
+			TempStorage.DeleteLocalTempStorageManifests(CmdEnv);
+			DeleteFile(UnitTestFile);
+			DeleteFile(UnitTestFile2);
+			DeleteFile(UnitTestFile3);
+		}
+	}
+}
diff --git a/Engine/Source/Programs/AutomationTool/TempStorage.cs b/Engine/Source/Programs/AutomationTool/TempStorage.cs
deleted file mode 100644
