Rename GUBP.Automation.cs to GUBP.cs

diff --git a/Engine/Source/Programs/AutomationTool/GUBP/GUBP.Automation.csproj b/Engine/Source/Programs/AutomationTool/GUBP/GUBP.Automation.csproj
--- a/Engine/Source/Programs/AutomationTool/GUBP/GUBP.Automation.csproj
+++ b/Engine/Source/Programs/AutomationTool/GUBP/GUBP.Automation.csproj
@@ -58,7 +58,7 @@
     <Compile Include="..\..\DotNETCommon\MetaData.cs">
       <Link>Properties\MetaData.cs</Link>
     </Compile>
-    <Compile Include="GUBP.Automation.cs" />
+    <Compile Include="GUBP.cs" />
     <Compile Include="TempStorage.cs">
       <SubType>Code</SubType>
     </Compile>
@@ -74,5 +74,8 @@
       <Name>AutomationTool</Name>
     </ProjectReference>
   </ItemGroup>
+  <ItemGroup>
+    <Folder Include="Nodes\" />
+  </ItemGroup>
   <Import Project="$(MSBuildToolsPath)\Microsoft.CSharp.targets" />
 </Project>
\ No newline at end of file
