Separate out branch setup and build nodes for GUBP into separate files.

diff --git a/Engine/Source/Programs/AutomationTool/GUBP/ElectricCommander.cs b/Engine/Source/Programs/AutomationTool/GUBP/ElectricCommander.cs
new file mode 100644
--- /dev/null
+++ b/Engine/Source/Programs/AutomationTool/GUBP/ElectricCommander.cs
@@ -0,0 +1,95 @@
+﻿using System;
+using System.Collections.Generic;
+using System.Linq;
+using System.Text;
+using System.Xml;
+
+public class ECJobPropsUtils
+{
+    public static HashSet<string> ErrorsFromProps(string Filename)
+    {
+        var Result = new HashSet<string>();
+        XmlDocument Doc = new XmlDocument();
+        Doc.Load(Filename);
+        foreach (XmlElement ChildNode in Doc.FirstChild.ChildNodes)
+        {
+            if (ChildNode.Name == "propertySheet")
+            {
+                foreach (XmlElement PropertySheetChild in ChildNode.ChildNodes)
+                {
+                    if (PropertySheetChild.Name == "property")
+                    {
+                        bool IsDiag = false;
+                        foreach (XmlElement PropertySheetChildDiag in PropertySheetChild.ChildNodes)
+                        {
+                            if (PropertySheetChildDiag.Name == "propertyName" && PropertySheetChildDiag.InnerText == "ec_diagnostics")
+                            {
+                                IsDiag = true;
+                            }
+                            if (IsDiag && PropertySheetChildDiag.Name == "propertySheet")
+                            {
+                                foreach (XmlElement PropertySheetChildDiagSheet in PropertySheetChildDiag.ChildNodes)
+                                {
+                                    if (PropertySheetChildDiagSheet.Name == "property")
+                                    {
+                                        bool IsError = false;
+                                        foreach (XmlElement PropertySheetChildDiagSheetElem in PropertySheetChildDiagSheet.ChildNodes)
+                                        {
+                                            if (PropertySheetChildDiagSheetElem.Name == "propertyName" && PropertySheetChildDiagSheetElem.InnerText.StartsWith("error-"))
+                                            {
+                                                IsError = true;
+                                            }
+                                            if (IsError && PropertySheetChildDiagSheetElem.Name == "propertySheet")
+                                            {
+                                                foreach (XmlElement PropertySheetChildDiagSheetElemInner in PropertySheetChildDiagSheetElem.ChildNodes)
+                                                {
+                                                    if (PropertySheetChildDiagSheetElemInner.Name == "property")
+                                                    {
+                                                        bool IsMessage = false;
+                                                        foreach (XmlElement PropertySheetChildDiagSheetElemInner2 in PropertySheetChildDiagSheetElemInner.ChildNodes)
+                                                        {
+                                                            if (PropertySheetChildDiagSheetElemInner2.Name == "propertyName" && PropertySheetChildDiagSheetElemInner2.InnerText == "message")
+                                                            {
+                                                                IsMessage = true;
+                                                            }
+                                                            if (IsMessage && PropertySheetChildDiagSheetElemInner2.Name == "value")
+                                                            {
+                                                                if (!PropertySheetChildDiagSheetElemInner2.InnerText.Contains("LogTailsAndChanges")
+                                                                    && !PropertySheetChildDiagSheetElemInner2.InnerText.Contains("-MyJobStepId=")
+                                                                    && !PropertySheetChildDiagSheetElemInner2.InnerText.Contains("CommandUtils.Run: Run: Took ")
+                                                                    )
+                                                                {
+                                                                    Result.Add(PropertySheetChildDiagSheetElemInner2.InnerText);
+                                                                }
+                                                            }
+                                                        }
+                                                    }
+                                                }
+                                            }
+                                        }
+                                    }
+                                }
+                            }
+                        }
+                    }
+                }
+            }
+        }
+        return Result;
+    }
+}
+
+public class TestECJobErrorParse : BuildCommand
+{
+    public override void ExecuteBuild()
+    {
+        Log("*********************** TestECJobErrorParse");
+
+        string Filename = CombinePaths(@"P:\Builds\UE4\GUBP\++depot+UE4-2104401-RootEditor_Failed\Engine\Saved\Logs", "RootEditor_Failed.log");
+        var Errors = ECJobPropsUtils.ErrorsFromProps(Filename);
+        foreach (var ThisError in Errors)
+        {
+            Log("Error: {0}", ThisError);
+        }
+    }
+}
diff --git a/Engine/Source/Programs/AutomationTool/GUBP/GUBP.Automation.csproj b/Engine/Source/Programs/AutomationTool/GUBP/GUBP.Automation.csproj
--- a/Engine/Source/Programs/AutomationTool/GUBP/GUBP.Automation.csproj
+++ b/Engine/Source/Programs/AutomationTool/GUBP/GUBP.Automation.csproj
@@ -7,7 +7,7 @@
     <ProjectGuid>{9A1A118B-37FC-4F33-BBF4-B431B6E3B3AC}</ProjectGuid>
     <OutputType>Library</OutputType>
     <AppDesignerFolder>Properties</AppDesignerFolder>
-    <RootNamespace>GUBP.Automation</RootNamespace>
+    <RootNamespace>AutomationTool</RootNamespace>
     <AssemblyName>GUBP.Automation</AssemblyName>
     <TargetFrameworkVersion>v4.0</TargetFrameworkVersion>
     <FileAlignment>512</FileAlignment>
@@ -58,7 +58,10 @@
     <Compile Include="..\..\DotNETCommon\MetaData.cs">
       <Link>Properties\MetaData.cs</Link>
     </Compile>
+    <Compile Include="ElectricCommander.cs" />
     <Compile Include="GUBP.cs" />
+    <Compile Include="LegacyBranchSetup.cs" />
+    <Compile Include="LegacyNodes.cs" />
     <Compile Include="TempStorage.cs">
       <SubType>Code</SubType>
     </Compile>
