Separate out branch setup and build nodes for GUBP into separate files.

diff --git a/Engine/Source/Programs/AutomationTool/GUBP/LegacyBranchSetup.cs b/Engine/Source/Programs/AutomationTool/GUBP/LegacyBranchSetup.cs
new file mode 100644
--- /dev/null
+++ b/Engine/Source/Programs/AutomationTool/GUBP/LegacyBranchSetup.cs
@@ -0,0 +1,998 @@
+﻿// Copyright 1998-2015 Epic Games, Inc. All Rights Reserved.
+using System;
+using System.Collections.Generic;
+using System.Text;
+using System.IO;
+using AutomationTool;
+using UnrealBuildTool;
+using System.Reflection;
+using System.Xml;
+using System.Linq;
+
+partial class GUBP
+{
+    public abstract class GUBPNodeAdder
+    {
+        public virtual void AddNodes(GUBP bp, UnrealTargetPlatform InHostPlatform)
+        {
+        }
+    }
+
+    private static List<GUBPNodeAdder> Adders;
+
+	Type[] GetTypesFromAssembly(Assembly Asm)
+	{
+		Type[] AllTypesWeCanGet;
+		try
+		{
+			AllTypesWeCanGet = Asm.GetTypes();
+		}
+		catch(ReflectionTypeLoadException Exc)
+		{
+			AllTypesWeCanGet = Exc.Types;
+		}
+		return AllTypesWeCanGet;
+	}
+
+    private void AddCustomNodes(UnrealTargetPlatform InHostPlatform)
+    {
+        if (Adders == null)
+        {
+            Adders = new List<GUBPNodeAdder>();
+            Assembly[] LoadedAssemblies = AppDomain.CurrentDomain.GetAssemblies();
+            foreach (var Dll in LoadedAssemblies)
+            {
+				Type[] AllTypes = GetTypesFromAssembly(Dll);
+                foreach (var PotentialConfigType in AllTypes)
+                {
+                    if (PotentialConfigType != typeof(GUBPNodeAdder) && typeof(GUBPNodeAdder).IsAssignableFrom(PotentialConfigType))
+                    {
+                        GUBPNodeAdder Config = Activator.CreateInstance(PotentialConfigType) as GUBPNodeAdder;
+                        if (Config != null)
+                        {
+                            Adders.Add(Config);
+                        }
+                    }
+                }
+            }
+        }
+        foreach(var Adder in Adders)
+        {
+            Adder.AddNodes(this, InHostPlatform);
+        }
+    }
+
+    public abstract class GUBPBranchHacker
+    {
+        public class BranchOptions
+        {            
+            public List<UnrealTargetPlatform> PlatformsToRemove = new List<UnrealTargetPlatform>();
+			public List<string> ExcludeNodes = new List<string>();
+			public List<UnrealTargetPlatform> ExcludePlatformsForEditor = new List<UnrealTargetPlatform>();
+            public List<UnrealTargetPlatform> RemovePlatformFromPromotable = new List<UnrealTargetPlatform>();
+            public List<string> ProjectsToCook = new List<string>();
+            public List<string> PromotablesWithoutTools = new List<string>();
+            public List<string> NodesToRemovePseudoDependencies = new List<string>();
+            public List<string> EnhanceAgentRequirements = new List<string>();
+			public bool bNoAutomatedTesting = false;
+			public bool bNoDocumentation = false;
+			public bool bNoInstalledEngine = false;
+			public bool bMakeFormalBuildWithoutLabelPromotable = false;
+			public bool bNoMonolithicDependenciesForCooks = false;
+			public bool bNoEditorDependenciesForTools = false;
+			public Dictionary<string, sbyte> FrequencyBarriers = new Dictionary<string,sbyte>();
+			public int QuantumOverride = 0;
+        }
+        public virtual void ModifyOptions(GUBP bp, ref BranchOptions Options, string Branch)
+        {
+        }
+    }
+
+    private static List<GUBPBranchHacker> BranchHackers;
+    private GUBPBranchHacker.BranchOptions GetBranchOptions(string Branch)
+    {
+        if (BranchHackers == null)
+        {
+            BranchHackers = new List<GUBPBranchHacker>();
+            Assembly[] LoadedAssemblies = AppDomain.CurrentDomain.GetAssemblies();
+            foreach (var Dll in LoadedAssemblies)
+            {
+				Type[] AllTypes = GetTypesFromAssembly(Dll);
+                foreach (var PotentialConfigType in AllTypes)
+                {
+                    if (PotentialConfigType != typeof(GUBPBranchHacker) && typeof(GUBPBranchHacker).IsAssignableFrom(PotentialConfigType))
+                    {
+                        GUBPBranchHacker Config = Activator.CreateInstance(PotentialConfigType) as GUBPBranchHacker;
+                        if (Config != null)
+                        {
+                            BranchHackers.Add(Config);
+                        }
+                    }
+                }
+            }
+        }
+        var Result = new GUBPBranchHacker.BranchOptions();
+        foreach (var Hacker in BranchHackers)
+        {
+            Hacker.ModifyOptions(this, ref Result, Branch);
+        }
+        return Result;
+    }
+
+    public abstract class GUBPEmailHacker
+    {        
+        public virtual List<string> AddEmails(GUBP bp, string Branch, string NodeName)
+        {
+            return new List<string>();
+        }
+        public virtual List<string> ModifyEmail(string Email, GUBP bp, string Branch, string NodeName)
+        {
+            return new List<string>{Email};
+        }
+        public virtual List<string> FinalizeEmails(List<string> Emails, GUBP bp, string Branch, string NodeName)
+        {
+            return Emails;
+        }
+		public virtual bool VetoEmailingCausers(GUBP bp, string Branch, string NodeName)
+		{
+			return false; // People who have submitted since last-green will be included unless vetoed by overriding this method. 
+		}
+    }
+
+    private static List<GUBPEmailHacker> EmailHackers;
+    private string HackEmails(string Emails, string Causers, string Branch, string NodeName)
+    {
+        string OnlyEmail = ParseParamValue("OnlyEmail");
+        if (!String.IsNullOrEmpty(OnlyEmail))
+        {
+            return OnlyEmail;
+        }
+        EmailHint = ParseParamValue("EmailHint");
+        if(EmailHint == null)
+        {
+            EmailHint = "";
+        }			
+        if (EmailHackers == null)
+        {
+            EmailHackers = new List<GUBPEmailHacker>();
+            Assembly[] LoadedAssemblies = AppDomain.CurrentDomain.GetAssemblies();
+            foreach (var Dll in LoadedAssemblies)
+            {
+				Type[] AllTypes = GetTypesFromAssembly(Dll);
+                foreach (var PotentialConfigType in AllTypes)
+                {
+                    if (PotentialConfigType != typeof(GUBPEmailHacker) && typeof(GUBPEmailHacker).IsAssignableFrom(PotentialConfigType))
+                    {
+                        GUBPEmailHacker Config = Activator.CreateInstance(PotentialConfigType) as GUBPEmailHacker;
+                        if (Config != null)
+                        {
+                            EmailHackers.Add(Config);
+                        }
+                    }
+                }
+            }
+        }
+        List<string> Result = new List<string>(Emails.Split(' '));
+		if(!EmailHackers.Any(x => x.VetoEmailingCausers(this, Branch, NodeName)))
+		{
+			Result.AddRange(Causers.Split(' '));
+		}
+        foreach (var EmailHacker in EmailHackers)
+        {
+            Result.AddRange(EmailHacker.AddEmails(this, Branch, NodeName));
+        }
+        foreach (var EmailHacker in EmailHackers)
+        {
+            var NewResult = new List<string>();
+            foreach (var Email in Result)
+            {
+                NewResult.AddRange(EmailHacker.ModifyEmail(Email, this, Branch, NodeName));
+            }
+            Result = NewResult;
+        }
+        foreach (var EmailHacker in EmailHackers)
+        {
+            Result = EmailHacker.FinalizeEmails(Result, this, Branch, NodeName);
+        }
+        string FinalEmails = "";
+        int Count = 0;
+        foreach (var Email in Result)
+        {
+            FinalEmails = GUBPNode.MergeSpaceStrings(FinalEmails, Email);
+            Count++;
+        }
+        return FinalEmails;
+    }
+    public abstract class GUBPFrequencyHacker
+    {
+        public virtual int GetNodeFrequency(GUBP bp, string Branch, string NodeName, int BaseFrequency)
+        {
+            return new int();
+        }
+    }
+    private static List<GUBPFrequencyHacker> FrequencyHackers;
+    private int HackFrequency(GUBP bp, string Branch, string NodeName, int BaseFrequency)
+    {
+        int Frequency = BaseFrequency;
+        if (FrequencyHackers == null)
+        {
+            FrequencyHackers = new List<GUBPFrequencyHacker>();
+            Assembly[] LoadedAssemblies = AppDomain.CurrentDomain.GetAssemblies();
+            foreach (var Dll in LoadedAssemblies)
+            {
+				Type[] AllTypes = GetTypesFromAssembly(Dll);
+                foreach (var PotentialConfigType in AllTypes)
+                {
+                    if (PotentialConfigType != typeof(GUBPFrequencyHacker) && typeof(GUBPFrequencyHacker).IsAssignableFrom(PotentialConfigType))
+                    {
+                        GUBPFrequencyHacker Config = Activator.CreateInstance(PotentialConfigType) as GUBPFrequencyHacker;
+                        if (Config != null)
+                        {
+                            FrequencyHackers.Add(Config);
+                        }
+                    }
+                }
+            }
+        }
+        foreach(var FrequencyHacker in FrequencyHackers)
+        {
+           Frequency = FrequencyHacker.GetNodeFrequency(bp, Branch, NodeName, BaseFrequency);            
+        }
+        return Frequency;
+    }    
+
+	void AddNodesForBranch(bool bNoAutomatedTesting)
+	{
+        AddNode(new VersionFilesNode());
+
+
+        foreach (var HostPlatform in HostPlatforms)
+        {
+			AddNode(new ToolsForCompileNode(HostPlatform));
+			
+			if (!BranchOptions.ExcludePlatformsForEditor.Contains(HostPlatform))
+			{
+				AddNode(new RootEditorNode(HostPlatform));			
+				AddNode(new ToolsNode(this, HostPlatform));            
+				AddNode(new InternalToolsNode(this, HostPlatform));
+			if (HostPlatform == UnrealTargetPlatform.Win64 && ActivePlatforms.Contains(UnrealTargetPlatform.Linux))
+				{
+					if (!BranchOptions.ExcludePlatformsForEditor.Contains(UnrealTargetPlatform.Linux))
+					{
+						AddNode(new ToolsCrossCompileNode(HostPlatform));
+					}
+				}
+				foreach (var ProgramTarget in Branch.BaseEngineProject.Properties.Programs)
+				{
+                    if (!BranchOptions.ExcludeNodes.Contains(ProgramTarget.TargetName))
+                    {
+                        bool bInternalOnly;
+                        bool SeparateNode;
+                        bool CrossCompile;
+
+                        if (ProgramTarget.Rules.GUBP_AlwaysBuildWithTools(HostPlatform, out bInternalOnly, out SeparateNode, out CrossCompile) && ProgramTarget.Rules.SupportsPlatform(HostPlatform) && SeparateNode)
+                        {
+                            if (bInternalOnly)
+                            {
+                                AddNode(new SingleInternalToolsNode(HostPlatform, ProgramTarget));
+                            }
+                            else
+                            {
+                                AddNode(new SingleToolsNode(HostPlatform, ProgramTarget));
+                            }
+                        }
+                        if (ProgramTarget.Rules.GUBP_IncludeNonUnityToolTest())
+                        {
+                            AddNode(new NonUnityToolNode(HostPlatform, ProgramTarget));
+                        }
+                    }
+				}
+				foreach(var CodeProj in Branch.CodeProjects)
+				{
+					foreach(var ProgramTarget in CodeProj.Properties.Programs)
+					{
+                        if (!BranchOptions.ExcludeNodes.Contains(ProgramTarget.TargetName))
+                        {
+                            bool bInternalNodeOnly;
+                            bool SeparateNode;
+                            bool CrossCompile;
+
+                            if (ProgramTarget.Rules.GUBP_AlwaysBuildWithTools(HostPlatform, out bInternalNodeOnly, out SeparateNode, out CrossCompile) && ProgramTarget.Rules.SupportsPlatform(HostPlatform) && SeparateNode)
+                            {
+                                if (bInternalNodeOnly)
+                                {
+                                    AddNode(new SingleInternalToolsNode(HostPlatform, ProgramTarget));
+                                }
+                                else
+                                {
+                                    AddNode(new SingleToolsNode(HostPlatform, ProgramTarget));
+                                }
+                            }
+                            if (ProgramTarget.Rules.GUBP_IncludeNonUnityToolTest())
+                            {
+                                AddNode(new NonUnityToolNode(HostPlatform, ProgramTarget));
+                            }
+                        }
+					}
+				}
+
+				AddNode(new EditorAndToolsNode(this, HostPlatform));
+
+				if (bOrthogonalizeEditorPlatforms)
+				{
+					foreach (var Plat in ActivePlatforms)
+					{
+						if (Plat != HostPlatform && Plat != GetAltHostPlatform(HostPlatform))
+						{
+							if (Platform.Platforms[HostPlatform].CanHostPlatform(Plat))
+							{
+								AddNode(new EditorPlatformNode(HostPlatform, Plat));
+							}
+						}
+					}
+				}
+			}
+
+            bool DoASharedPromotable = false;
+
+            int NumSharedCode = 0;
+            foreach (var CodeProj in Branch.CodeProjects)
+            {
+                var Options = CodeProj.Options(HostPlatform);
+
+                if (Options.bIsPromotable && !Options.bSeparateGamePromotion)
+                {
+                    NumSharedCode++;
+                }
+            }
+
+            var NonCodeProjectNames = new Dictionary<string, List<UnrealTargetPlatform>>();
+            var NonCodeFormalBuilds = new Dictionary<string, List<TargetRules.GUBPFormalBuild>>();
+            {
+                var Target = Branch.BaseEngineProject.Properties.Targets[TargetRules.TargetType.Editor];
+
+                foreach (var Codeless in Target.Rules.GUBP_NonCodeProjects_BaseEditorTypeOnly(HostPlatform))
+                {
+                    var Proj = Branch.FindGame(Codeless.Key);
+                    if (Proj == null)
+                    {
+                        Log(System.Diagnostics.TraceEventType.Verbose, "{0} was listed as a codeless project by GUBP_NonCodeProjects_BaseEditorTypeOnly, however it does not exist in this branch.", Codeless.Key);
+                    }
+                    else if (Proj.Properties.bIsCodeBasedProject)
+                    {
+						if (!Branch.NonCodeProjects.Contains(Proj)) 
+						{ 
+							Branch.NonCodeProjects.Add(Proj);
+							NonCodeProjectNames.Add(Codeless.Key, Codeless.Value);
+						}
+                    }
+                    else
+                    {
+                        NonCodeProjectNames.Add(Codeless.Key, Codeless.Value);
+                    }
+                }
+
+                var TempNonCodeFormalBuilds = Target.Rules.GUBP_GetNonCodeFormalBuilds_BaseEditorTypeOnly();
+                var HostMonos = GetMonolithicPlatformsForUProject(HostPlatform, Branch.BaseEngineProject, true);
+
+                foreach (var Codeless in TempNonCodeFormalBuilds)
+                {
+                    if (NonCodeProjectNames.ContainsKey(Codeless.Key))
+                    {
+                        var PlatList = Codeless.Value;
+                        var NewPlatList = new List<TargetRules.GUBPFormalBuild>();
+                        foreach (var PlatPair in PlatList)
+                        {
+                            if (HostMonos.Contains(PlatPair.TargetPlatform))
+                            {
+                                NewPlatList.Add(PlatPair);
+                            }
+                        }
+                        if (NewPlatList.Count > 0)
+                        {
+                            NonCodeFormalBuilds.Add(Codeless.Key, NewPlatList);
+                        }
+                    }
+                    else
+                    {
+                        Log(System.Diagnostics.TraceEventType.Verbose, "{0} was listed as a codeless formal build GUBP_GetNonCodeFormalBuilds_BaseEditorTypeOnly, however it does not exist in this branch.", Codeless.Key);
+                    }
+                }
+            }
+
+            DoASharedPromotable = NumSharedCode > 0 || NonCodeProjectNames.Count > 0 || NonCodeFormalBuilds.Count > 0;
+
+			if (!BranchOptions.ExcludePlatformsForEditor.Contains(HostPlatform))
+			{
+			AddNode(new NonUnityTestNode(HostPlatform));
+			}
+
+            if (DoASharedPromotable)
+            {
+                var AgentSharingGroup = "Shared_EditorTests" + HostPlatformNode.StaticGetHostPlatformSuffix(HostPlatform);
+
+                var Options = Branch.BaseEngineProject.Options(HostPlatform);
+
+                if (!Options.bIsPromotable || Options.bSeparateGamePromotion)
+                {
+                    throw new AutomationException("We assume that if we have shared promotable, the base engine is in it.");
+                }
+
+
+				if (HostPlatform == UnrealTargetPlatform.Win64) //temp hack till automated testing works on other platforms than Win64
+                {
+
+                    var EditorTests = Branch.BaseEngineProject.Properties.Targets[TargetRules.TargetType.Editor].Rules.GUBP_GetEditorTests_EditorTypeOnly(HostPlatform);
+                    var EditorTestNodes = new List<string>();
+                    foreach (var Test in EditorTests)
+                    {
+                        if (!bNoAutomatedTesting)
+                        {							
+                            EditorTestNodes.Add(AddNode(new UATTestNode(this, HostPlatform, Branch.BaseEngineProject, Test.Key, Test.Value, AgentSharingGroup)));
+                        
+                            foreach (var NonCodeProject in Branch.NonCodeProjects)
+                            {
+                                if (!NonCodeProjectNames.ContainsKey(NonCodeProject.GameName))
+                                {
+                                    continue;
+                                }
+                                if (HostPlatform == UnrealTargetPlatform.Mac) continue; //temp hack till mac automated testing works
+                                EditorTestNodes.Add(AddNode(new UATTestNode(this, HostPlatform, NonCodeProject, Test.Key, Test.Value, AgentSharingGroup)));
+                            }
+                        }
+                    }
+                    if (EditorTestNodes.Count > 0)
+                    {
+                        AddNode(new GameAggregateNode(this, HostPlatform, Branch.BaseEngineProject, "AllEditorTests", EditorTestNodes, 0.0f));
+                    }
+                }
+
+
+                var ServerPlatforms = new List<UnrealTargetPlatform>();
+                var GamePlatforms = new List<UnrealTargetPlatform>();
+
+                foreach (var Kind in BranchInfo.MonolithicKinds)
+                {
+                    if (Branch.BaseEngineProject.Properties.Targets.ContainsKey(Kind))
+                    {
+                        var Target = Branch.BaseEngineProject.Properties.Targets[Kind];
+                        var Platforms = Target.Rules.GUBP_GetPlatforms_MonolithicOnly(HostPlatform);
+                        if (Platforms.Contains(HostPlatform))
+                        {
+                            // we want the host platform first since some some pseudodependencies look to see if the shared promotable exists.
+                            Platforms.Remove(HostPlatform);
+                            Platforms.Insert(0, HostPlatform);
+                        }
+                        foreach (var Plat in Platforms)
+                        {
+                            if (!Platform.Platforms[HostPlatform].CanHostPlatform(Plat))
+                            {
+                                throw new AutomationException("Project {0} asked for platform {1} with host {2}, but the host platform cannot build that platform.", Branch.BaseEngineProject.GameName, Plat.ToString(), HostPlatform.ToString());
+                            }
+                            if (bNoIOSOnPC && Plat == UnrealTargetPlatform.IOS && HostPlatform == UnrealTargetPlatform.Win64)
+                            {
+                                continue;
+                            }
+
+                            if (ActivePlatforms.Contains(Plat))
+                            {
+                                if (Kind == TargetRules.TargetType.Server && !ServerPlatforms.Contains(Plat))
+                                {
+                                    ServerPlatforms.Add(Plat);
+                                }
+                                if (Kind == TargetRules.TargetType.Game && !GamePlatforms.Contains(Plat))
+                                {
+                                    GamePlatforms.Add(Plat);
+                                }
+                                if (!GUBPNodes.ContainsKey(GamePlatformMonolithicsNode.StaticGetFullName(HostPlatform, Branch.BaseEngineProject, Plat)))
+                                {
+									if(GamePlatformMonolithicsNode.HasPrecompiledTargets(Branch.BaseEngineProject, HostPlatform, Plat))
+									{
+	                                    AddNode(new GamePlatformMonolithicsNode(this, HostPlatform, Branch.BaseEngineProject, Plat, InPrecompiled: true));
+									}
+                                    AddNode(new GamePlatformMonolithicsNode(this, HostPlatform, Branch.BaseEngineProject, Plat));
+                                }
+								if (Plat == UnrealTargetPlatform.Win32 && Target.Rules.GUBP_BuildWindowsXPMonolithics() && Kind == TargetRules.TargetType.Game)
+								{
+									if (!GUBPNodes.ContainsKey(GamePlatformMonolithicsNode.StaticGetFullName(HostPlatform, Branch.BaseEngineProject, Plat, true)))
+									{
+										AddNode(new GamePlatformMonolithicsNode(this, HostPlatform, Branch.BaseEngineProject, Plat, true));
+									}
+	                            }
+                            }
+                        }
+                    }
+                }
+
+                var CookedAgentSharingGroup = "Shared_CookedTests" + HostPlatformNode.StaticGetHostPlatformSuffix(HostPlatform);
+
+                var GameTestNodes = new List<string>();
+                var GameCookNodes = new List<string>();
+				//var FormalAgentSharingGroup = "Shared_FormalBuilds" + HostPlatformNode.StaticGetHostPlatformSuffix(HostPlatform);
+
+
+                //foreach (var Kind in BranchInfo.MonolithicKinds)//for now, non-code projects don't do client or server.
+                {
+                    var Kind = TargetRules.TargetType.Game;
+                    if (Branch.BaseEngineProject.Properties.Targets.ContainsKey(Kind))
+                    {
+                        var Target = Branch.BaseEngineProject.Properties.Targets[Kind];
+                        var Platforms = Target.Rules.GUBP_GetPlatforms_MonolithicOnly(HostPlatform);
+                        foreach (var Plat in Platforms)
+                        {
+                            if (!Platform.Platforms[HostPlatform].CanHostPlatform(Plat))
+                            {
+                                throw new AutomationException("Project {0} asked for platform {1} with host {2}, but the host platform cannot build that platform.", Branch.BaseEngineProject.GameName, Plat.ToString(), HostPlatform.ToString());
+                            }
+                            if (bNoIOSOnPC && Plat == UnrealTargetPlatform.IOS && HostPlatform == UnrealTargetPlatform.Win64)
+                            {
+                                continue;
+                            }
+
+                            if (ActivePlatforms.Contains(Plat))
+                            {
+                                string CookedPlatform = Platform.Platforms[Plat].GetCookPlatform(Kind == TargetRules.TargetType.Server, Kind == TargetRules.TargetType.Client, "");								
+                                if (!GUBPNodes.ContainsKey(CookNode.StaticGetFullName(HostPlatform, Branch.BaseEngineProject, CookedPlatform)))
+                                {
+                                    GameCookNodes.Add(AddNode(new CookNode(this, HostPlatform, Branch.BaseEngineProject, Plat, CookedPlatform)));
+                                }
+                                if (!GUBPNodes.ContainsKey(GamePlatformCookedAndCompiledNode.StaticGetFullName(HostPlatform, Branch.BaseEngineProject, Plat)))
+                                {
+                                    AddNode(new GamePlatformCookedAndCompiledNode(this, HostPlatform, Branch.BaseEngineProject, Plat, false));
+                                }
+                                var GameTests = Target.Rules.GUBP_GetGameTests_MonolithicOnly(HostPlatform, GetAltHostPlatform(HostPlatform), Plat);
+                                var RequiredPlatforms = new List<UnrealTargetPlatform> { Plat };
+                                if (!bNoAutomatedTesting)
+                                {
+                                    var ThisMonoGameTestNodes = new List<string>();	
+                                    
+                                    foreach (var Test in GameTests)
+                                    {
+                                        var TestName = Test.Key + "_" + Plat.ToString();
+                                        ThisMonoGameTestNodes.Add(AddNode(new UATTestNode(this, HostPlatform, Branch.BaseEngineProject, TestName, Test.Value, CookedAgentSharingGroup, false, RequiredPlatforms)));
+                                    }
+                                    if (ThisMonoGameTestNodes.Count > 0)
+                                    {										
+                                        GameTestNodes.Add(AddNode(new GameAggregateNode(this, HostPlatform, Branch.BaseEngineProject, "CookedTests_" + Plat.ToString() + "_" + Kind.ToString() + HostPlatformNode.StaticGetHostPlatformSuffix(HostPlatform), ThisMonoGameTestNodes, 0.0f)));
+                                    }
+                                }
+
+                                foreach (var NonCodeProject in Branch.NonCodeProjects)
+                                {
+                                    if (!NonCodeProjectNames.ContainsKey(NonCodeProject.GameName) || !NonCodeProjectNames[NonCodeProject.GameName].Contains(Plat))
+                                    {
+                                        continue;
+                                    }
+                                    if (BranchOptions.ProjectsToCook.Contains(NonCodeProject.GameName) || BranchOptions.ProjectsToCook.Count == 0)
+                                    {
+                                    if (!GUBPNodes.ContainsKey(CookNode.StaticGetFullName(HostPlatform, NonCodeProject, CookedPlatform)))
+                                    {
+                                        GameCookNodes.Add(AddNode(new CookNode(this, HostPlatform, NonCodeProject, Plat, CookedPlatform)));
+                                    }
+                                    if (!GUBPNodes.ContainsKey(GamePlatformCookedAndCompiledNode.StaticGetFullName(HostPlatform, NonCodeProject, Plat)))
+                                    {
+                                        AddNode(new GamePlatformCookedAndCompiledNode(this, HostPlatform, NonCodeProject, Plat, false));
+
+                                        if (NonCodeFormalBuilds.ContainsKey(NonCodeProject.GameName))
+                                        {
+                                            var PlatList = NonCodeFormalBuilds[NonCodeProject.GameName];
+                                            foreach (var PlatPair in PlatList)
+                                            {
+                                                if (PlatPair.TargetPlatform == Plat)
+                                                {
+														var NodeName = AddNode(new FormalBuildNode(this, NonCodeProject, HostPlatform, new List<UnrealTargetPlatform>() { Plat }, new List<UnrealTargetConfiguration>() { PlatPair.TargetConfig }));                                                    
+													// we don't want this delayed
+													// this would normally wait for the testing phase, we just want to build it right away
+													RemovePseudodependencyFromNode(
+														CookNode.StaticGetFullName(HostPlatform, NonCodeProject, CookedPlatform),
+														WaitForTestShared.StaticGetFullName());
+													string BuildAgentSharingGroup = "";
+													if (Options.bSeparateGamePromotion)
+													{
+														BuildAgentSharingGroup = NonCodeProject.GameName + "_MakeFormalBuild_" + Plat.ToString() + HostPlatformNode.StaticGetHostPlatformSuffix(HostPlatform);
+														if (Plat == UnrealTargetPlatform.IOS || Plat == UnrealTargetPlatform.Android) // These trash build products, so we need to use different agents
+														{
+															BuildAgentSharingGroup = "";
+														}
+														GUBPNodes[CookNode.StaticGetFullName(HostPlatform, NonCodeProject, CookedPlatform)].AgentSharingGroup = BuildAgentSharingGroup;
+														GUBPNodes[GamePlatformCookedAndCompiledNode.StaticGetFullName(HostPlatform, NonCodeProject, Plat)].AgentSharingGroup = BuildAgentSharingGroup;
+														GUBPNodes[NodeName].AgentSharingGroup = BuildAgentSharingGroup;
+													}
+													else
+													{
+														//GUBPNodes[NodeName].AgentSharingGroup = FormalAgentSharingGroup;
+															if (Plat == UnrealTargetPlatform.XboxOne)
+														{
+															GUBPNodes[NodeName].AgentSharingGroup = "";
+														}
+													}
+													if (PlatPair.bTest)
+													{
+														AddNode(new FormalBuildTestNode(this, NonCodeProject, HostPlatform, Plat, PlatPair.TargetConfig));
+												}
+											}
+										}
+									}
+                                        }
+
+                                    if (!bNoAutomatedTesting)
+                                        {
+										if (HostPlatform == UnrealTargetPlatform.Mac || HostPlatform == UnrealTargetPlatform.Linux) continue; //temp hack till Linux and Mac automated testing works
+										var ThisMonoGameTestNodes = new List<string>();
+										foreach (var Test in GameTests)
+										{
+											var TestName = Test.Key + "_" + Plat.ToString();
+											ThisMonoGameTestNodes.Add(AddNode(new UATTestNode(this, HostPlatform, NonCodeProject, TestName, Test.Value, CookedAgentSharingGroup, false, RequiredPlatforms)));
+										}
+										if (ThisMonoGameTestNodes.Count > 0)
+                                            {
+											GameTestNodes.Add(AddNode(new GameAggregateNode(this, HostPlatform, NonCodeProject, "CookedTests_" + Plat.ToString() + "_" + Kind.ToString() + HostPlatformNode.StaticGetHostPlatformSuffix(HostPlatform), ThisMonoGameTestNodes, 0.0f)));
+										}
+									}
+								}
+							}
+						}
+					}
+                }
+                }
+#if false
+                //for now, non-code projects don't do client or server.
+                foreach (var ServerPlatform in ServerPlatforms)
+                {
+                    var ServerTarget = Branch.BaseEngineProject.Properties.Targets[TargetRules.TargetType.Server];
+                    foreach (var GamePlatform in GamePlatforms)
+                    {
+                        var Target = Branch.BaseEngineProject.Properties.Targets[TargetRules.TargetType.Game];
+
+                        foreach (var NonCodeProject in Branch.NonCodeProjects)
+                        {
+                            if (!NonCodeProjectNames.ContainsKey(NonCodeProject.GameName) || !NonCodeProjectNames.ContainsKey(NonCodeProject.GameName) ||
+                                    !NonCodeProjectNames[NonCodeProject.GameName].Contains(ServerPlatform)  || !NonCodeProjectNames[NonCodeProject.GameName].Contains(GamePlatform) )
+                            {
+                                continue;
+                            }
+
+                            var ClientServerTests = Target.Rules.GUBP_GetClientServerTests_MonolithicOnly(HostPlatform, GetAltHostPlatform(HostPlatform), ServerPlatform, GamePlatform);
+                            var RequiredPlatforms = new List<UnrealTargetPlatform> { ServerPlatform };
+                            if (ServerPlatform != GamePlatform)
+                            {
+                                RequiredPlatforms.Add(GamePlatform);
+                            }
+                            foreach (var Test in ClientServerTests)
+                            {
+                                GameTestNodes.Add(AddNode(new UATTestNode(this, HostPlatform, NonCodeProject, Test.Key + "_" + GamePlatform.ToString() + "_" + ServerPlatform.ToString(), Test.Value, false, RequiredPlatforms, true)));
+                            }
+                        }
+                    }
+                }
+#endif
+                if (GameTestNodes.Count > 0)
+                {
+                    AddNode(new GameAggregateNode(this, HostPlatform, Branch.BaseEngineProject, "AllCookedTests", GameTestNodes));
+                }
+            }
+			
+			if(HostPlatform == MakeFeaturePacksNode.GetDefaultBuildPlatform(this) && !BranchOptions.bNoInstalledEngine)
+			{
+				AddNode(new MakeFeaturePacksNode(HostPlatform, Branch.AllProjects.Where(x => MakeFeaturePacksNode.IsFeaturePack(x))));
+			}
+
+            foreach (var CodeProj in Branch.CodeProjects)
+            {
+                var Options = CodeProj.Options(HostPlatform);
+
+                if (!Options.bIsPromotable && !Options.bTestWithShared && !Options.bIsNonCode && !Options.bBuildAnyway)
+                {
+                    continue; // we skip things that aren't promotable and aren't tested - except noncode as code situations
+                }
+                var AgentShareName = CodeProj.GameName;
+                if (!Options.bSeparateGamePromotion)
+                {
+                    AgentShareName = "Shared";
+                }
+
+				if (!BranchOptions.ExcludePlatformsForEditor.Contains(HostPlatform) && !Options.bIsNonCode)
+				{
+					EditorGameNode Node = (EditorGameNode)TryFindNode(EditorGameNode.StaticGetFullName(HostPlatform, CodeProj));
+					if(Node == null)
+					{
+					AddNode(new EditorGameNode(this, HostPlatform, CodeProj));
+				}
+					else
+					{
+						Node.AddProject(CodeProj);
+					}
+				}
+
+				if (!bNoAutomatedTesting && HostPlatform == UnrealTargetPlatform.Win64) //temp hack till automated testing works on other platforms than Win64
+                {
+					if (CodeProj.Properties.Targets.ContainsKey(TargetRules.TargetType.Editor))
+					{
+						var EditorTests = CodeProj.Properties.Targets[TargetRules.TargetType.Editor].Rules.GUBP_GetEditorTests_EditorTypeOnly(HostPlatform);
+						var EditorTestNodes = new List<string>();
+						string AgentSharingGroup = "";
+						if (EditorTests.Count > 1)
+						{
+							AgentSharingGroup = AgentShareName + "_EditorTests" + HostPlatformNode.StaticGetHostPlatformSuffix(HostPlatform);
+						}
+						foreach (var Test in EditorTests)
+						{
+							EditorTestNodes.Add(AddNode(new UATTestNode(this, HostPlatform, CodeProj, Test.Key, Test.Value, AgentSharingGroup)));
+							if (!Options.bTestWithShared || !HasNode(WaitForTestShared.StaticGetFullName()))
+							{
+								RemovePseudodependencyFromNode((UATTestNode.StaticGetFullName(HostPlatform, CodeProj, Test.Key)), WaitForTestShared.StaticGetFullName());
+							}
+						}
+						if (EditorTestNodes.Count > 0)
+						{
+							AddNode(new GameAggregateNode(this, HostPlatform, CodeProj, "AllEditorTests", EditorTestNodes, 0.0f));
+						}
+					}
+                }
+
+                var CookedAgentSharingGroup = AgentShareName + "_CookedTests" + HostPlatformNode.StaticGetHostPlatformSuffix(HostPlatform);
+				//var FormalAgentSharingGroup = "Shared_FormalBuilds" + HostPlatformNode.StaticGetHostPlatformSuffix(HostPlatform);
+                var ServerPlatforms = new List<UnrealTargetPlatform>();
+                var GamePlatforms = new List<UnrealTargetPlatform>();
+                var GameTestNodes = new List<string>();
+                foreach (var Kind in BranchInfo.MonolithicKinds)
+                {
+                    if (CodeProj.Properties.Targets.ContainsKey(Kind))
+                    {
+                        var Target = CodeProj.Properties.Targets[Kind];
+                        var Platforms = Target.Rules.GUBP_GetPlatforms_MonolithicOnly(HostPlatform);
+						var AdditionalPlatforms = Target.Rules.GUBP_GetBuildOnlyPlatforms_MonolithicOnly(HostPlatform);
+						var AllPlatforms = Platforms.Union(AdditionalPlatforms);
+						foreach (var Plat in AllPlatforms)
+                        {
+                            if (!Platform.Platforms[HostPlatform].CanHostPlatform(Plat))
+                            {
+                                throw new AutomationException("Project {0} asked for platform {1} with host {2}, but the host platform cannot build that platform.", CodeProj.GameName, Plat.ToString(), HostPlatform.ToString());
+                            }
+                            if (bNoIOSOnPC && Plat == UnrealTargetPlatform.IOS && HostPlatform == UnrealTargetPlatform.Win64)
+                            {
+                                continue;
+                            }
+							if(Plat == UnrealTargetPlatform.Win32 && Target.Rules.GUBP_BuildWindowsXPMonolithics())
+							{
+								if(!GUBPNodes.ContainsKey(GamePlatformMonolithicsNode.StaticGetFullName(HostPlatform, CodeProj, Plat, true)))
+								{
+									AddNode(new GamePlatformMonolithicsNode(this, HostPlatform, CodeProj, Plat, true));
+								}
+							}
+                            if (ActivePlatforms.Contains(Plat))
+                            {
+                                if (Kind == TargetRules.TargetType.Server && !ServerPlatforms.Contains(Plat))
+                                {
+                                    ServerPlatforms.Add(Plat);
+                                }
+                                if (Kind == TargetRules.TargetType.Game && !GamePlatforms.Contains(Plat))
+                                {
+                                    GamePlatforms.Add(Plat);
+                                }
+                                if (!GUBPNodes.ContainsKey(GamePlatformMonolithicsNode.StaticGetFullName(HostPlatform, CodeProj, Plat)))
+                                {
+									if(GamePlatformMonolithicsNode.HasPrecompiledTargets(CodeProj, HostPlatform, Plat))
+									{
+	                                    AddNode(new GamePlatformMonolithicsNode(this, HostPlatform, CodeProj, Plat, InPrecompiled: true));
+									}
+                                    AddNode(new GamePlatformMonolithicsNode(this, HostPlatform, CodeProj, Plat));
+                                }
+                                var FormalBuildConfigs = Target.Rules.GUBP_GetConfigsForFormalBuilds_MonolithicOnly(HostPlatform);
+								if (!AdditionalPlatforms.Contains(Plat) && (BranchOptions.ProjectsToCook.Contains(CodeProj.GameName) || BranchOptions.ProjectsToCook.Count == 0))
+								{
+									string CookedPlatform = Platform.Platforms[Plat].GetCookPlatform(Kind == TargetRules.TargetType.Server, Kind == TargetRules.TargetType.Client, "");
+									if (Target.Rules.GUBP_AlternateCookPlatform(HostPlatform, CookedPlatform) != "")
+									{
+										CookedPlatform = Target.Rules.GUBP_AlternateCookPlatform(HostPlatform, CookedPlatform);
+									}
+									if (!GUBPNodes.ContainsKey(CookNode.StaticGetFullName(HostPlatform, CodeProj, CookedPlatform)))
+									{
+										AddNode(new CookNode(this, HostPlatform, CodeProj, Plat, CookedPlatform));
+									}
+									if (!GUBPNodes.ContainsKey(GamePlatformCookedAndCompiledNode.StaticGetFullName(HostPlatform, CodeProj, Plat)))
+									{
+										AddNode(new GamePlatformCookedAndCompiledNode(this, HostPlatform, CodeProj, Plat, true));
+									}
+									
+
+									foreach (var Config in FormalBuildConfigs)
+									{
+										string FormalNodeName = null;                                    
+										if (Kind == TargetRules.TargetType.Client)
+										{
+											if (Plat == Config.TargetPlatform)
+											{
+												FormalNodeName = AddNode(new FormalBuildNode(this, CodeProj, HostPlatform, InClientTargetPlatforms: new List<UnrealTargetPlatform>() { Config.TargetPlatform }, InClientConfigs: new List<UnrealTargetConfiguration>() { Config.TargetConfig }, InClientNotGame: true));
+											}
+										}
+										else if (Kind == TargetRules.TargetType.Server)
+										{
+											if (Plat == Config.TargetPlatform)
+											{
+												FormalNodeName = AddNode(new FormalBuildNode(this, CodeProj, HostPlatform, InServerTargetPlatforms: new List<UnrealTargetPlatform>() { Config.TargetPlatform }, InServerConfigs: new List<UnrealTargetConfiguration>() { Config.TargetConfig }));
+											}
+										}
+										else if (Kind == TargetRules.TargetType.Game)
+										{
+											if (Plat == Config.TargetPlatform)
+											{
+												FormalNodeName = AddNode(new FormalBuildNode(this, CodeProj, HostPlatform, InClientTargetPlatforms: new List<UnrealTargetPlatform>() { Config.TargetPlatform }, InClientConfigs: new List<UnrealTargetConfiguration>() { Config.TargetConfig }));
+											}
+										}
+										if (FormalNodeName != null)
+										{
+											// we don't want this delayed
+											// this would normally wait for the testing phase, we just want to build it right away
+											RemovePseudodependencyFromNode(
+												CookNode.StaticGetFullName(HostPlatform, CodeProj, CookedPlatform),
+												WaitForTestShared.StaticGetFullName());
+											string BuildAgentSharingGroup = "";
+											
+											if (Options.bSeparateGamePromotion)
+											{
+												BuildAgentSharingGroup = CodeProj.GameName + "_MakeFormalBuild_" + Plat.ToString() + HostPlatformNode.StaticGetHostPlatformSuffix(HostPlatform);
+												if (Plat == UnrealTargetPlatform.IOS || Plat == UnrealTargetPlatform.Android || Plat == UnrealTargetPlatform.XboxOne) // These trash build products, so we need to use different agents
+												{
+													BuildAgentSharingGroup = "";
+												}
+												GUBPNodes[CookNode.StaticGetFullName(HostPlatform, CodeProj, CookedPlatform)].AgentSharingGroup = BuildAgentSharingGroup;
+												GUBPNodes[GamePlatformCookedAndCompiledNode.StaticGetFullName(HostPlatform, CodeProj, Plat)].AgentSharingGroup = BuildAgentSharingGroup;
+												GUBPNodes[FormalNodeName].AgentSharingGroup = BuildAgentSharingGroup;
+											}
+											else
+											{
+												//GUBPNodes[FormalNodeName].AgentSharingGroup = FormalAgentSharingGroup;
+												if (Plat == UnrealTargetPlatform.XboxOne)
+												{
+													GUBPNodes[FormalNodeName].AgentSharingGroup = "";
+												}
+											}
+											if (Config.bTest)
+											{
+												AddNode(new FormalBuildTestNode(this, CodeProj, HostPlatform, Plat, Config.TargetConfig));
+											}																				
+										}
+									}
+									if (!bNoAutomatedTesting && FormalBuildConfigs.Count > 0)
+									{
+										if (HostPlatform == UnrealTargetPlatform.Mac || HostPlatform == UnrealTargetPlatform.Linux) continue; //temp hack till Linux and Mac automated testing works
+										var GameTests = Target.Rules.GUBP_GetGameTests_MonolithicOnly(HostPlatform, GetAltHostPlatform(HostPlatform), Plat);
+										var RequiredPlatforms = new List<UnrealTargetPlatform> { Plat };
+										var ThisMonoGameTestNodes = new List<string>();
+
+										foreach (var Test in GameTests)
+										{
+											var TestNodeName = Test.Key + "_" + Plat.ToString();
+											ThisMonoGameTestNodes.Add(AddNode(new UATTestNode(this, HostPlatform, CodeProj, TestNodeName, Test.Value, CookedAgentSharingGroup, false, RequiredPlatforms)));
+											if (!Options.bTestWithShared || !HasNode(WaitForTestShared.StaticGetFullName()))
+											{
+												RemovePseudodependencyFromNode((UATTestNode.StaticGetFullName(HostPlatform, CodeProj, TestNodeName)), WaitForTestShared.StaticGetFullName());
+											}
+										}
+										if (ThisMonoGameTestNodes.Count > 0)
+										{
+											GameTestNodes.Add(AddNode(new GameAggregateNode(this, HostPlatform, CodeProj, "CookedTests_" + Plat.ToString() + "_" + Kind.ToString() + HostPlatformNode.StaticGetHostPlatformSuffix(HostPlatform), ThisMonoGameTestNodes, 0.0f)));
+										}
+									}
+								}
+							}
+						}
+					}
+				}
+				if (!bNoAutomatedTesting)
+				{
+					foreach (var ServerPlatform in ServerPlatforms)
+					{
+						foreach (var GamePlatform in GamePlatforms)
+						{
+							if (HostPlatform == UnrealTargetPlatform.Mac || HostPlatform == UnrealTargetPlatform.Linux) continue; //temp hack till Linux and Mac automated testing works
+							var Target = CodeProj.Properties.Targets[TargetRules.TargetType.Game];
+							var ClientServerTests = Target.Rules.GUBP_GetClientServerTests_MonolithicOnly(HostPlatform, GetAltHostPlatform(HostPlatform), ServerPlatform, GamePlatform);
+							var RequiredPlatforms = new List<UnrealTargetPlatform> { ServerPlatform };
+							if (ServerPlatform != GamePlatform)
+							{
+								RequiredPlatforms.Add(GamePlatform);
+							}
+							foreach (var Test in ClientServerTests)
+							{
+								var TestNodeName = Test.Key + "_" + GamePlatform.ToString() + "_" + ServerPlatform.ToString();
+								GameTestNodes.Add(AddNode(new UATTestNode(this, HostPlatform, CodeProj, TestNodeName, Test.Value, CookedAgentSharingGroup, false, RequiredPlatforms)));
+							}
+						}
+					}
+					if (GameTestNodes.Count > 0)
+					{
+						AddNode(new GameAggregateNode(this, HostPlatform, CodeProj, "AllCookedTests", GameTestNodes));
+					}
+				}
+			}
+		}
+
+        int NumSharedAllHosts = 0;
+        foreach (var CodeProj in Branch.CodeProjects)
+        {
+            if (CodeProj.Properties.Targets.ContainsKey(TargetRules.TargetType.Editor))
+            {
+                bool AnySeparate = false;
+                var PromotedHosts = new List<UnrealTargetPlatform>();
+
+                foreach (var HostPlatform in HostPlatforms)
+                {
+                    if (!BranchOptions.ExcludePlatformsForEditor.Contains(HostPlatform) && !BranchOptions.RemovePlatformFromPromotable.Contains(HostPlatform))
+					{
+						var Options = CodeProj.Options(HostPlatform);
+						AnySeparate = AnySeparate || Options.bSeparateGamePromotion;
+						if (Options.bIsPromotable)
+						{
+							if (!Options.bSeparateGamePromotion)
+							{
+								NumSharedAllHosts++;
+							}
+							PromotedHosts.Add(HostPlatform);
+						}
+					}
+                }
+                if (PromotedHosts.Count > 0)
+                {
+                    AddNode(new GameAggregatePromotableNode(this, PromotedHosts, CodeProj, true));
+                    if (AnySeparate)
+                    {
+                        AddNode(new WaitForGamePromotionUserInput(this, CodeProj, false));
+                        AddNode(new GameLabelPromotableNode(this, CodeProj, false));
+                        AddNode(new WaitForGamePromotionUserInput(this, CodeProj, true));
+                        AddNode(new GameLabelPromotableNode(this, CodeProj, true));
+                    }
+                }
+            }
+        }
+        if (NumSharedAllHosts > 0)
+        {
+            AddNode(new GameAggregatePromotableNode(this, HostPlatforms, Branch.BaseEngineProject, false));
+
+            AddNode(new SharedAggregatePromotableNode(this, HostPlatforms));
+            AddNode(new WaitForSharedPromotionUserInput(this, false));
+            AddNode(new SharedLabelPromotableNode(this, false));
+            AddNode(new SharedLabelPromotableSuccessNode());
+            AddNode(new WaitForTestShared(this));
+            AddNode(new WaitForSharedPromotionUserInput(this, true));
+            AddNode(new SharedLabelPromotableNode(this, true));			
+        }
+        foreach (var HostPlatform in HostPlatforms)
+        {
+            AddCustomNodes(HostPlatform);
+        }
+        
+        if (HasNode(ToolsForCompileNode.StaticGetFullName(UnrealTargetPlatform.Win64)))
+        {
+			if (HasNode(GamePlatformMonolithicsNode.StaticGetFullName(UnrealTargetPlatform.Mac, Branch.BaseEngineProject, UnrealTargetPlatform.IOS)) && HasNode(ToolsNode.StaticGetFullName(UnrealTargetPlatform.Win64)))
+			{
+				//AddNode(new IOSOnPCTestNode(this)); - Disable IOSOnPCTest until a1011 crash is fixed
+			}
+			//AddNode(new VSExpressTestNode(this));
+			if (ActivePlatforms.Contains(UnrealTargetPlatform.Linux) && !BranchOptions.ExcludePlatformsForEditor.Contains(UnrealTargetPlatform.Linux))
+			{
+				AddNode(new RootEditorCrossCompileLinuxNode(UnrealTargetPlatform.Win64));
+			}
+            if (!bPreflightBuild)
+            {
+                AddNode(new CleanSharedTempStorageNode(this));
+            }
+        }
+#if false
+        // this doesn't work for lots of reasons...we can't figure out what the dependencies are until far later
+        if (bPreflightBuild)
+        {
+            GeneralSuccessNode PreflightSuccessNode = new GeneralSuccessNode("Preflight");
+            foreach (var NodeToDo in GUBPNodes)
+            {
+                if (NodeToDo.Value.RunInEC())
+                {
+                    PreflightSuccessNode.AddPseudodependency(NodeToDo.Key);
+                }
+            }
+            AddNode(PreflightSuccessNode);
+        }
+#endif
+	}
+}
