Added staging of bootstrap app on Mac

diff --git a/Engine/Source/Programs/AutomationTool/Mac/MacPlatform.Automation.cs b/Engine/Source/Programs/AutomationTool/Mac/MacPlatform.Automation.cs
--- a/Engine/Source/Programs/AutomationTool/Mac/MacPlatform.Automation.cs
+++ b/Engine/Source/Programs/AutomationTool/Mac/MacPlatform.Automation.cs
@@ -153,6 +153,118 @@ public class MacPlatform : Platform
 			string UnrealCEFSubProcessPath = CombinePaths("Engine/Binaries", SC.PlatformDir, "UnrealCEFSubProcess.app");
 			StageAppBundle(SC, StagedFileType.NonUFS, CombinePaths(SC.LocalRoot, "Engine/Binaries", SC.PlatformDir, "UnrealCEFSubProcess.app"), UnrealCEFSubProcessPath);
 		}
+
+		// Stage the bootstrap executable
+		if(!Params.NoBootstrapExe)
+		{
+			foreach(BuildReceipt Receipt in SC.StageTargetReceipts)
+			{
+				BuildProduct Executable = Receipt.BuildProducts.FirstOrDefault(x => x.Type == BuildProductType.Executable);
+				if(Executable != null)
+				{
+					// only create bootstraps for executables
+					if (SC.NonUFSStagingFiles.ContainsKey(Executable.Path) && Executable.Path.Replace("\\", "/").Contains("/" + TargetPlatformType.ToString() + "/"))
+					{
+						string BootstrapArguments = "";
+						if (!SC.IsCodeBasedProject && !ShouldStageCommandLine(Params, SC))
+						{
+							BootstrapArguments = String.Format("../../../{0}/{0}.uproject", SC.ShortProjectName);
+						}
+
+						string BootstrapExeName;
+						if(SC.StageTargetConfigurations.Count > 1)
+						{
+							BootstrapExeName = Path.GetFileName(Executable.Path) + ".app";
+						}
+						else if(Params.IsCodeBasedProject)
+						{
+							BootstrapExeName = Receipt.GetProperty("TargetName", SC.ShortProjectName) + ".app";
+						}
+						else
+						{
+							BootstrapExeName = SC.ShortProjectName + ".app";
+						}
+
+						string AppPath = Executable.Path.Substring(0, Executable.Path.LastIndexOf(".app/") + 4);
+						string AppRelativePath = SC.NonUFSStagingFiles[Executable.Path].Substring(0, SC.NonUFSStagingFiles[Executable.Path].LastIndexOf(".app/") + 4);
+						StageBootstrapExecutable(SC, BootstrapExeName, AppPath, AppRelativePath, BootstrapArguments);
+					}
+				}
+			}
+		}
+	}
+
+	string GetValueFromInfoPlist(string InfoPlist, string Key)
+	{
+		string Value = "";
+		string KeyString = "<key>" + Key + "</key>";
+		int KeyIndex = InfoPlist.IndexOf(KeyString);
+		if (KeyIndex > 0)
+		{
+			int ValueStartIndex = InfoPlist.IndexOf("<string>", KeyIndex + KeyString.Length) + "<string>".Length;
+			int ValueEndIndex = InfoPlist.IndexOf("</string>", ValueStartIndex);
+			if (ValueStartIndex > 0 && ValueEndIndex > ValueStartIndex)
+			{
+				Value = InfoPlist.Substring(ValueStartIndex, ValueEndIndex - ValueStartIndex);
+			}
+		}
+		return Value;
+	}
+
+	void StageBootstrapExecutable(DeploymentContext SC, string ExeName, string TargetFile, string StagedRelativeTargetPath, string StagedArguments)
+	{
+		Console.WriteLine(string.Format("'{0}' '{1}' '{2}' '{3}'", ExeName, TargetFile, StagedRelativeTargetPath, StagedArguments));
+
+		string InputApp = CombinePaths(SC.LocalRoot, "Engine", "Binaries", SC.PlatformDir, "BootstrapPackagedGame.app");
+		if (InternalUtils.SafeDirectoryExists(InputApp))
+		{
+			// Create the new bootstrap program
+			string IntermediateDir = CombinePaths(SC.ProjectRoot, "Intermediate", "Staging");
+			InternalUtils.SafeCreateDirectory(IntermediateDir);
+
+			string IntermediateApp = CombinePaths(IntermediateDir, ExeName);
+			if (Directory.Exists(IntermediateApp))
+			{
+				Directory.Delete(IntermediateApp, true);
+			}
+			CloneDirectory(InputApp, IntermediateApp);
+
+			// Rename the executable
+			string GameName = Path.GetFileNameWithoutExtension(ExeName);
+			File.Move(CombinePaths(IntermediateApp, "Contents", "MacOS", "BootstrapPackagedGame"), CombinePaths(IntermediateApp, "Contents", "MacOS", GameName));
+
+			// Copy the icon
+			string SrcInfoPlistPath = CombinePaths(TargetFile, "Contents", "Info.plist");
+			string SrcInfoPlist = File.ReadAllText(SrcInfoPlistPath);
+
+			string IconName = GetValueFromInfoPlist(SrcInfoPlist, "CFBundleIconFile");
+			if (!string.IsNullOrEmpty(IconName))
+			{
+				string IconPath = CombinePaths(TargetFile, "Contents", "Resources", IconName + ".icns");
+				InternalUtils.SafeCreateDirectory(CombinePaths(IntermediateApp, "Contents", "Resources"));
+				File.Copy(IconPath, CombinePaths(IntermediateApp, "Contents", "Resources", IconName + ".icns"));
+			}
+
+			// Update Info.plist contents
+			string DestInfoPlistPath = CombinePaths(IntermediateApp, "Contents", "Info.plist");
+			string DestInfoPlist = File.ReadAllText(DestInfoPlistPath);
+
+			string AppIdentifier = GetValueFromInfoPlist(SrcInfoPlist, "CFBundleIdentifier");
+			if (AppIdentifier == "com.epicgames.UE4Game")
+			{
+				AppIdentifier = "";
+			}
+
+			DestInfoPlist = DestInfoPlist.Replace("com.epicgames.BootstrapPackagedGame", string.IsNullOrEmpty(AppIdentifier) ? "com.epicgames." + GameName + "_bootstrap" : AppIdentifier + "_bootstrap");
+			DestInfoPlist = DestInfoPlist.Replace("BootstrapPackagedGame", GameName);
+			DestInfoPlist = DestInfoPlist.Replace("__UE4_ICON_FILE__", IconName);
+			DestInfoPlist = DestInfoPlist.Replace("__UE4_APP_TO_LAUNCH__", StagedRelativeTargetPath);
+			DestInfoPlist = DestInfoPlist.Replace("__UE4_COMMANDLINE__", StagedArguments);
+
+			File.WriteAllText(DestInfoPlistPath, DestInfoPlist);
+
+			StageAppBundle(SC, StagedFileType.NonUFS, IntermediateApp, ExeName);
+		}
 	}
 
 	public override void Package(ProjectParams Params, DeploymentContext SC, int WorkingCL)
