Do not stage UE4CommandLine.txt on Mac

diff --git a/Engine/Source/Programs/AutomationTool/Mac/MacPlatform.Automation.cs b/Engine/Source/Programs/AutomationTool/Mac/MacPlatform.Automation.cs
--- a/Engine/Source/Programs/AutomationTool/Mac/MacPlatform.Automation.cs
+++ b/Engine/Source/Programs/AutomationTool/Mac/MacPlatform.Automation.cs
@@ -413,4 +413,8 @@ public class MacPlatform : Platform
 		}
 		return false;
 	}
+	public override bool ShouldStageCommandLine(ProjectParams Params, DeploymentContext SC)
+	{
+		return false; // !String.IsNullOrEmpty(Params.StageCommandline) || !String.IsNullOrEmpty(Params.RunCommandline) || (!Params.IsCodeBasedProject && Params.NoBootstrapExe);
+	}
 }
