Always create both GameName/Binaries/Mac and Engine/Binaries/Mac when packaging a Mac game

diff --git a/Engine/Source/Programs/AutomationTool/Mac/MacPlatform.Automation.cs b/Engine/Source/Programs/AutomationTool/Mac/MacPlatform.Automation.cs
--- a/Engine/Source/Programs/AutomationTool/Mac/MacPlatform.Automation.cs
+++ b/Engine/Source/Programs/AutomationTool/Mac/MacPlatform.Automation.cs
@@ -213,8 +213,6 @@ public class MacPlatform : Platform
 
 	void StageBootstrapExecutable(DeploymentContext SC, string ExeName, string TargetFile, string StagedRelativeTargetPath, string StagedArguments)
 	{
-		Console.WriteLine(string.Format("'{0}' '{1}' '{2}' '{3}'", ExeName, TargetFile, StagedRelativeTargetPath, StagedArguments));
-
 		string InputApp = CombinePaths(SC.LocalRoot, "Engine", "Binaries", SC.PlatformDir, "BootstrapPackagedGame.app");
 		if (InternalUtils.SafeDirectoryExists(InputApp))
 		{
@@ -362,9 +360,10 @@ public class MacPlatform : Platform
 
 		if (!SC.bIsCombiningMultiplePlatforms)
 		{
-			// creating this directory when the content isn't moved into the application causes it 
+			// creating these directories when the content isn't moved into the application causes it 
 			// to fail to load, and isn't needed
-			Directory.CreateDirectory(CombinePaths(TargetPath, ExeName.StartsWith("UE4Game") ? "Engine" : SC.ShortProjectName, "Binaries", "Mac"));
+			Directory.CreateDirectory(CombinePaths(TargetPath, "Engine", "Binaries", "Mac"));
+			Directory.CreateDirectory(CombinePaths(TargetPath, SC.ShortProjectName, "Binaries", "Mac"));
 		}
 	}
 
