Copy version number and copyright text from game's Info.plist when staging Mac bootstrap app

diff --git a/Engine/Source/Programs/AutomationTool/Mac/MacPlatform.Automation.cs b/Engine/Source/Programs/AutomationTool/Mac/MacPlatform.Automation.cs
--- a/Engine/Source/Programs/AutomationTool/Mac/MacPlatform.Automation.cs
+++ b/Engine/Source/Programs/AutomationTool/Mac/MacPlatform.Automation.cs
@@ -194,9 +194,9 @@ public class MacPlatform : Platform
 		}
 	}
 
-	string GetValueFromInfoPlist(string InfoPlist, string Key)
+	string GetValueFromInfoPlist(string InfoPlist, string Key, string DefaultValue = "")
 	{
-		string Value = "";
+		string Value = DefaultValue;
 		string KeyString = "<key>" + Key + "</key>";
 		int KeyIndex = InfoPlist.IndexOf(KeyString);
 		if (KeyIndex > 0)
@@ -253,11 +253,18 @@ public class MacPlatform : Platform
 				AppIdentifier = "";
 			}
 
+			string Copyright = GetValueFromInfoPlist(SrcInfoPlist, "NSHumanReadableCopyright");
+			string BundleVersion = GetValueFromInfoPlist(SrcInfoPlist, "CFBundleVersion", "1");
+			string ShortVersion = GetValueFromInfoPlist(SrcInfoPlist, "CFBundleShortVersionString", "1.0");
+
 			DestInfoPlist = DestInfoPlist.Replace("com.epicgames.BootstrapPackagedGame", string.IsNullOrEmpty(AppIdentifier) ? "com.epicgames." + GameName + "_bootstrap" : AppIdentifier + "_bootstrap");
 			DestInfoPlist = DestInfoPlist.Replace("BootstrapPackagedGame", GameName);
 			DestInfoPlist = DestInfoPlist.Replace("__UE4_ICON_FILE__", IconName);
 			DestInfoPlist = DestInfoPlist.Replace("__UE4_APP_TO_LAUNCH__", StagedRelativeTargetPath);
 			DestInfoPlist = DestInfoPlist.Replace("__UE4_COMMANDLINE__", StagedArguments);
+			DestInfoPlist = DestInfoPlist.Replace("__UE4_COPYRIGHT__", Copyright);
+			DestInfoPlist = DestInfoPlist.Replace("__UE4_BUNDLE_VERSION__", BundleVersion);
+			DestInfoPlist = DestInfoPlist.Replace("__UE4_SHORT_VERSION__", ShortVersion);
 
 			File.WriteAllText(DestInfoPlistPath, DestInfoPlist);
 
