Added shared feature packs to list in rocket build list

diff --git a/Engine/Source/Programs/AutomationTool/Scripts/RocketBuild.Automation.cs b/Engine/Source/Programs/AutomationTool/Scripts/RocketBuild.Automation.cs
--- a/Engine/Source/Programs/AutomationTool/Scripts/RocketBuild.Automation.cs
+++ b/Engine/Source/Programs/AutomationTool/Scripts/RocketBuild.Automation.cs
@@ -38,6 +38,7 @@ namespace Rocket
 			"TP_2DSideScrollerBP",
 			"TP_VehicleAdv",
 			"TP_VehicleAdvBP",
+
 		};
 
 		static readonly string[] CurrentFeaturePacks = 
@@ -65,7 +66,24 @@ namespace Rocket
 			"TP_VehicleAdv",
 			"TP_VehicleAdvBP",
 			"StarterContent",
-			"MobileStarterContent"
+			"MobileStarterContent",
+			"GeometryHigh",		//Shared packs
+			"GeometryStandard",
+			"MannequinHigh",
+			"MannequinStandard",
+			"RollingHigh",
+			"RollingStandard",
+			"2DSideScrollerStandard",
+			"FirstPersonStandard",
+			"FlyingStandard",
+			"PuzzleStandard",
+			"RollingStandard",
+			"RollingHigh",
+			"SideScrollerStandard",
+			"ThirdPersonStandard",
+			"TwinStickStandard",
+			"VehicleStandard",
+			"VehicleAdvStandard",
 		};
 
 		public RocketBuild()
