UEB-260 - Break AutomationTool into AutomationUtils that all automation projects depend on, and AutomationTool, which essentially only contains the startup code.
* Remove ErrorReporter.Error, replace with AutomationException with Error Code.
* Move ErrorCodes to AutomationException.
* Don't return exit codes. Solely rely on exceptions to propagate exit codes.
* Remove MainProc delegate
* Remove setting of Environment.ExitCode as it is ignored when main returns an int.
* ShutdownLogging is nothrow, as all exceptions would be ignored anyway.
* Wrap all shutdown steps so further ones get a chance to run.
* Move HostPlatform.Initialize into the global try/catch block

diff --git a/Engine/Source/Programs/AutomationTool/AllDesktop/AllDesktop.Automation.csproj b/Engine/Source/Programs/AutomationTool/AllDesktop/AllDesktop.Automation.csproj
--- a/Engine/Source/Programs/AutomationTool/AllDesktop/AllDesktop.Automation.csproj
+++ b/Engine/Source/Programs/AutomationTool/AllDesktop/AllDesktop.Automation.csproj
@@ -1,4 +1,4 @@
-<?xml version="1.0" encoding="utf-8"?>
+﻿<?xml version="1.0" encoding="utf-8"?>
 <Project ToolsVersion="4.0" DefaultTargets="Build" xmlns="http://schemas.microsoft.com/developer/msbuild/2003">
   <Import Project="$(MSBuildExtensionsPath)\$(MSBuildToolsVersion)\Microsoft.Common.props" Condition="Exists('$(MSBuildExtensionsPath)\$(MSBuildToolsVersion)\Microsoft.Common.props')" />
   <PropertyGroup>
@@ -60,16 +60,16 @@
     <Compile Include="Properties\AssemblyInfo.cs" />
   </ItemGroup>
   <ItemGroup>
-    <ProjectReference Include="..\AutomationTool.csproj">
-      <Project>{A35B0261-D868-475E-A098-EB14CE6DC89F}</Project>
-      <Name>AutomationTool</Name>
-      <Private>False</Private>
-    </ProjectReference>
     <ProjectReference Include="..\..\UnrealBuildTool\UnrealBuildTool.csproj">
       <Project>{FD7C5E1A-CFE4-4FD5-A525-1EB1599A39AC}</Project>
       <Name>UnrealBuildTool</Name>
       <Private>False</Private>
     </ProjectReference>
+    <ProjectReference Include="..\AutomationUtils\AutomationUtils.Automation.csproj">
+      <Project>{2c96a7f2-b1a3-4258-8e0a-e588ff41a53e}</Project>
+      <Name>AutomationUtils.Automation</Name>
+      <Private>False</Private>
+    </ProjectReference>
   </ItemGroup>
   <Import Project="$(MSBuildToolsPath)\Microsoft.CSharp.targets" />
 </Project>
\ No newline at end of file
diff --git a/Engine/Source/Programs/AutomationTool/Android/Android.Automation.csproj b/Engine/Source/Programs/AutomationTool/Android/Android.Automation.csproj
--- a/Engine/Source/Programs/AutomationTool/Android/Android.Automation.csproj
+++ b/Engine/Source/Programs/AutomationTool/Android/Android.Automation.csproj
@@ -61,16 +61,16 @@
     <Compile Include="Properties\AssemblyInfo.cs" />
   </ItemGroup>
   <ItemGroup>
-    <ProjectReference Include="..\AutomationTool.csproj">
-      <Project>{a35b0261-d868-475e-a098-eb14ce6dc89f}</Project>
-      <Name>AutomationTool</Name>
-      <Private>False</Private>
-    </ProjectReference>
     <ProjectReference Include="..\..\UnrealBuildTool\UnrealBuildTool.csproj">
       <Project>{fd7c5e1a-cfe4-4fd5-a525-1eb1599a39ac}</Project>
       <Name>UnrealBuildTool</Name>
       <Private>False</Private>
     </ProjectReference>
+    <ProjectReference Include="..\AutomationUtils\AutomationUtils.Automation.csproj">
+      <Project>{2c96a7f2-b1a3-4258-8e0a-e588ff41a53e}</Project>
+      <Name>AutomationUtils.Automation</Name>
+      <Private>False</Private>
+    </ProjectReference>
     <ProjectReference Include="..\Scripts\AutomationScripts.Automation.csproj">
       <Project>{8aa00d65-0954-4a27-ac0d-fb8b1106120f}</Project>
       <Name>AutomationScripts.Automation</Name>
@@ -78,4 +78,4 @@
     </ProjectReference>
   </ItemGroup>
   <Import Project="$(MSBuildToolsPath)\Microsoft.CSharp.targets" />
-</Project>
+</Project>
\ No newline at end of file
diff --git a/Engine/Source/Programs/AutomationTool/Android/AndroidPlatform.Automation.cs b/Engine/Source/Programs/AutomationTool/Android/AndroidPlatform.Automation.cs
--- a/Engine/Source/Programs/AutomationTool/Android/AndroidPlatform.Automation.cs
+++ b/Engine/Source/Programs/AutomationTool/Android/AndroidPlatform.Automation.cs
@@ -63,15 +63,13 @@ public class AndroidPlatform : Platform
 		string PackageName = GetPackageInfo(ApkName, false);
 		if (PackageName == null)
 		{
-            ErrorReporter.Error("Failed to get package name from " + ApkName, (int)ErrorCodes.Error_FailureGettingPackageInfo);
-			throw new AutomationException("Failed to get package name from " + ApkName);
+			throw new AutomationException(ErrorCodes.Error_FailureGettingPackageInfo, "Failed to get package name from " + ApkName);
 		}
 
 		string PackageVersion = GetPackageInfo(ApkName, true);
 		if (PackageVersion == null || PackageVersion.Length == 0)
 		{
-            ErrorReporter.Error("Failed to get package version from " + ApkName, (int)ErrorCodes.Error_FailureGettingPackageInfo);
-			throw new AutomationException("Failed to get package version from " + ApkName);
+			throw new AutomationException(ErrorCodes.Error_FailureGettingPackageInfo, "Failed to get package version from " + ApkName);
 		}
 
 		if (PackageVersion.Length > 0)
@@ -169,8 +167,7 @@ public class AndroidPlatform : Platform
                     if (FileExists_NoExceptions(UE4SOName) == false)
 					{
                         Log("Failed to find game .so " + UE4SOName);
-						AutomationTool.ErrorReporter.Error("Stage Failed.", (int)AutomationTool.ErrorCodes.Error_MissingExecutable);
-                        throw new AutomationException("Could not find .so {0}. You may need to build the UE4 project with your target configuration and platform.", UE4SOName);
+                        throw new AutomationException(ErrorCodes.Error_MissingExecutable, "Stage Failed. Could not find .so {0}. You may need to build the UE4 project with your target configuration and platform.", UE4SOName);
 					}
 				}
 				
@@ -317,9 +314,7 @@ public class AndroidPlatform : Platform
 	{
 		if (SC.StageTargetConfigurations.Count != 1)
 		{
-            string ErrorString = String.Format("Android is currently only able to package one target configuration at a time, but StageTargetConfigurations contained {0} configurations", SC.StageTargetConfigurations.Count);
-            ErrorReporter.Error(ErrorString, (int)ErrorCodes.Error_OnlyOneTargetConfigurationSupported);
-			throw new AutomationException(ErrorString);
+			throw new AutomationException(ErrorCodes.Error_OnlyOneTargetConfigurationSupported, "Android is currently only able to package one target configuration at a time, but StageTargetConfigurations contained {0} configurations", SC.StageTargetConfigurations.Count);
 		}
 
 		var Architectures = UnrealBuildTool.AndroidToolChain.GetAllArchitectures();
@@ -340,15 +335,11 @@ public class AndroidPlatform : Platform
 				// verify the files exist
 				if (!FileExists(ApkName))
 				{
-					string ErrorString = String.Format("ARCHIVE FAILED - {0} was not found", ApkName);
-					ErrorReporter.Error(ErrorString, (int)ErrorCodes.Error_AppNotFound);
-					throw new AutomationException(ErrorString);
+					throw new AutomationException(ErrorCodes.Error_AppNotFound, "ARCHIVE FAILED - {0} was not found", ApkName);
 				}
 				if (!bPackageDataInsideApk && !FileExists(ObbName))
 				{
-					string ErrorString = String.Format("ARCHIVE FAILED - {0} was not found", ObbName);
-					ErrorReporter.Error(ErrorString, (int)ErrorCodes.Error_ObbNotFound);
-					throw new AutomationException(ErrorString);
+                    throw new AutomationException(ErrorCodes.Error_ObbNotFound, "ARCHIVE FAILED - {0} was not found", ObbName);
 				}
 
 				SC.ArchiveFiles(Path.GetDirectoryName(ApkName), Path.GetFileName(ApkName));
@@ -595,8 +586,7 @@ public class AndroidPlatform : Platform
 					}
 				}
 
-				ErrorReporter.Error(ErrorMessage, (int)ErrorCodes.Error_AppInstallFailed);
-				throw new AutomationException(ErrorMessage);
+				throw new AutomationException(ErrorCodes.Error_AppInstallFailed, ErrorMessage);
 			}
 		}
  
@@ -871,8 +861,7 @@ public class AndroidPlatform : Platform
 		string[] Subdirs = Directory.GetDirectories(path);
         if (Subdirs.Length == 0)
         {
-            ErrorReporter.Error("Failed to find %ANDROID_HOME%/build-tools subdirectory", (int)ErrorCodes.Error_AndroidBuildToolsPathNotFound);
-            throw new AutomationException("Failed to find %ANDROID_HOME%/build-tools subdirectory");
+            throw new AutomationException(ErrorCodes.Error_AndroidBuildToolsPathNotFound, "Failed to find %ANDROID_HOME%/build-tools subdirectory");
         }
 		// we expect there to be one, so use the first one
 		return Path.Combine(Subdirs[0], Utils.IsRunningOnMono ? "aapt" : "aapt.exe");
@@ -937,9 +926,7 @@ public class AndroidPlatform : Platform
 		// if after the fallbacks, we still don't have it, we can't continue
 		if (!AppArchitectures.Contains(DeviceArch))
 		{
-			string ErrorString = String.Format("Unable to run because you don't have an apk that is usable on {0}. Looked for {1}", Params.Device, DeviceArch);
-            ErrorReporter.Error(ErrorString, (int)ErrorCodes.Error_NoApkSuitableForArchitecture);
-            throw new AutomationException(ErrorString);
+            throw new AutomationException(ErrorCodes.Error_NoApkSuitableForArchitecture, "Unable to run because you don't have an apk that is usable on {0}. Looked for {1}", Params.Device, DeviceArch);
 		}
 
 		return DeviceArch;
@@ -989,8 +976,7 @@ public class AndroidPlatform : Platform
 		string PackageName = GetPackageInfo(ApkName, false);
 		if (PackageName == null)
 		{
-            ErrorReporter.Error("Failed to get package name from " + ClientApp, (int)ErrorCodes.Error_FailureGettingPackageInfo);
-            throw new AutomationException("Failed to get package name from " + ClientApp);
+            throw new AutomationException(ErrorCodes.Error_FailureGettingPackageInfo, "Failed to get package name from " + ClientApp);
 		}
 
 		if (Params.Prebuilt)
diff --git a/Engine/Source/Programs/AutomationTool/AutomationTool.csproj b/Engine/Source/Programs/AutomationTool/AutomationTool.csproj
--- a/Engine/Source/Programs/AutomationTool/AutomationTool.csproj
+++ b/Engine/Source/Programs/AutomationTool/AutomationTool.csproj
@@ -96,40 +96,9 @@
     <Compile Include="..\DotNETCommon\MetaData.cs">
       <Link>Properties\MetaData.cs</Link>
     </Compile>
-    <Compile Include="Automation.cs" />
-    <Compile Include="AutomationExpection.cs" />
-    <Compile Include="BuildCommand.cs" />
-    <Compile Include="BuildEnvironment.cs" />
-    <Compile Include="BuildUtils.cs" />
-    <Compile Include="CommandletUtils.cs" />
-    <Compile Include="DeploymentContext.cs" />
-    <Compile Include="Distiller.cs" />
-    <Compile Include="FileFilter.cs" />
-    <Compile Include="HelpUtils.cs" />
-    <Compile Include="HostPlatform.cs" />
-    <Compile Include="LinuxHostPlatform.cs" />
-    <Compile Include="LocalBuildEnvironment.cs" />
-    <Compile Include="LocalP4Environment.cs" />
-    <Compile Include="Log.cs" />
-    <Compile Include="MacHostPlatform.cs" />
-    <Compile Include="MCPPublic.cs" />
-    <Compile Include="P4Blame.cs" />
-    <Compile Include="P4Environment.cs" />
-    <Compile Include="P4Utils.cs" />
-    <Compile Include="Platform.cs" />
-    <Compile Include="ProcessUtils.cs" />
     <Compile Include="Program.cs" />
-    <Compile Include="ProjectParams.cs" />
-    <Compile Include="ProjectUtils.cs" />
     <Compile Include="Properties\AssemblyInfo.cs" />
-    <Compile Include="ScriptCompiler.cs" />
-    <Compile Include="CommandUtils.cs" />
     <Compile Include="SharedUtils.cs" />
-    <Compile Include="UBTUtils.cs" />
-    <Compile Include="UE4Build.cs" />
-    <Compile Include="Utils.cs" />
-    <Compile Include="WatchdogTimer.cs" />
-    <Compile Include="WindowsHostPlatform.cs" />
   </ItemGroup>
   <ItemGroup>
     <ProjectReference Include="..\DotNETCommon\DotNETUtilities\DotNETUtilities.csproj">
@@ -141,10 +110,18 @@
       <Name>UnrealBuildTool</Name>
       <Private>False</Private>
     </ProjectReference>
+    <ProjectReference Include="AutomationUtils\AutomationUtils.Automation.csproj">
+      <Project>{2c96a7f2-b1a3-4258-8e0a-e588ff41a53e}</Project>
+      <Name>AutomationUtils.Automation</Name>
+      <Private>False</Private>
+    </ProjectReference>
   </ItemGroup>
   <ItemGroup>
     <None Include="app.config" />
   </ItemGroup>
+  <ItemGroup>
+    <WCFMetadata Include="Service References\" />
+  </ItemGroup>
   <Import Project="$(MSBuildToolsPath)\Microsoft.CSharp.targets" />
   <!-- To modify your build process, add your task inside one of the targets below and uncomment it. 
        Other similar extension points exist, see Microsoft.Common.targets.
diff --git a/Engine/Source/Programs/AutomationTool/AutomationUtils/Automation.cs b/Engine/Source/Programs/AutomationTool/AutomationUtils/Automation.cs
new file mode 100644
--- /dev/null
+++ b/Engine/Source/Programs/AutomationTool/AutomationUtils/Automation.cs
@@ -0,0 +1,495 @@
+﻿// Copyright 1998-2015 Epic Games, Inc. All Rights Reserved.
+using System;
+using System.Collections.Generic;
+using System.Linq;
+using System.Text;
+using System.IO;
+using System.Reflection;
+using System.Diagnostics;
+using UnrealBuildTool;
+using Tools.DotNETCommon.CaselessDictionary;
+
+namespace AutomationTool
+{
+	#region Command info
+
+	/// <summary>
+	/// Command to execute info.
+	/// </summary>
+	class CommandInfo
+	{
+		public string CommandName;
+		public List<string> Arguments = new List<string>();
+
+		public override string ToString()
+		{
+			string Result = CommandName;
+			Result += "(";
+			for (int Index = 0; Index < Arguments.Count; ++Index)
+			{
+				if (Index > 0)
+				{
+					Result += ", ";
+				}
+				Result += Arguments[Index];
+			}
+			Result += ")";
+			return Result;
+		}
+	}
+
+	#endregion
+
+	#region Command Line helpers
+
+	/// <summary>
+	/// Helper class for automatically registering command line params.
+	/// </summary>
+	public class CommandLineArg
+	{
+		public delegate void OnSetDelegate();
+		public string Name;
+		public OnSetDelegate SetDelegate;
+
+		/// <summary>
+		/// True of this arg was set in the command line.
+		/// </summary>
+		public bool IsSet { get; private set; }
+
+		public CommandLineArg(string InName, OnSetDelegate InSet = null)
+		{
+			Name = InName;
+			SetDelegate = InSet;
+			GlobalCommandLine.RegisteredArgs.Add(Name, this);
+		}
+
+		public void Set()
+		{
+			IsSet = true;
+			if (SetDelegate != null)
+			{
+				SetDelegate();
+			}
+		}
+
+		public override string ToString()
+		{
+			return String.Format("{0}, {1}", Name, IsSet);
+		}
+
+		/// <summary>
+		/// Returns true of this arg was set in the command line.
+		/// </summary>
+		public static implicit operator bool(CommandLineArg Arg)
+		{
+			return Arg.IsSet;
+		}
+	}
+
+	/// <summary>
+	/// Global command line parameters.
+	/// </summary>
+	public static class GlobalCommandLine
+	{
+		/// <summary>
+		/// List of all registered command line parameters (global, not command-specific)
+		/// </summary>
+		public static CaselessDictionary<CommandLineArg> RegisteredArgs = new CaselessDictionary<CommandLineArg>();
+
+		public static CommandLineArg CompileOnly = new CommandLineArg("-CompileOnly");
+		public static CommandLineArg Verbose = new CommandLineArg("-Verbose");
+		public static CommandLineArg Submit = new CommandLineArg("-Submit");
+		public static CommandLineArg NoSubmit = new CommandLineArg("-NoSubmit");
+		public static CommandLineArg ForceLocal = new CommandLineArg("-ForceLocal");
+		public static CommandLineArg NoP4 = new CommandLineArg("-NoP4");
+		public static CommandLineArg P4 = new CommandLineArg("-P4");
+		public static CommandLineArg Preprocess = new CommandLineArg("-Preprocess");
+		public static CommandLineArg NoCompile = new CommandLineArg("-NoCompile");
+		public static CommandLineArg NoCompileEditor = new CommandLineArg("-NoCompileEditor");
+		public static CommandLineArg Help = new CommandLineArg("-Help");
+		public static CommandLineArg List = new CommandLineArg("-List");
+		public static CommandLineArg Rocket = new CommandLineArg("-Rocket");
+		public static CommandLineArg NoKill = new CommandLineArg("-NoKill");
+		public static CommandLineArg Installed = new CommandLineArg("-Installed");
+		public static CommandLineArg UTF8Output = new CommandLineArg("-UTF8Output");
+		public static CommandLineArg NoAutoSDK = new CommandLineArg("-NoAutoSDK");
+		public static CommandLineArg IgnoreJunk = new CommandLineArg("-ignorejunk");
+		/// <summary>
+		/// Force initialize static members by calling this.
+		/// </summary>
+		public static void Init()
+		{
+			Log.TraceVerbose("Registered {0} command line parameters.", RegisteredArgs.Count);
+		}
+	}
+
+	#endregion
+
+	[Help(
+@"Executes scripted commands
+
+AutomationTool.exe [-verbose] [-compileonly] [-p4] Command0 [-Arg0 -Arg1 -Arg2 …] Command1 [-Arg0 -Arg1 …] Command2 [-Arg0 …] Commandn … [EnvVar0=MyValue0 … EnvVarn=MyValuen]"
+)]
+	[Help("verbose", "Enables verbose logging")]
+	[Help("nop4", "Disables Perforce functionality (default if not run on a build machine)")]
+	[Help("p4", "Enables Perforce functionality (default if run on a build machine)")]
+	[Help("compileonly", "Does not run any commands, only compiles them")]
+	[Help("forcelocal", "Forces local execution")]
+	[Help("help", "Displays help")]
+	[Help("list", "Lists all available commands")]
+	[Help("submit", "Allows UAT command to submit changes")]
+	[Help("nosubmit", "Prevents any submit attempts")]
+	[Help("nokill", "Does not kill any spawned processes on exit")]
+	[Help("ignorejunk", "Prevents UBT from cleaning junk files")]
+	public static class Automation
+	{
+		#region Command line parsing
+
+		/// <summary>
+		/// Parses command line parameter.
+		/// </summary>
+		/// <param name="ParamIndex">Parameter index</param>
+		/// <param name="CommandLine">Command line</param>
+		/// <param name="CurrentCommand">Recently parsed command</param>
+		/// <param name="OutAdditionalScriptsFolders">Additional script locations</param>
+		/// <returns>True if the parameter has been successfully parsed.</returns>
+		private static void ParseParam(string CurrentParam, CommandInfo CurrentCommand, List<string> OutAdditionalScriptsFolders)
+		{
+			bool bGlobalParam = false;
+			foreach (var RegisteredParam in GlobalCommandLine.RegisteredArgs)
+			{
+				if (String.Compare(CurrentParam, RegisteredParam.Key, true) == 0)
+				{
+					// Register and exit, we're done here.
+					RegisteredParam.Value.Set();
+					bGlobalParam = true;
+					break;
+				}
+			}
+
+			// The parameter was not found in the list of global parameters, continue looking...
+			if (CurrentParam.StartsWith("-ScriptDir=", StringComparison.InvariantCultureIgnoreCase))
+			{
+				var ScriptDir = CurrentParam.Substring(CurrentParam.IndexOf('=') + 1);
+				if (Directory.Exists(ScriptDir))
+				{
+					OutAdditionalScriptsFolders.Add(ScriptDir);
+					Log.TraceVerbose("Found additional script dir: {0}", ScriptDir);
+				}
+				else
+				{
+					throw new AutomationException("Specified ScriptDir doesn't exist: {0}", ScriptDir);
+				}
+			}
+			else if (CurrentParam.StartsWith("-"))
+			{
+				if (CurrentCommand != null)
+				{
+					CurrentCommand.Arguments.Add(CurrentParam.Substring(1));
+				}
+				else if (!bGlobalParam)
+				{
+					throw new AutomationException("Unknown parameter {0} in the command line that does not belong to any command.", CurrentParam);
+				}
+			}
+			else if (CurrentParam.Contains("="))
+			{
+				// Environment variable
+				int ValueStartIndex = CurrentParam.IndexOf('=') + 1;
+				string EnvVarName = CurrentParam.Substring(0, ValueStartIndex - 1);
+				if (String.IsNullOrEmpty(EnvVarName))
+				{
+					throw new AutomationException("Unable to parse environment variable that has no name. Error when parsing command line param {0}", CurrentParam);
+				}
+				string EnvVarValue = CurrentParam.Substring(ValueStartIndex);
+				CommandUtils.SetEnvVar(EnvVarName, EnvVarValue);
+			}
+		}
+
+		/// <summary>
+		/// Parse the command line and create a list of commands to execute.
+		/// </summary>
+		/// <param name="CommandLine">Command line</param>
+		/// <param name="OutCommandsToExecute">List containing the names of commands to execute</param>
+		/// <param name="OutAdditionalScriptsFolders">Optional list of additional paths to look for scripts in</param>
+		private static void ParseCommandLine(string[] CommandLine, List<CommandInfo> OutCommandsToExecute, List<string> OutAdditionalScriptsFolders)
+		{
+			// Initialize global command line parameters
+			GlobalCommandLine.Init();
+
+			Log.TraceInformation("Parsing command line: {0}", String.Join(" ", CommandLine));
+
+			CommandInfo CurrentCommand = null;
+			for (int Index = 0; Index < CommandLine.Length; ++Index)
+			{
+				var Param = CommandLine[Index];
+				if (Param.StartsWith("-") || Param.Contains("="))
+				{
+					ParseParam(CommandLine[Index], CurrentCommand, OutAdditionalScriptsFolders);
+				}
+				else
+				{
+					CurrentCommand = new CommandInfo();
+					CurrentCommand.CommandName = CommandLine[Index];
+					OutCommandsToExecute.Add(CurrentCommand);
+				}
+			}
+
+			// Validate
+			var Result = OutCommandsToExecute.Count > 0 || GlobalCommandLine.Help || GlobalCommandLine.CompileOnly || GlobalCommandLine.List;
+			if (OutCommandsToExecute.Count > 0)
+			{
+				Log.TraceVerbose("Found {0} scripts to execute:", OutCommandsToExecute.Count);
+				foreach (var Command in OutCommandsToExecute)
+				{
+					Log.TraceVerbose("  " + Command.ToString());
+				}
+			}
+			else if (!Result)
+			{
+				throw new AutomationException("Failed to find scripts to execute in the command line params.");
+			}
+			if (GlobalCommandLine.NoP4 && GlobalCommandLine.P4)
+			{
+				throw new AutomationException("{0} and {1} can't be set simultaneously.", GlobalCommandLine.NoP4.Name, GlobalCommandLine.P4.Name);
+			}
+			if (GlobalCommandLine.NoSubmit && GlobalCommandLine.Submit)
+			{
+				throw new AutomationException("{0} and {1} can't be set simultaneously.", GlobalCommandLine.NoSubmit.Name, GlobalCommandLine.Submit.Name);
+			}
+
+			if (GlobalCommandLine.Rocket)
+			{
+				UnrealBuildTool.UnrealBuildTool.bRunningRocket = true;
+			}
+		}
+
+		#endregion
+
+		#region Main Program
+
+		/// <summary>
+		/// Main method.
+		/// </summary>
+		/// <param name="CommandLine">Command line</param>
+		public static void Process(string[] CommandLine)
+		{
+			// Initial check for local or build machine runs BEFORE we parse the command line (We need this value set
+			// in case something throws the exception while parsing the command line)
+			IsBuildMachine = !String.IsNullOrEmpty(Environment.GetEnvironmentVariable("uebp_LOCAL_ROOT"));
+
+			// Scan the command line for commands to execute.
+			var CommandsToExecute = new List<CommandInfo>();
+			var AdditionalScriptsFolders = new List<string>();
+			ParseCommandLine(CommandLine, CommandsToExecute, AdditionalScriptsFolders);
+
+			// Check for build machine override (force local)
+			IsBuildMachine = GlobalCommandLine.ForceLocal ? false : IsBuildMachine;
+			Log.TraceInformation("IsBuildMachine={0}", IsBuildMachine);
+			Environment.SetEnvironmentVariable("IsBuildMachine", IsBuildMachine ? "1" : "0");
+
+			// should we kill processes on exit
+			ShouldKillProcesses = !GlobalCommandLine.NoKill;
+			Log.TraceInformation("ShouldKillProcesses={0}", ShouldKillProcesses);
+
+			if (CommandsToExecute.Count == 0 && GlobalCommandLine.Help)
+			{
+				DisplayHelp();
+				return;
+			}
+
+			// Disable AutoSDKs if specified on the command line
+			if (GlobalCommandLine.NoAutoSDK)
+			{
+				UEBuildPlatform.bAllowAutoSDKSwitching = false;
+			}
+
+			// Setup environment
+			Log.TraceInformation("Setting up command environment.");
+			CommandUtils.InitCommandEnvironment();
+
+			// Change CWD to UE4 root.
+			Environment.CurrentDirectory = CommandUtils.CmdEnv.LocalRoot;
+
+			// Fill in the project info
+			UnrealBuildTool.UProjectInfo.FillProjectInfo();
+
+			// Clean rules folders up
+			ProjectUtils.CleanupFolders();
+
+			// Compile scripts.
+			Log.TraceInformation("Compiling scripts.");
+			ScriptCompiler Compiler = new ScriptCompiler();
+            var StartScriptCompile = DateTime.Now.ToString();
+			Compiler.FindAndCompileAllScripts(AdditionalScriptsFolders: AdditionalScriptsFolders);
+            var FinishScriptCompile = DateTime.Now.ToString();
+            CommandUtils.PrintCSVFile(String.Format("UAT,ScriptCompile,{0},{1}", StartScriptCompile, FinishScriptCompile));
+
+			if (GlobalCommandLine.CompileOnly)
+			{
+				Log.TraceInformation("Compilation successful, exiting (CompileOnly)");
+				return;
+			}
+
+			if (GlobalCommandLine.List)
+			{
+				ListAvailableCommands(Compiler.Commands);
+				return;
+			}
+
+			if (GlobalCommandLine.Help)
+			{
+				DisplayHelp(CommandsToExecute, Compiler.Commands);
+				return;
+			}
+
+			// Enable or disable P4 support
+			CommandUtils.InitP4Support(CommandsToExecute, Compiler.Commands);
+			if (CommandUtils.P4Enabled)
+			{
+				Log.TraceInformation("Setting up Perforce environment.");
+				CommandUtils.InitP4Environment();
+				CommandUtils.InitDefaultP4Connection();
+			}
+
+			// Find and execute commands.
+			Execute(CommandsToExecute, Compiler.Commands);
+
+			return;
+		}
+
+		/// <summary>
+		/// Execute commands specified in the command line.
+		/// </summary>
+		/// <param name="CommandsToExecute"></param>
+		/// <param name="Commands"></param>
+		private static void Execute(List<CommandInfo> CommandsToExecute, CaselessDictionary<Type> Commands)
+		{
+			for (int CommandIndex = 0; CommandIndex < CommandsToExecute.Count; ++CommandIndex)
+			{
+				var CommandInfo = CommandsToExecute[CommandIndex];
+				Log.TraceVerbose("Attempting to execute {0}", CommandInfo.ToString());
+				Type CommandType;
+				if (!Commands.TryGetValue(CommandInfo.CommandName, out CommandType))
+				{
+					throw new AutomationException("Failed to find command {0}", CommandInfo.CommandName);
+				}
+				else
+				{
+					BuildCommand Command = (BuildCommand)Activator.CreateInstance(CommandType);
+					Command.Params = CommandInfo.Arguments.ToArray();
+					Command.Execute();
+					// dispose of the class if necessary
+					{
+						var CommandDisposable = Command as IDisposable;
+						if (CommandDisposable != null)
+						{
+							CommandDisposable.Dispose();
+						}
+					}
+
+					// Make sure there's no directories on the stack.
+					CommandUtils.ClearDirStack();
+				}
+			}
+
+			Log.TraceInformation("Script execution successful, exiting.");
+		}
+
+		#endregion
+
+		#region Help
+
+		/// <summary>
+		/// Display help for the specified commands (to execute)
+		/// </summary>
+		/// <param name="CommandsToExecute">List of commands specified in the command line.</param>
+		/// <param name="Commands">All discovered command objects.</param>
+		private static void DisplayHelp(List<CommandInfo> CommandsToExecute, CaselessDictionary<Type> Commands)
+		{
+			for (int CommandIndex = 0; CommandIndex < CommandsToExecute.Count; ++CommandIndex)
+			{
+				var CommandInfo = CommandsToExecute[CommandIndex];
+				Type CommandType;
+				if (Commands.TryGetValue(CommandInfo.CommandName, out CommandType) == false)
+				{
+					Log.TraceError("Help: Failed to find command {0}", CommandInfo.CommandName);
+				}
+				else
+				{
+					CommandUtils.Help(CommandType);
+				}
+			}
+		}
+
+		/// <summary>
+		/// Display AutomationTool.exe help.
+		/// </summary>
+		private static void DisplayHelp()
+		{
+			CommandUtils.LogHelp(typeof(Automation));
+		}
+
+		/// <summary>
+		/// List all available commands.
+		/// </summary>
+		/// <param name="Commands">All vailable commands.</param>
+		private static void ListAvailableCommands(CaselessDictionary<Type> Commands)
+		{
+			string Message = Environment.NewLine;
+			Message += "Available commands:" + Environment.NewLine;
+			foreach (var AvailableCommand in Commands)
+			{
+				Message += String.Format("  {0}{1}", AvailableCommand.Key, Environment.NewLine);
+			}
+			CommandUtils.Log(Message);
+		}
+
+		#endregion
+
+		#region Properties
+
+		/// <summary>
+		/// True if this process is running on a build machine, false if locally.
+		/// </summary>
+		/// <remarks>
+		/// The reason one this property exists in Automation class and not BuildEnvironment is that
+		/// it's required long before BuildEnvironment is initialized.
+		/// </remarks>
+		public static bool IsBuildMachine
+		{
+			get
+			{
+				if (!bIsBuildMachine.HasValue)
+				{
+					throw new AutomationException("Trying to access IsBuildMachine property before it was initialized.");
+				}
+				return (bool)bIsBuildMachine;				
+			}
+			private set
+			{
+				bIsBuildMachine = value;
+			}
+		}
+		private static bool? bIsBuildMachine;
+
+		public static bool ShouldKillProcesses
+		{
+			get
+			{
+				if (!bShouldKillProcesses.HasValue)
+				{
+					throw new AutomationException("Trying to access ShouldKillProcesses property before it was initialized.");
+				}
+				return (bool)bShouldKillProcesses;
+			}
+			private set
+			{
+				bShouldKillProcesses = value;
+			}
+		}
+		private static bool? bShouldKillProcesses;
+		#endregion
+	}
+
+}
diff --git a/Engine/Source/Programs/AutomationTool/AutomationUtils/AutomationExpection.cs b/Engine/Source/Programs/AutomationTool/AutomationUtils/AutomationExpection.cs
new file mode 100644
--- /dev/null
+++ b/Engine/Source/Programs/AutomationTool/AutomationUtils/AutomationExpection.cs
@@ -0,0 +1,105 @@
+﻿// Copyright 1998-2015 Epic Games, Inc. All Rights Reserved.
+using System;
+using System.Collections.Generic;
+using System.Linq;
+using System.Text;
+
+namespace AutomationTool
+{
+    // NOTE: this needs to be kept in sync with EditorAnalytics.h and iPhonePackager.cs
+    public enum ErrorCodes
+    {
+        Error_UATNotFound = -1,
+        Error_Success = 0,
+        Error_Unknown = 1,
+        Error_Arguments = 2,
+        Error_UnknownCommand = 3,
+        Error_SDKNotFound = 10,
+        Error_ProvisionNotFound = 11,
+        Error_CertificateNotFound = 12,
+        Error_ProvisionAndCertificateNotFound = 13,
+        Error_InfoPListNotFound = 14,
+        Error_KeyNotFoundInPList = 15,
+        Error_ProvisionExpired = 16,
+        Error_CertificateExpired = 17,
+        Error_CertificateProvisionMismatch = 18,
+        Error_CodeUnsupported = 19,
+        Error_PluginsUnsupported = 20,
+        Error_UnknownCookFailure = 25,
+        Error_UnknownDeployFailure = 26,
+        Error_UnknownBuildFailure = 27,
+        Error_UnknownPackageFailure = 28,
+        Error_UnknownLaunchFailure = 29,
+        Error_StageMissingFile = 30,
+        Error_FailedToCreateIPA = 31,
+        Error_FailedToCodeSign = 32,
+        Error_DeviceBackupFailed = 33,
+        Error_AppUninstallFailed = 34,
+        Error_AppInstallFailed = 35,
+        Error_AppNotFound = 36,
+        Error_StubNotSignedCorrectly = 37,
+        Error_IPAMissingInfoPList = 38,
+        Error_DeleteFile = 39,
+        Error_DeleteDirectory = 40,
+        Error_CreateDirectory = 41,
+        Error_CopyFile = 42,
+        Error_OnlyOneObbFileSupported = 50,
+        Error_FailureGettingPackageInfo = 51,
+        Error_OnlyOneTargetConfigurationSupported = 52,
+        Error_ObbNotFound = 53,
+        Error_AndroidBuildToolsPathNotFound = 54,
+        Error_NoApkSuitableForArchitecture = 55,
+        Error_FilesInstallFailed = 56,
+        Error_RemoteCertificatesNotFound = 57,
+        Error_LauncherFailed = 100,
+        Error_UATLaunchFailure = 101,
+        Error_FailedToDeleteStagingDirectory = 102,
+        Error_MissingExecutable = 103,
+        Error_DeviceNotSetupForDevelopment = 150,
+        Error_DeviceOSNewerThanSDK = 151,
+    };
+
+    /// <summary>
+    /// Exception class used by the AutomationTool to throw exceptions. Allows setting an exit code that will be passed to the entry routine to return to the system on program exit.
+    /// If no exit code is given, Error_Unkonwn is used.
+    /// </summary>
+    public class AutomationException : System.Exception
+	{
+        public ErrorCodes ErrorCode = ErrorCodes.Error_Unknown;
+
+		public AutomationException()
+		{
+		}
+
+		public AutomationException(string Msg)
+			:base(Msg)
+		{
+		}
+
+        public AutomationException(ErrorCodes ErrorCode, string Msg)
+            : base(Msg)
+        {
+            this.ErrorCode = ErrorCode;
+        }
+
+        public AutomationException(Exception InnerException, string Format, params object[] Args)
+			:base(string.Format(Format, Args), InnerException)
+		{
+		}
+
+        public AutomationException(ErrorCodes ErrorCode, Exception InnerException, string Format, params object[] Args)
+            : base(string.Format(Format, Args), InnerException)
+        {
+            this.ErrorCode = ErrorCode;
+        }
+
+        public AutomationException(string Format, params object[] Args)
+			: base(string.Format(Format, Args)) { }
+
+        public AutomationException(ErrorCodes ErrorCode, string Format, params object[] Args)
+            : base(string.Format(Format, Args)) 
+        {
+            this.ErrorCode = ErrorCode;
+        }
+    }
+}
diff --git a/Engine/Source/Programs/AutomationTool/AutomationUtils/AutomationUtils.Automation.csproj b/Engine/Source/Programs/AutomationTool/AutomationUtils/AutomationUtils.Automation.csproj
new file mode 100644
--- /dev/null
+++ b/Engine/Source/Programs/AutomationTool/AutomationUtils/AutomationUtils.Automation.csproj
@@ -0,0 +1,100 @@
+﻿<?xml version="1.0" encoding="utf-8"?>
+<Project ToolsVersion="11.0" DefaultTargets="Build" xmlns="http://schemas.microsoft.com/developer/msbuild/2003">
+  <Import Project="$(MSBuildExtensionsPath)\$(MSBuildToolsVersion)\Microsoft.Common.props" Condition="Exists('$(MSBuildExtensionsPath)\$(MSBuildToolsVersion)\Microsoft.Common.props')" />
+  <PropertyGroup>
+    <Configuration Condition=" '$(Configuration)' == '' ">Debug</Configuration>
+    <Platform Condition=" '$(Platform)' == '' ">AnyCPU</Platform>
+    <ProjectGuid>{2C96A7F2-B1A3-4258-8E0A-E588FF41A53E}</ProjectGuid>
+    <OutputType>Library</OutputType>
+    <AppDesignerFolder>Properties</AppDesignerFolder>
+    <RootNamespace>AutomationUtils.Automation</RootNamespace>
+    <AssemblyName>AutomationUtils.Automation</AssemblyName>
+    <TargetFrameworkVersion>v4.0</TargetFrameworkVersion>
+    <FileAlignment>512</FileAlignment>
+  </PropertyGroup>
+  <PropertyGroup Condition=" '$(Configuration)|$(Platform)' == 'Debug|AnyCPU' ">
+    <DebugSymbols>true</DebugSymbols>
+    <DebugType>full</DebugType>
+    <Optimize>false</Optimize>
+    <OutputPath>..\..\..\..\Binaries\DotNET\</OutputPath>
+    <DefineConstants>DEBUG;TRACE</DefineConstants>
+    <ErrorReport>prompt</ErrorReport>
+    <WarningLevel>4</WarningLevel>
+    <PlatformTarget>AnyCPU</PlatformTarget>
+    <TreatWarningsAsErrors>true</TreatWarningsAsErrors>
+  </PropertyGroup>
+  <PropertyGroup Condition=" '$(Configuration)|$(Platform)' == 'Development|AnyCPU' ">
+    <DebugType>pdbonly</DebugType>
+    <Optimize>true</Optimize>
+    <OutputPath>..\..\..\..\Binaries\DotNET\</OutputPath>
+    <DefineConstants>TRACE</DefineConstants>
+    <ErrorReport>prompt</ErrorReport>
+    <WarningLevel>4</WarningLevel>
+    <PlatformTarget>AnyCPU</PlatformTarget>
+    <TreatWarningsAsErrors>true</TreatWarningsAsErrors>
+  </PropertyGroup>
+  <ItemGroup>
+    <Reference Include="Ionic.Zip.Reduced">
+      <HintPath>..\..\..\..\Binaries\DotNET\Ionic.Zip.Reduced.dll</HintPath>
+      <Private>False</Private>
+    </Reference>
+    <Reference Include="System" />
+    <Reference Include="System.Core" />
+    <Reference Include="System.Management" />
+    <Reference Include="System.Runtime.Serialization" />
+    <Reference Include="System.Web" />
+    <Reference Include="System.Web.Extensions" />
+    <Reference Include="System.Xml.Linq" />
+    <Reference Include="System.Data.DataSetExtensions" />
+    <Reference Include="Microsoft.CSharp" />
+    <Reference Include="System.Data" />
+    <Reference Include="System.Xml" />
+  </ItemGroup>
+  <ItemGroup>
+    <Compile Include="Automation.cs" />
+    <Compile Include="AutomationExpection.cs" />
+    <Compile Include="BuildCommand.cs" />
+    <Compile Include="BuildEnvironment.cs" />
+    <Compile Include="BuildUtils.cs" />
+    <Compile Include="CommandletUtils.cs" />
+    <Compile Include="CommandUtils.cs" />
+    <Compile Include="DeploymentContext.cs" />
+    <Compile Include="Distiller.cs" />
+    <Compile Include="FileFilter.cs" />
+    <Compile Include="HelpUtils.cs" />
+    <Compile Include="HostPlatform.cs" />
+    <Compile Include="LinuxHostPlatform.cs" />
+    <Compile Include="LocalBuildEnvironment.cs" />
+    <Compile Include="LocalP4Environment.cs" />
+    <Compile Include="Log.cs" />
+    <Compile Include="MacHostPlatform.cs" />
+    <Compile Include="MCPPublic.cs" />
+    <Compile Include="P4Blame.cs" />
+    <Compile Include="P4Environment.cs" />
+    <Compile Include="P4Utils.cs" />
+    <Compile Include="Platform.cs" />
+    <Compile Include="ProcessUtils.cs" />
+    <Compile Include="ProjectParams.cs" />
+    <Compile Include="ProjectUtils.cs" />
+    <Compile Include="Properties\AssemblyInfo.cs" />
+    <Compile Include="ScriptCompiler.cs" />
+    <Compile Include="UBTUtils.cs" />
+    <Compile Include="UE4Build.cs" />
+    <Compile Include="Utils.cs" />
+    <Compile Include="WatchdogTimer.cs" />
+    <Compile Include="WindowsHostPlatform.cs" />
+  </ItemGroup>
+  <ItemGroup>
+    <ProjectReference Include="..\..\DotNETCommon\DotNETUtilities\DotNETUtilities.csproj">
+      <Project>{5d7d66e8-8c76-4af9-b3ec-2ef03421d730}</Project>
+      <Name>DotNETUtilities</Name>
+      <Private>False</Private>
+    </ProjectReference>
+    <ProjectReference Include="..\..\UnrealBuildTool\UnrealBuildTool.csproj">
+      <Project>{fd7c5e1a-cfe4-4fd5-a525-1eb1599a39ac}</Project>
+      <Name>UnrealBuildTool</Name>
+      <Private>False</Private>
+    </ProjectReference>
+  </ItemGroup>
+  <Import Project="$(MSBuildToolsPath)\Microsoft.CSharp.targets" />
+</Project>
\ No newline at end of file
diff --git a/Engine/Source/Programs/AutomationTool/AutomationUtils/BuildEnvironment.cs b/Engine/Source/Programs/AutomationTool/AutomationUtils/BuildEnvironment.cs
new file mode 100644
--- /dev/null
+++ b/Engine/Source/Programs/AutomationTool/AutomationUtils/BuildEnvironment.cs
@@ -0,0 +1,234 @@
+// Copyright 1998-2015 Epic Games, Inc. All Rights Reserved.
+using System;
+using System.Collections.Generic;
+using System.IO;
+using System.Text;
+using System.Reflection;
+using Microsoft.Win32;
+using System.Diagnostics;
+using UnrealBuildTool;
+
+namespace AutomationTool
+{
+	/// <summary>
+	/// Defines the environment variable names that will be used to setup the environemt.
+	/// </summary>
+	static class EnvVarNames
+	{
+		// Command Environment
+		static public readonly string LocalRoot = "uebp_LOCAL_ROOT";
+		static public readonly string LogFolder = "uebp_LogFolder";
+        static public readonly string CSVFile = "uebp_CSVFile";
+		static public readonly string EngineSavedFolder = "uebp_EngineSavedFolder";
+		static public readonly string NETFrameworkDir = "FrameworkDir";
+		static public readonly string NETFrameworkVersion = "FrameworkVersion";
+
+		// Perforce Environment
+		static public readonly string P4Port = "uebp_PORT";		
+		static public readonly string ClientRoot = "uebp_CLIENT_ROOT";
+		static public readonly string Changelist = "uebp_CL";
+		static public readonly string User = "uebp_USER";
+		static public readonly string Client = "uebp_CLIENT";
+		static public readonly string BuildRootP4 = "uebp_BuildRoot_P4";
+		static public readonly string BuildRootEscaped = "uebp_BuildRoot_Escaped";		
+		static public readonly string LabelToSync = "uebp_LabelToSync";
+		static public readonly string P4Password = "uebp_PASS";
+	}
+
+
+	/// <summary>
+	/// Environment to allow access to commonly used environment variables.
+	/// </summary>
+	public class CommandEnvironment
+	{
+		/// <summary>
+		/// Path to a file we know to always exist under the UE4 root directory.
+		/// </summary>
+		public static readonly string KnownFileRelativeToRoot = @"Engine/Config/BaseEngine.ini";
+
+		#region Command Environment properties
+
+		public string LocalRoot { get; protected set; }
+		public string EngineSavedFolder { get; protected set; }
+		public string LogFolder { get; protected set; }
+        public string CSVFile { get; protected set; }
+		public string RobocopyExe { get; protected set; }
+		public string MountExe { get; protected set; }
+		public string CmdExe { get; protected set; }
+		public string UATExe { get; protected set; }		
+		public string TimestampAsString { get; protected set; }
+		public bool HasCapabilityToCompile { get; protected set; }
+		public string MsBuildExe { get; protected set; }
+		public string MsDevExe { get; protected set; }
+
+		#endregion
+
+		internal CommandEnvironment()
+		{
+			InitEnvironment();
+		}
+
+		/// <summary>
+		/// Sets the location of the exe.
+		/// </summary>
+		protected void SetUATLocation()
+		{
+			if (String.IsNullOrEmpty(UATExe))
+			{
+				UATExe = CommandUtils.CombinePaths(Path.GetFullPath(InternalUtils.ExecutingAssemblyLocation));
+			}
+			if (!CommandUtils.FileExists_NoExceptions(UATExe))
+			{
+				throw new AutomationException("Could not find AutomationTool.exe. Reflection indicated it was here: {0}", UATExe);
+			}
+		}
+
+		/// <summary>
+		/// Sets the location of the AutomationTool/Saved directory
+		/// </summary>
+		protected void SetUATSavedPath()
+		{
+			var LocalRootPath = CommandUtils.GetEnvVar(EnvVarNames.LocalRoot);
+			var SavedPath = CommandUtils.CombinePaths(PathSeparator.Slash, LocalRootPath, "Engine", "Programs", "AutomationTool", "Saved");
+			CommandUtils.SetEnvVar(EnvVarNames.EngineSavedFolder, SavedPath);
+		}
+
+		/// <summary>
+		/// Initializes the environement.
+		/// </summary>
+		protected virtual void InitEnvironment()
+		{
+			SetUATLocation();
+
+			LocalRoot = CommandUtils.GetEnvVar(EnvVarNames.LocalRoot);
+			if (String.IsNullOrEmpty(CommandUtils.GetEnvVar(EnvVarNames.EngineSavedFolder)))
+			{
+				SetUATSavedPath();
+			}
+
+			if (LocalRoot.EndsWith(":"))
+			{
+				LocalRoot += Path.DirectorySeparatorChar;
+			}
+
+			EngineSavedFolder = CommandUtils.GetEnvVar(EnvVarNames.EngineSavedFolder);
+            CSVFile = CommandUtils.GetEnvVar(EnvVarNames.CSVFile);            
+			LogFolder = CommandUtils.GetEnvVar(EnvVarNames.LogFolder);
+			RobocopyExe = CommandUtils.CombinePaths(Environment.SystemDirectory, "robocopy.exe");
+			MountExe = CommandUtils.CombinePaths(Environment.SystemDirectory, "mount.exe");
+			CmdExe = Utils.IsRunningOnMono ? "/bin/sh" : CommandUtils.CombinePaths(Environment.SystemDirectory, "cmd.exe");
+
+			if (String.IsNullOrEmpty(LogFolder))
+			{
+				throw new AutomationException("Environment is not set up correctly: LogFolder is not set!");
+			}
+
+			if (String.IsNullOrEmpty(LocalRoot))
+			{
+				throw new AutomationException("Environment is not set up correctly: LocalRoot is not set!");
+			}
+
+			if (String.IsNullOrEmpty(EngineSavedFolder))
+			{
+				throw new AutomationException("Environment is not set up correctly: EngineSavedFolder is not set!");
+			}
+
+			// Make sure that the log folder exists
+			var LogFolderInfo = new DirectoryInfo(LogFolder);
+			if (!LogFolderInfo.Exists)
+			{
+				LogFolderInfo.Create();
+			}
+
+			// Setup the timestamp string
+			DateTime LocalTime = DateTime.Now;
+
+			string TimeStamp = LocalTime.Year + "-"
+						+ LocalTime.Month.ToString("00") + "-"
+						+ LocalTime.Day.ToString("00") + "_"
+						+ LocalTime.Hour.ToString("00") + "."
+						+ LocalTime.Minute.ToString("00") + "."
+						+ LocalTime.Second.ToString("00");
+
+			TimestampAsString = TimeStamp;
+
+			SetupBuildEnvironment();
+
+			LogSettings();
+		}
+
+		void LogSettings()
+		{
+			Log.TraceVerbose("Command Environment settings:");
+			Log.TraceVerbose("CmdExe={0}", CmdExe);
+			Log.TraceVerbose("EngineSavedFolder={0}", EngineSavedFolder);
+			Log.TraceVerbose("HasCapabilityToCompile={0}", HasCapabilityToCompile);
+			Log.TraceVerbose("LocalRoot={0}", LocalRoot);
+			Log.TraceVerbose("LogFolder={0}", LogFolder);
+			Log.TraceVerbose("MountExe={0}", MountExe);
+			Log.TraceVerbose("MsBuildExe={0}", MsBuildExe);
+			Log.TraceVerbose("MsDevExe={0}", MsDevExe);
+			Log.TraceVerbose("RobocopyExe={0}", RobocopyExe);
+			Log.TraceVerbose("TimestampAsString={0}", TimestampAsString);
+			Log.TraceVerbose("UATExe={0}", UATExe);			
+		}
+
+		#region Compiler Setup
+
+		/// <summary>
+		/// Initializes build environemnt: finds the path to msbuild.exe
+		/// </summary>
+		void SetupBuildEnvironment()
+		{
+			// Assume we have the capability co compile.
+			HasCapabilityToCompile = true;
+
+			try
+			{
+				HostPlatform.Current.SetFrameworkVars();
+			}
+			catch (Exception)
+			{
+				// Something went wrong, we can't compile.
+				Log.WriteLine(TraceEventType.Warning, "SetFrameworkVars failed. Assuming no compilation capability.");
+				HasCapabilityToCompile = false;
+			}
+
+			if (HasCapabilityToCompile)
+			{
+				try
+				{
+					MsBuildExe = HostPlatform.Current.GetMsBuildExe();
+				}
+				catch (Exception Ex)
+				{
+					Log.WriteLine(TraceEventType.Warning, Ex.Message);
+					Log.WriteLine(TraceEventType.Warning, "Assuming no compilation capability.");
+					HasCapabilityToCompile = false;
+					MsBuildExe = "";
+				}
+			}
+
+			if (HasCapabilityToCompile)
+			{
+				try
+				{
+					MsDevExe = HostPlatform.Current.GetMsDevExe();
+				}
+				catch (Exception Ex)
+				{
+					Log.WriteLine(TraceEventType.Warning, Ex.Message);
+					Log.WriteLine(TraceEventType.Warning, "Assuming no solution compilation capability.");
+					MsDevExe = "";
+				}
+			}
+
+			Log.TraceVerbose("CompilationEvironment.HasCapabilityToCompile={0}", HasCapabilityToCompile);
+			Log.TraceVerbose("CompilationEvironment.MsBuildExe={0}", MsBuildExe);
+			Log.TraceVerbose("CompilationEvironment.MsDevExe={0}", MsDevExe);
+		}
+
+		#endregion
+	}
+
+}
diff --git a/Engine/Source/Programs/AutomationTool/AutomationUtils/BuildUtils.cs b/Engine/Source/Programs/AutomationTool/AutomationUtils/BuildUtils.cs
new file mode 100644
--- /dev/null
+++ b/Engine/Source/Programs/AutomationTool/AutomationUtils/BuildUtils.cs
@@ -0,0 +1,118 @@
+﻿// Copyright 1998-2015 Epic Games, Inc. All Rights Reserved.
+using System;
+using System.Collections.Generic;
+using System.Text;
+using System.IO;
+
+namespace AutomationTool
+{
+
+	public partial class CommandUtils
+	{
+		/// <summary>
+		/// Runs msbuild.exe with the specified arguments. Automatically creates a logfile. When 
+		/// no LogName is specified, the executable name is used as logfile base name.
+		/// </summary>
+		/// <param name="Env">BuildEnvironment to use.</param>
+		/// <param name="Project">Path to the project to build.</param>
+		/// <param name="Arguments">Arguments to pass to msbuild.exe.</param>
+		/// <param name="LogName">Optional logfile name.</param>
+		public static void MsBuild(CommandEnvironment Env, string Project, string Arguments, string LogName)
+		{
+			if (String.IsNullOrEmpty(Env.MsBuildExe))
+			{
+				throw new AutomationException("Unable to find msbuild.exe at: \"{0}\"", Env.MsBuildExe);
+			}
+			if (!FileExists(Project))
+			{
+				throw new AutomationException("Project {0} does not exist!", Project);
+			}
+			var RunArguments = MakePathSafeToUseWithCommandLine(Project);
+			if (!String.IsNullOrEmpty(Arguments))
+			{
+				RunArguments += " " + Arguments;
+			}
+			RunAndLog(Env, Env.MsBuildExe, RunArguments, LogName);
+		}
+
+		/// <summary>
+		/// Builds a Visual Studio solution with MsDevEnv. Automatically creates a logfile. When 
+		/// no LogName is specified, the executable name is used as logfile base name.
+		/// </summary>
+		/// <param name="Env">BuildEnvironment to use.</param>
+		/// <param name="SolutionFile">Path to the solution file</param>
+		/// <param name="BuildConfig">Configuration to build.</param>
+		/// <param name="LogName">Optional logfile name.</param>
+		public static void BuildSolution(CommandEnvironment Env, string SolutionFile, string BuildConfig = "Development", string LogName = null)
+		{
+			if (!FileExists(SolutionFile))
+			{
+				throw new AutomationException(String.Format("Unabled to build Solution {0}. Solution file not found.", SolutionFile));
+			}
+			if (String.IsNullOrEmpty(Env.MsDevExe))
+			{
+				throw new AutomationException(String.Format("Unabled to build Solution {0}. devenv.com not found.", SolutionFile));
+			}
+			string CmdLine = String.Format("\"{0}\" /build \"{1}\"", SolutionFile, BuildConfig);
+            var StartBuildSolution = DateTime.Now.ToString();
+			RunAndLog(Env, Env.MsDevExe, CmdLine, LogName);
+            var FinishBuildSolution = DateTime.Now.ToString();
+            PrintCSVFile(String.Format("UAT,Compile,{0},{1}", StartBuildSolution, FinishBuildSolution));
+		}
+
+		/// <summary>
+		/// Builds a CSharp project with msbuild.exe. Automatically creates a logfile. When 
+		/// no LogName is specified, the executable name is used as logfile base name.
+		/// </summary>
+		/// <param name="Env">BuildEnvironment to use.</param>
+		/// <param name="ProjectFile">Path to the project file. Must be a .csproj file.</param>
+		/// <param name="BuildConfig">Configuration to build.</param>
+		/// <param name="LogName">Optional logfile name.</param>
+		public static void BuildCSharpProject(CommandEnvironment Env, string ProjectFile, string BuildConfig = "Development", string LogName = null)
+		{
+			if (!ProjectFile.EndsWith(".csproj"))
+			{
+				throw new AutomationException(String.Format("Unabled to build Project {0}. Not a valid .csproj file.", ProjectFile));
+			}
+			if (!FileExists(ProjectFile))
+			{
+				throw new AutomationException(String.Format("Unabled to build Project {0}. Project file not found.", ProjectFile));
+			}
+
+			string CmdLine = String.Format(@"/verbosity:normal /target:Rebuild /property:Configuration={0} /property:Platform=AnyCPU", BuildConfig);
+			MsBuild(Env, ProjectFile, CmdLine, LogName);
+		}
+
+        /// <summary>
+        /// returns true if this is a mac executable using some awful conventions
+        /// <param name="Filename">Filename</param>
+        /// </summary>
+        public static bool IsProbablyAMacOrIOSExe(string Filename)
+        {
+            return
+                Path.GetExtension(Filename) == ".sh" ||
+                Path.GetExtension(Filename) == ".command" ||
+                ((
+                    CommandUtils.CombinePaths(Filename).ToLower().Contains(CommandUtils.CombinePaths("Binaries", "Mac").ToLower()) ||
+                    CommandUtils.CombinePaths(Filename).ToLower().Contains(CommandUtils.CombinePaths("Binaries", "IOS").ToLower()) ||
+					CommandUtils.CombinePaths(Filename).ToLower().Contains(CommandUtils.CombinePaths("Binaries", "ThirdParty").ToLower()) ||
+					CommandUtils.CombinePaths(Filename).ToLower().Contains(".app/Contents/MacOS".ToLower())
+                ) && (Path.GetExtension(Filename) == "" || Path.GetExtension(Filename) == "."));
+        }
+
+		/// <summary>
+		/// Sets an executable bit for Unix executables and adds read permission for all users plus write permission for owner.
+		/// <param name="Filename">Filename</param>
+		/// </summary>
+		public static void FixUnixFilePermissions(string Filename)
+		{
+			string Permissions = IsProbablyAMacOrIOSExe(Filename) ? "0755" : "0644";
+			var Result = CommandUtils.Run("sh", string.Format("-c 'chmod {0} \"{1}\"'", Permissions, Filename.Replace("'", "'\"'\"'")));
+			if (Result.ExitCode != 0)
+			{
+				throw new AutomationException(String.Format("Failed to chmod \"{0}\"", Filename));
+			}
+		}
+	}
+
+}
