UEB-260 - Break AutomationTool into AutomationUtils that all automation projects depend on, and AutomationTool, which essentially only contains the startup code.

diff --git a/Engine/Source/Programs/AutomationTool/AutomationUtils/CommandUtils.cs b/Engine/Source/Programs/AutomationTool/AutomationUtils/CommandUtils.cs
new file mode 100644
--- /dev/null
+++ b/Engine/Source/Programs/AutomationTool/AutomationUtils/CommandUtils.cs
@@ -0,0 +1,2537 @@
+﻿// Copyright 1998-2015 Epic Games, Inc. All Rights Reserved.
+using System;
+using System.Collections;
+using System.Collections.Generic;
+using System.Diagnostics;
+using System.IO;
+using System.Net;
+using System.Runtime.Serialization;
+using System.Runtime.Serialization.Json;
+using System.Text;
+using System.Threading;
+using Tools.DotNETCommon.XmlHandler;
+using UnrealBuildTool;
+using System.Runtime.CompilerServices;
+using System.Linq;
+using Tools.DotNETCommon.CaselessDictionary;
+using Tools.DotNETCommon.HarvestEnvVars;
+
+namespace AutomationTool
+{
+	#region ParamList
+
+	/// <summary>
+	/// Wrapper around List with support for multi parameter constructor, i.e:
+	///   var Maps = new ParamList<string>("Map1", "Map2");
+	/// </summary>
+	/// <typeparam name="T"></typeparam>
+	public class ParamList<T> : List<T>
+	{
+		public ParamList(params T[] Args)
+		{
+			AddRange(Args);
+		}
+
+		public ParamList(ICollection<T> Collection)
+			: base(Collection != null ? Collection : new T[] {})
+		{
+		}
+
+		public override string ToString()
+		{
+			var Text = "";
+			for (int Index = 0; Index < Count; ++Index)
+			{
+				if (Index > 0)
+				{
+					Text += ", ";
+				}
+				Text += this[Index].ToString();				
+			}
+			return Text;
+		}
+	}
+
+	#endregion
+
+	#region PathSeparator
+
+	public enum PathSeparator
+	{
+		Default = 0,
+		Slash,
+		Backslash,
+		Depot,
+		Local
+	}
+
+	#endregion
+
+	/// <summary>
+	/// Base utility function for script commands.
+	/// </summary>
+	public partial class CommandUtils
+	{
+		#region Environment Setup
+
+		static private CommandEnvironment CmdEnvironment;
+
+		/// <summary>
+		/// BuildEnvironment to use for this buildcommand. This is initialized by InitBuildEnvironment. As soon
+		/// as the script execution in ExecuteBuild begins, the BuildEnv is set up and ready to use.
+		/// </summary>
+		static public CommandEnvironment CmdEnv
+		{
+			get
+			{
+				if (CmdEnvironment == null)
+				{
+					throw new AutomationException("Attempt to use CommandEnvironment before it was initialized.");
+				}
+				return CmdEnvironment;
+			}
+		}
+
+		/// <summary>
+		/// Initializes build environment. If the build command needs a specific env-var mapping or
+		/// has an extended BuildEnvironment, it must implement this method accordingly.
+		/// </summary>
+		/// <returns>Initialized and ready to use BuildEnvironment</returns>
+		static internal void InitCommandEnvironment()
+		{
+			CmdEnvironment = Automation.IsBuildMachine ? new CommandEnvironment() : new LocalCommandEnvironment(); ;
+		}
+
+		#endregion
+
+		#region Logging
+
+		/// <summary>
+		/// Returns a formatted string for the specified exception (including inner exception if any).
+		/// </summary>
+		/// <param name="Ex"></param>
+		/// <returns>Formatted exception string.</returns>
+		public static string ExceptionToString(Exception Ex)
+		{
+			return LogUtils.FormatException(Ex);
+		}
+
+		/// <summary>
+		/// Writes formatted text to log (with TraceEventType.Log).
+		/// </summary>
+		/// <param name="Format">Format string</param>
+		/// <param name="Args">Parameters</param>
+		[MethodImplAttribute(MethodImplOptions.NoInlining)]
+		public static void Log(string Format, params object[] Args)
+		{
+			UnrealBuildTool.Log.WriteLine(1, null, TraceEventType.Information, Format, Args);
+		}
+
+		/// <summary>
+		/// Writes formatted text to log (with TraceEventType.Log).
+		/// </summary>
+		/// <param name="Message">Text</param>
+		[MethodImplAttribute(MethodImplOptions.NoInlining)]
+		public static void Log(string Message)
+		{
+            UnrealBuildTool.Log.WriteLine(1, null, TraceEventType.Information, Message);
+		}
+
+		/// <summary>
+		/// Writes formatted text to log (with TraceEventType.Error).
+		/// </summary>
+		/// <param name="Format">Format string</param>
+		/// <param name="Args">Parameters</param>
+		[MethodImplAttribute(MethodImplOptions.NoInlining)]
+		public static void LogError(string Format, params object[] Args)
+		{
+            UnrealBuildTool.Log.WriteLine(1, null, TraceEventType.Error, Format, Args);
+		}
+
+		/// <summary>
+		/// Writes formatted text to log (with TraceEventType.Error).
+		/// </summary>
+		/// <param name="Message">Text</param>
+		[MethodImplAttribute(MethodImplOptions.NoInlining)]
+		public static void LogError(string Message)
+		{
+            UnrealBuildTool.Log.WriteLine(1, null, TraceEventType.Error, Message);
+		}
+
+		/// <summary>
+		/// Writes formatted text to log (with TraceEventType.Warning).
+		/// </summary>
+		/// <param name="Format">Format string</param>
+		/// <param name="Args">Parameters</param>
+		[MethodImplAttribute(MethodImplOptions.NoInlining)]
+		public static void LogWarning(string Format, params object[] Args)
+		{
+            UnrealBuildTool.Log.WriteLine(1, null, TraceEventType.Warning, Format, Args);
+		}
+
+		/// <summary>
+		/// Writes a message to log (with TraceEventType.Warning).
+		/// </summary>
+		/// <param name="Message">Text</param>
+		[MethodImplAttribute(MethodImplOptions.NoInlining)]
+		public static void LogWarning(string Message)
+		{
+            UnrealBuildTool.Log.WriteLine(1, null, TraceEventType.Warning, Message);
+		}
+
+		/// <summary>
+		/// Writes formatted text to log (with TraceEventType.Verbose).
+		/// </summary>
+		/// <param name="Foramt">Format string</param>
+		/// <param name="Args">Arguments</param>
+		[MethodImplAttribute(MethodImplOptions.NoInlining)]
+		public static void LogVerbose(string Format, params object[] Args)
+		{
+            UnrealBuildTool.Log.WriteLine(1, null, TraceEventType.Verbose, Format, Args);
+		}
+
+		/// <summary>
+		/// Writes formatted text to log (with TraceEventType.Verbose).
+		/// </summary>
+		/// <param name="Message">Text</param>
+		[MethodImplAttribute(MethodImplOptions.NoInlining)]
+		public static void LogVerbose(string Message)
+		{
+            UnrealBuildTool.Log.WriteLine(1, null, TraceEventType.Verbose, Message);
+		}
+
+		/// <summary>
+		/// Writes formatted text to log.
+		/// </summary>
+		/// <param name="Verbosity">Verbosity</param>
+		/// <param name="Format">Format string</param>
+		/// <param name="Args">Arguments</param>
+		[MethodImplAttribute(MethodImplOptions.NoInlining)]
+		public static void Log(TraceEventType Verbosity, string Format, params object[] Args)
+		{
+            UnrealBuildTool.Log.WriteLine(1, null, Verbosity, Format, Args);
+		}
+
+		/// <summary>
+		/// Writes formatted text to log.
+		/// </summary>
+		/// <param name="Verbosity">Verbosity</param>
+		/// <param name="Message">Text</param>
+		[MethodImplAttribute(MethodImplOptions.NoInlining)]
+		public static void Log(TraceEventType Verbosity, string Message)
+		{
+            UnrealBuildTool.Log.WriteLine(1, null, Verbosity, Message);
+		}
+
+		/// <summary>
+		/// Dumps exception to log.
+		/// </summary>
+		/// <param name="Verbosity">Verbosity</param>
+		/// <param name="Ex">Exception</param>
+		[MethodImplAttribute(MethodImplOptions.NoInlining)]
+		public static void Log(TraceEventType Verbosity, Exception Ex)
+		{
+            UnrealBuildTool.Log.WriteLine(1, null, Verbosity, LogUtils.FormatException(Ex));
+		}
+
+		#endregion
+
+		#region Progress Logging
+
+		public static void LogPushProgress(bool bShowProgress, int Numerator, int Denominator)
+		{
+			if(bShowProgress)
+			{
+				Log("[@progress push {0}/{1} skipline]", Numerator, Denominator);
+			}
+		}
+
+		public static void LogPopProgress(bool bShowProgress)
+		{
+			if(bShowProgress)
+			{
+				Log("[@progress pop skipline]");
+			}
+		}
+
+		public static void LogIncrementProgress(bool bShowProgress, int Numerator, int Denominator)
+		{
+			if(bShowProgress)
+			{
+				Log("[@progress increment {0}/{1} skipline]", Numerator, Denominator);
+			}
+		}
+
+		public static void LogSetProgress(bool bShowProgress, string Format, params string[] Args)
+		{
+			if(bShowProgress)
+			{
+				Log("[@progress '{0}' skipline]", String.Format(Format, Args));
+			}
+		}
+
+		public static void LogSetProgress(bool bShowProgress, int Numerator, int Denominator, string Format, params string[] Args)
+		{
+			if(bShowProgress)
+			{
+				Log("[@progress {0}/{1} '{2}' skipline]", Numerator, Denominator, String.Format(Format, Args));
+			}
+		}
+
+		#endregion
+
+		#region IO
+
+		/// <summary>
+		/// Finds files in specified paths. 
+		/// </summary>
+		/// <param name="SearchPattern">Pattern</param>
+		/// <param name="Recursive">Recursive search</param>
+		/// <param name="Paths">Paths to search</param>
+		/// <returns>An array of files found in the specified paths</returns>
+		public static string[] FindFiles(string SearchPattern, bool Recursive, params string[] Paths)
+		{
+			List<string> FoundFiles = new List<string>();
+			foreach (var PathToSearch in Paths)
+			{
+				var NormalizedPath = ConvertSeparators(PathSeparator.Default, PathToSearch);
+				if (DirectoryExists(NormalizedPath))
+				{
+					var FoundInPath = InternalUtils.SafeFindFiles(NormalizedPath, SearchPattern, Recursive);
+					if (FoundInPath == null)
+					{
+						throw new AutomationException(String.Format("Failed to find files in '{0}'", NormalizedPath));
+					}
+					FoundFiles.AddRange(FoundInPath);
+				}
+			}
+			return FoundFiles.ToArray();
+		}
+
+		/// <summary>
+		/// Finds files in specified paths. 
+		/// </summary>
+		/// <param name="SearchPattern">Pattern</param>
+		/// <param name="Recursive">Recursive search</param>
+		/// <param name="Paths">Paths to search</param>
+		/// <returns>An array of files found in the specified paths</returns>
+		public static string[] FindFiles_NoExceptions(string SearchPattern, bool Recursive, params string[] Paths)
+		{
+			List<string> FoundFiles = new List<string>();
+			foreach (var PathToSearch in Paths)
+			{
+				var NormalizedPath = ConvertSeparators(PathSeparator.Default, PathToSearch);
+				if (DirectoryExists(NormalizedPath))
+				{
+					var FoundInPath = InternalUtils.SafeFindFiles(NormalizedPath, SearchPattern, Recursive);
+					if (FoundInPath != null)
+					{
+						FoundFiles.AddRange(FoundInPath);
+					}
+				}
+			}
+			return FoundFiles.ToArray();
+		}
+        /// <summary>
+        /// Finds files in specified paths. 
+        /// </summary>
+        /// <param name="SearchPattern">Pattern</param>
+        /// <param name="Recursive">Recursive search</param>
+        /// <param name="Paths">Paths to search</param>
+        /// <returns>An array of files found in the specified paths</returns>
+        public static string[] FindFiles_NoExceptions(bool bQuiet, string SearchPattern, bool Recursive, params string[] Paths)
+        {
+            List<string> FoundFiles = new List<string>();
+            foreach (var PathToSearch in Paths)
+            {
+                var NormalizedPath = ConvertSeparators(PathSeparator.Default, PathToSearch);
+                if (DirectoryExists(NormalizedPath))
+                {
+                    var FoundInPath = InternalUtils.SafeFindFiles(NormalizedPath, SearchPattern, Recursive, bQuiet);
+                    if (FoundInPath != null)
+                    {
+                        FoundFiles.AddRange(FoundInPath);
+                    }
+                }
+            }
+            return FoundFiles.ToArray();
+        }
+		/// <summary>
+		/// Finds files in specified paths. 
+		/// </summary>
+		/// <param name="SearchPattern">Pattern</param>
+		/// <param name="Recursive">Recursive search</param>
+		/// <param name="Paths">Paths to search</param>
+		/// <returns>An array of files found in the specified paths</returns>
+		public static string[] FindDirectories(bool bQuiet, string SearchPattern, bool Recursive, params string[] Paths)
+		{
+			List<string> FoundDirs = new List<string>();
+			foreach (var PathToSearch in Paths)
+			{
+				var NormalizedPath = ConvertSeparators(PathSeparator.Default, PathToSearch);
+				if (DirectoryExists(NormalizedPath))
+				{
+					var FoundInPath = InternalUtils.SafeFindDirectories(NormalizedPath, SearchPattern, Recursive, bQuiet);
+					if (FoundInPath == null)
+					{
+						throw new AutomationException(String.Format("Failed to find directories in '{0}'", NormalizedPath));
+					}
+					FoundDirs.AddRange(FoundInPath);
+				}
+			}
+			return FoundDirs.ToArray();
+		}
+		/// <summary>
+		/// Finds Directories in specified paths. 
+		/// </summary>
+		/// <param name="SearchPattern">Pattern</param>
+		/// <param name="Recursive">Recursive search</param>
+		/// <param name="Paths">Paths to search</param>
+		/// <returns>An array of files found in the specified paths</returns>
+		public static string[] FindDirectories_NoExceptions(bool bQuiet, string SearchPattern, bool Recursive, params string[] Paths)
+		{
+			List<string> FoundDirs = new List<string>();
+			foreach (var PathToSearch in Paths)
+			{
+				var NormalizedPath = ConvertSeparators(PathSeparator.Default, PathToSearch);
+				if (DirectoryExists(NormalizedPath))
+				{
+					var FoundInPath = InternalUtils.SafeFindDirectories(NormalizedPath, SearchPattern, Recursive, bQuiet);
+					if (FoundInPath != null)
+					{
+						FoundDirs.AddRange(FoundInPath);
+					}
+				}
+			}
+			return FoundDirs.ToArray();
+		}
+		/// <summary>
+		/// Deletes a file(s). 
+		/// If the file does not exist, silently succeeds.
+		/// If the deletion of the file fails, this function throws an Exception.
+		/// </summary>
+		/// <param name="Filename">Filename</param>
+		public static void DeleteFile(params string[] Filenames)
+		{
+			foreach (var Filename in Filenames)
+			{
+				var NormalizedFilename = ConvertSeparators(PathSeparator.Default, Filename);
+				if (!InternalUtils.SafeDeleteFile(NormalizedFilename))
+				{
+					throw new AutomationException(String.Format("Failed to delete file '{0}'", NormalizedFilename));
+				}
+			}
+		}
+        /// <summary>
+        /// Deletes a file(s). 
+        /// If the file does not exist, silently succeeds.
+        /// If the deletion of the file fails, this function throws an Exception.
+        /// </summary>
+        /// <param name="Filename">Filename</param>
+        public static void DeleteFile(bool bQuiet, params string[] Filenames)
+        {
+            foreach (var Filename in Filenames)
+            {
+                var NormalizedFilename = ConvertSeparators(PathSeparator.Default, Filename);
+                if (!InternalUtils.SafeDeleteFile(NormalizedFilename, bQuiet))
+                {
+                    throw new AutomationException(String.Format("Failed to delete file '{0}'", NormalizedFilename));
+                }
+            }
+        }
+
+		/// <summary>
+		/// Deletes a file(s). 
+		/// If the deletion of the file fails, prints a warning.
+		/// </summary>
+		/// <param name="Filename">Filename</param>
+        public static bool DeleteFile_NoExceptions(params string[] Filenames)
+		{
+			bool Result = true;
+			foreach (var Filename in Filenames)
+			{
+				var NormalizedFilename = ConvertSeparators(PathSeparator.Default, Filename);
+				if (!InternalUtils.SafeDeleteFile(NormalizedFilename))
+				{
+					Log(TraceEventType.Warning, "Failed to delete file '{0}'", NormalizedFilename);
+					Result = false;
+				}
+			}
+			return Result;
+		}
+		/// <summary>
+		/// Deletes a file(s). 
+		/// If the deletion of the file fails, prints a warning.
+		/// </summary>
+		/// <param name="Filename">Filename</param>
+        /// <param name="bQuiet">if true, then don't retry and don't print much.</param>
+        public static bool DeleteFile_NoExceptions(string Filename, bool bQuiet = false)
+		{
+			bool Result = true;
+			var NormalizedFilename = ConvertSeparators(PathSeparator.Default, Filename);
+			if (!InternalUtils.SafeDeleteFile(NormalizedFilename, bQuiet))
+			{
+				Log(bQuiet ? TraceEventType.Information : TraceEventType.Warning, "Failed to delete file '{0}'", NormalizedFilename);
+				Result = false;
+			}
+			return Result;
+		}
+
+		/// <summary>
+		/// Deletes a directory(or directories) including its contents (recursively, will delete read-only files).
+		/// If the deletion of the directory fails, this function throws an Exception.
+		/// </summary>
+		/// <param name="bQuiet">Suppresses log output if true</param>
+		/// <param name="Directory">Directory</param>
+		public static void DeleteDirectory(bool bQuiet, params string[] Directories)
+		{
+			foreach (var Directory in Directories)
+			{
+				var NormalizedDirectory = ConvertSeparators(PathSeparator.Default, Directory);
+				if (!InternalUtils.SafeDeleteDirectory(NormalizedDirectory, bQuiet))
+				{
+					throw new AutomationException(String.Format("Failed to delete directory '{0}'", NormalizedDirectory));
+				}
+			}
+		}
+
+		/// <summary>
+		/// Deletes a directory(or directories) including its contents (recursively, will delete read-only files).
+		/// If the deletion of the directory fails, this function throws an Exception.
+		/// </summary>
+		/// <param name="Directory">Directory</param>
+		public static void DeleteDirectory(params string[] Directories)
+		{
+			DeleteDirectory(false, Directories);
+		}
+
+		/// <summary>
+		/// Deletes a directory(or directories) including its contents (recursively, will delete read-only files).
+		/// If the deletion of the directory fails, prints a warning.
+		/// </summary>
+		/// <param name="bQuiet">Suppresses log output if true</param>
+		/// <param name="Directory">Directory</param>
+		public static bool DeleteDirectory_NoExceptions(bool bQuiet, params string[] Directories)
+		{
+			bool Result = true;
+            foreach (var Directory in Directories)
+            {
+                var NormalizedDirectory = ConvertSeparators(PathSeparator.Default, Directory);
+                try
+                {
+                    if (!InternalUtils.SafeDeleteDirectory(NormalizedDirectory, bQuiet))
+                    {
+                        Log(TraceEventType.Warning, "Failed to delete directory '{0}'", NormalizedDirectory);
+                        Result = false;
+                    }
+                }
+                catch (Exception Ex)
+                {
+					if (!bQuiet)
+					{
+						Log(TraceEventType.Warning, "Failed to delete directory, exception '{0}'", NormalizedDirectory);
+						Log(TraceEventType.Warning, Ex);
+					}
+                    Result = false;
+                }
+            }
+
+
+			return Result;
+		}
+
+		/// <summary>
+		/// Deletes a directory(or directories) including its contents (recursively, will delete read-only files).
+		/// If the deletion of the directory fails, prints a warning.
+		/// </summary>
+		/// <param name="Directory">Directory</param>
+		public static bool DeleteDirectory_NoExceptions(params string[] Directories)
+		{
+			return DeleteDirectory_NoExceptions(false, Directories);
+		}
+
+
+		/// <summary>
+		/// Attempts to delete a directory, if that fails deletes all files and folder from the specified directory.
+		/// This works around the issue when the user has a file open in a notepad from that directory. Somehow deleting the file works but
+		/// deleting the directory with the file that's open, doesn't.
+		/// </summary>
+		/// <param name="DirectoryName"></param>
+		public static void DeleteDirectoryContents(string DirectoryName)
+		{
+			Log("DeleteDirectoryContents({0})", DirectoryName);
+			const bool bQuiet = true;
+			var Files = CommandUtils.FindFiles_NoExceptions(bQuiet, "*", false, DirectoryName);
+			foreach (var Filename in Files)
+			{
+				CommandUtils.DeleteFile_NoExceptions(Filename);
+			}
+			var Directories = CommandUtils.FindDirectories_NoExceptions(bQuiet, "*", false, DirectoryName);
+			foreach (var SubDirectoryName in Directories)
+			{
+				CommandUtils.DeleteDirectory_NoExceptions(bQuiet, SubDirectoryName);
+			}
+		}
+
+		/// <summary>
+		/// Checks if a directory(or directories) exists.
+		/// </summary>
+		/// <param name="Directory">Directory</param>
+		/// <returns>True if the directory exists, false otherwise.</returns>
+		public static bool DirectoryExists(params string[] Directories)
+		{
+			bool bExists = Directories.Length > 0;
+			foreach (var DirectoryName in Directories)
+			{
+				var NormalizedDirectory = ConvertSeparators(PathSeparator.Default, DirectoryName);
+				bExists = System.IO.Directory.Exists(NormalizedDirectory) && bExists;
+			}
+			return bExists;
+		}
+
+		/// <summary>
+		/// Checks if a directory(or directories) exists.
+		/// </summary>
+		/// <param name="Directory">Directory</param>
+		/// <returns>True if the directory exists, false otherwise.</returns>
+		public static bool DirectoryExists_NoExceptions(params string[] Directories)
+		{
+			bool bExists = Directories.Length > 0;
+			foreach (var DirectoryName in Directories)
+			{
+				var NormalizedDirectory = ConvertSeparators(PathSeparator.Default, DirectoryName);
+				try
+				{
+					bExists = System.IO.Directory.Exists(NormalizedDirectory) && bExists;
+				}
+				catch (Exception Ex)
+				{
+					Log(TraceEventType.Warning, "Unable to check if directory exists: {0}", NormalizedDirectory);
+					Log(TraceEventType.Warning, Ex);
+					bExists = false;
+					break;
+				}
+			}
+			return bExists;
+		}
+
+		/// <summary>
+		/// Creates a directory(or directories).
+		/// If the creation of the directory fails, this function throws an Exception.
+		/// </summary>
+		/// <param name="Directory">Directory</param>
+		public static void CreateDirectory(params string[] Directories)
+		{
+			foreach (var DirectoryName in Directories)
+			{
+				var NormalizedDirectory = ConvertSeparators(PathSeparator.Default, DirectoryName);
+				if (!InternalUtils.SafeCreateDirectory(NormalizedDirectory))
+				{
+					throw new AutomationException(String.Format("Failed to create directory '{0}'", NormalizedDirectory));
+				}
+			}
+		}
+
+        /// <summary>
+        /// Creates a directory(or directories).
+        /// If the creation of the directory fails, this function throws an Exception.
+        /// </summary>
+        /// <param name="Directory">Directory</param>
+        public static void CreateDirectory(bool bQuiet, params string[] Directories)
+        {
+            foreach (var DirectoryName in Directories)
+            {
+                var NormalizedDirectory = ConvertSeparators(PathSeparator.Default, DirectoryName);
+                if (!InternalUtils.SafeCreateDirectory(NormalizedDirectory, bQuiet))
+                {
+                    throw new AutomationException(String.Format("Failed to create directory '{0}'", NormalizedDirectory));
+                }
+            }
+        }
+
+		/// <summary>
+		/// Creates a directory (or directories).
+		/// If the creation of the directory fails, this function prints a warning.
+		/// </summary>
+		/// <param name="Directory">Directory</param>
+		public static bool CreateDirectory_NoExceptions(params string[] Directories)
+		{
+			bool Result = true;
+			foreach (var DirectoryName in Directories)
+			{
+				var NormalizedDirectory = ConvertSeparators(PathSeparator.Default, DirectoryName);
+				if (!InternalUtils.SafeCreateDirectory(NormalizedDirectory))
+				{
+					Log(TraceEventType.Warning, "Failed to create directory '{0}'", NormalizedDirectory);
+					Result = false;
+				}
+			}
+			return Result;
+		}
+
+		/// <summary>
+		/// Renames/moves a file.
+		/// If the rename of the file fails, this function throws an Exception.
+		/// </summary>
+		/// <param name="OldName">Old name</param>
+		/// <param name="NewName">new name</param>
+        public static void RenameFile(string OldName, string NewName, bool bQuiet = false)
+		{
+			var OldNormalized = ConvertSeparators(PathSeparator.Default, OldName);
+			var NewNormalized = ConvertSeparators(PathSeparator.Default, NewName);
+			if (!InternalUtils.SafeRenameFile(OldNormalized, NewNormalized, bQuiet))
+			{
+				throw new AutomationException(String.Format("Failed to rename/move file '{0}' to '{1}'", OldNormalized, NewNormalized));
+			}
+		}
+
+		/// <summary>
+		/// Renames/moves a file.
+		/// If the rename of the file fails, this function prints a warning.
+		/// </summary>
+		/// <param name="OldName">Old name</param>
+		/// <param name="NewName">new name</param>
+		public static bool RenameFile_NoExceptions(string OldName, string NewName)
+		{
+			var OldNormalized = ConvertSeparators(PathSeparator.Default, OldName);
+			var NewNormalized = ConvertSeparators(PathSeparator.Default, NewName);
+			var Result = InternalUtils.SafeRenameFile(OldNormalized, NewNormalized);
+			if (!Result)
+			{
+				Log(TraceEventType.Warning, "Failed to rename/move file '{0}' to '{1}'", OldName, NewName);
+			}
+			return Result;
+		}
+
+		/// <summary>
+		/// Checks if a file(s) exists.
+		/// </summary>
+		/// <param name="Filename">Filename.</param>
+		/// <returns>True if the file exists, false otherwise.</returns>
+		public static bool FileExists(params string[] Filenames)
+		{
+			bool bExists = Filenames.Length > 0;
+			foreach (var Filename in Filenames)
+			{
+				var NormalizedFilename = ConvertSeparators(PathSeparator.Default, Filename);
+				bExists = InternalUtils.SafeFileExists(NormalizedFilename) && bExists;
+			}
+			return bExists;
+		}
+
+		/// <summary>
+		/// Checks if a file(s) exists.
+		/// </summary>
+		/// <param name="Filename">Filename.</param>
+		/// <returns>True if the file exists, false otherwise.</returns>
+		public static bool FileExists_NoExceptions(params string[] Filenames)
+		{
+			// Standard version doesn't throw, but keep this function for consistency.
+			return FileExists(Filenames);
+		}
+        /// <summary>
+        /// Checks if a file(s) exists.
+        /// </summary>
+        /// <param name="Filename">Filename.</param>
+        /// <returns>True if the file exists, false otherwise.</returns>
+        public static bool FileExists(bool bQuiet, params string[] Filenames)
+        {
+            bool bExists = Filenames.Length > 0;
+            foreach (var Filename in Filenames)
+            {
+                var NormalizedFilename = ConvertSeparators(PathSeparator.Default, Filename);
+                bExists = InternalUtils.SafeFileExists(NormalizedFilename, bQuiet) && bExists;
+            }
+            return bExists;
+        }
+
+        /// <summary>
+        /// Checks if a file(s) exists.
+        /// </summary>
+        /// <param name="Filename">Filename.</param>
+        /// <returns>True if the file exists, false otherwise.</returns>
+        public static bool FileExists_NoExceptions(bool bQuiet, params string[] Filenames)
+        {
+            // Standard version doesn't throw, but keep this function for consistency.
+            return FileExists(bQuiet, Filenames);
+        }
+
+		static Stack<string> WorkingDirectoryStack = new Stack<string>();
+
+		/// <summary>
+		/// Pushes the current working directory onto a stack and sets CWD to a new value.
+		/// </summary>
+		/// <param name="WorkingDirectory">New working direcotry.</param>
+		public static void PushDir(string WorkingDirectory)
+		{
+			string OrigCurrentDirectory = Environment.CurrentDirectory;
+			WorkingDirectory = ConvertSeparators(PathSeparator.Default, WorkingDirectory);
+			try
+			{
+				Environment.CurrentDirectory = WorkingDirectory;
+			}
+			catch (Exception Ex)
+			{
+				throw new AutomationException(String.Format("Unable to change current directory to {0}", WorkingDirectory), Ex);
+			}
+
+			WorkingDirectoryStack.Push(OrigCurrentDirectory);
+		}
+
+		/// <summary>
+		/// Pushes the current working directory onto a stack and sets CWD to a new value.
+		/// </summary>
+		/// <param name="WorkingDirectory">New working direcotry.</param>
+		public static bool PushDir_NoExceptions(string WorkingDirectory)
+		{
+			bool Result = true;
+			string OrigCurrentDirectory = Environment.CurrentDirectory;
+			WorkingDirectory = ConvertSeparators(PathSeparator.Default, WorkingDirectory);
+			try
+			{
+				Environment.CurrentDirectory = WorkingDirectory;
+				WorkingDirectoryStack.Push(OrigCurrentDirectory);
+			}
+			catch
+			{
+				Log(TraceEventType.Warning, "Unable to change current directory to {0}", WorkingDirectory);
+				Result = false;
+			}
+			return Result;
+		}
+
+		/// <summary>
+		/// Pops the last working directory from a stack and sets it as the current working directory.
+		/// </summary>
+		public static void PopDir()
+		{
+			if (WorkingDirectoryStack.Count > 0)
+			{
+				Environment.CurrentDirectory = WorkingDirectoryStack.Pop();
+			}
+			else
+			{
+				throw new AutomationException("Unable to PopDir. WorkingDirectoryStack is empty.");
+			}
+		}
+
+		/// <summary>
+		/// Pops the last working directory from a stack and sets it as the current working directory.
+		/// </summary>
+		public static bool PopDir_NoExceptions()
+		{
+			bool Result = true;
+			if (WorkingDirectoryStack.Count > 0)
+			{
+				Environment.CurrentDirectory = WorkingDirectoryStack.Pop();
+			}
+			else
+			{
+				Log(TraceEventType.Warning, "Unable to PopDir. WorkingDirectoryStack is empty.");
+				Result = false;
+			}
+			return Result;
+		}
+
+		/// <summary>
+		/// Clears the directory stack
+		/// </summary>
+		public static void ClearDirStack()
+		{
+			while (WorkingDirectoryStack.Count > 0)
+			{
+				PopDir();
+			}
+		}
+
+		/// <summary>
+		/// Changes the current working directory.
+		/// </summary>
+		/// <param name="WorkingDirectory">New working directory.</param>
+		public static void ChDir(string WorkingDirectory)
+		{
+			WorkingDirectory = ConvertSeparators(PathSeparator.Default, WorkingDirectory);
+			try
+			{
+				Environment.CurrentDirectory = WorkingDirectory;
+			}
+			catch (Exception Ex)
+			{
+				throw new ArgumentException(String.Format("Unable to change current directory to {0}", WorkingDirectory), Ex);
+			}
+		}
+
+		/// <summary>
+		/// Changes the current working directory.
+		/// </summary>
+		/// <param name="WorkingDirectory">New working directory.</param>
+		public static bool ChDir_NoExceptions(string WorkingDirectory)
+		{
+			bool Result = true;
+			WorkingDirectory = ConvertSeparators(PathSeparator.Default, WorkingDirectory);
+			try
+			{
+				Environment.CurrentDirectory = WorkingDirectory;
+			}
+			catch
+			{
+				Log(TraceEventType.Warning, "Unable to change current directory to {0}", WorkingDirectory);
+				Result = false;
+			}
+			return Result;
+		}
+
+		/// <summary>
+		/// Sets file attributes. Will not change attributes that have not been specified.
+		/// </summary>
+		/// <param name="Filename">Filename</param>
+		/// <param name="ReadOnly">Read-only attribute</param>
+		/// <param name="Hidden">Hidden attribute.</param>
+		/// <param name="Archive">Archive attribute.</param>
+		public static void SetFileAttributes(string Filename, bool? ReadOnly = null, bool? Hidden = null, bool? Archive = null)
+		{
+			Filename = ConvertSeparators(PathSeparator.Default, Filename);
+			if (!File.Exists(Filename))
+			{
+				throw new AutomationException("Unable to set attributes for a non-exisiting file.", new FileNotFoundException("File not found.", Filename));
+			}
+
+			FileAttributes Attributes = File.GetAttributes(Filename);
+			Attributes = InternalSetAttributes(ReadOnly, Hidden, Archive, Attributes);
+			File.SetAttributes(Filename, Attributes);
+		}
+
+		/// <summary>
+		/// Sets file attributes. Will not change attributes that have not been specified.
+		/// </summary>
+		/// <param name="Filename">Filename</param>
+		/// <param name="ReadOnly">Read-only attribute</param>
+		/// <param name="Hidden">Hidden attribute.</param>
+		/// <param name="Archive">Archive attribute.</param>
+		public static bool SetFileAttributes_NoExceptions(string Filename, bool? ReadOnly = null, bool? Hidden = null, bool? Archive = null)
+		{			
+			Filename = ConvertSeparators(PathSeparator.Default, Filename);
+			if (!File.Exists(Filename))
+			{
+				Log(TraceEventType.Warning, "Unable to set attributes for a non-exisiting file ({0})", Filename);
+				return false;
+			}
+
+			bool Result = true;
+			try
+			{
+				FileAttributes Attributes = File.GetAttributes(Filename);
+				Attributes = InternalSetAttributes(ReadOnly, Hidden, Archive, Attributes);
+				File.SetAttributes(Filename, Attributes);
+			}
+			catch (Exception Ex)
+			{
+				Log(TraceEventType.Warning, "Error trying to set file attributes for: {0}", Filename);
+				Log(TraceEventType.Warning, Ex);
+				Result = false;
+			}
+			return Result;
+		}
+
+		private static FileAttributes InternalSetAttributes(bool? ReadOnly, bool? Hidden, bool? Archive, FileAttributes Attributes)
+		{
+			if (ReadOnly != null)
+			{
+				if ((bool)ReadOnly)
+				{
+					Attributes |= FileAttributes.ReadOnly;
+				}
+				else
+				{
+					Attributes &= ~FileAttributes.ReadOnly;
+				}
+			}
+			if (Hidden != null)
+			{
+				if ((bool)Hidden)
+				{
+					Attributes |= FileAttributes.Hidden;
+				}
+				else
+				{
+					Attributes &= ~FileAttributes.Hidden;
+				}
+			}
+			if (Archive != null)
+			{
+				if ((bool)Archive)
+				{
+					Attributes |= FileAttributes.Archive;
+				}
+				else
+				{
+					Attributes &= ~FileAttributes.Archive;
+				}
+			}
+			return Attributes;
+		}
+
+		/// <summary>
+		/// Writes a line of formatted string to a file. Creates the file if it does not exists.
+		/// If the file does exists, appends a new line.
+		/// </summary>
+		/// <param name="Filename">Filename</param>
+		/// <param name="Format">Format string</param>
+		/// <param name="Args">Arguments</param>
+		public static void WriteToFile(string Filename, string Format, params object[] Args)
+		{
+			WriteToFile(Filename, String.Format(Format, Args));
+		}
+
+		/// <summary>
+		/// Writes a line of formatted string to a file. Creates the file if it does not exists.
+		/// If the file does exists, appends a new line.
+		/// </summary>
+		/// <param name="Filename">Filename</param>
+		/// <param name="Text">Text to write</param>
+		public static void WriteToFile(string Filename, string Text)
+		{
+			try
+			{
+				Filename = ConvertSeparators(PathSeparator.Default, Filename);
+				FileStream Stream;
+				if (File.Exists(Filename))
+				{
+					Stream = new FileStream(Filename, FileMode.Append, FileAccess.Write, FileShare.Read);
+				}
+				else
+				{
+					Stream = new FileStream(Filename, FileMode.OpenOrCreate, FileAccess.Write, FileShare.Read);
+				}
+				if (Stream != null)
+				{
+					using (StreamWriter Writer = new StreamWriter(Stream))
+					{
+						Writer.WriteLine(Text);
+						Writer.Flush();
+						Writer.Close();
+					}
+					Stream.Dispose();
+				}
+			}
+			catch (Exception Ex)
+			{
+				throw new AutomationException(String.Format("Failed to Write to file {0}", Filename), Ex);
+			}
+		}
+
+		/// <summary>
+		/// Reads all text lines from a file.
+		/// </summary>
+		/// <param name="Filename">Filename</param>
+		/// <returns>Array of lines of text read from the file. null if the file did not exist or could not be read.</returns>
+		public static string[] ReadAllLines(string Filename)
+		{
+			Filename = ConvertSeparators(PathSeparator.Default, Filename);
+			return InternalUtils.SafeReadAllLines(Filename);
+		}
+
+		/// <summary>
+		/// Reads all text from a file.
+		/// </summary>
+		/// <param name="Filename">Filename</param>
+		/// <returns>All text read from the file. null if the file did not exist or could not be read.</returns>
+		public static string ReadAllText(string Filename)
+		{
+			Filename = ConvertSeparators(PathSeparator.Default, Filename);
+			return InternalUtils.SafeReadAllText(Filename);
+		}
+
+		/// <summary>
+		/// Writes lines to a file.
+		/// </summary>
+		/// <param name="Filename">Filename</param>
+		/// <param name="Lines">Text</param>
+		public static void WriteAllLines(string Filename, string[] Lines)
+		{
+			Filename = ConvertSeparators(PathSeparator.Default, Filename);
+			if (!InternalUtils.SafeWriteAllLines(Filename, Lines))
+			{
+				throw new AutomationException("Unable to write to file: {0}", Filename);
+			}
+		}
+
+		/// <summary>
+		/// Writes lines to a file.
+		/// </summary>
+		/// <param name="Filename">Filename</param>
+		/// <param name="Lines">Text</param>
+		public static bool WriteAllLines_NoExceptions(string Filename, string[] Lines)
+		{
+			Filename = ConvertSeparators(PathSeparator.Default, Filename);
+			return InternalUtils.SafeWriteAllLines(Filename, Lines);
+		}
+
+		/// <summary>
+		/// Writes text to a file.
+		/// </summary>
+		/// <param name="Filename">Filename</param>
+		/// <param name="Lines">Text</param>
+		public static void WriteAllText(string Filename, string Text)
+		{
+			Filename = ConvertSeparators(PathSeparator.Default, Filename);
+			if (!InternalUtils.SafeWriteAllText(Filename, Text))
+			{
+				throw new AutomationException("Unable to write to file: {0}", Filename);
+			}
+		}
+
+		/// <summary>
+		/// Writes text to a file.
+		/// </summary>
+		/// <param name="Filename">Filename</param>
+		/// <param name="Lines">Text</param>
+		public static bool WriteAllText_NoExceptions(string Filename, string Text)
+		{
+			Filename = ConvertSeparators(PathSeparator.Default, Filename);
+			return InternalUtils.SafeWriteAllText(Filename, Text);
+		}
+
+
+		/// <summary>
+		/// Writes byte array to a file.
+		/// </summary>
+		/// <param name="Filename">Filename</param>
+		/// <param name="Bytes">Byte array</param>
+		public static void WriteAllBytes(string Filename, byte[] Bytes)
+		{
+			Filename = ConvertSeparators(PathSeparator.Default, Filename);
+			if (!InternalUtils.SafeWriteAllBytes(Filename, Bytes))
+			{
+				throw new AutomationException("Unable to write to file: {0}", Filename);
+			}
+		}
+
+		/// <summary>
+		/// Writes byte array to a file.
+		/// </summary>
+		/// <param name="Filename">Filename</param>
+		/// <param name="Bytes">Byte array</param>
+		public static bool WriteAllBytes_NoExceptions(string Filename, byte[] Bytes)
+		{
+			Filename = ConvertSeparators(PathSeparator.Default, Filename);
+			return InternalUtils.SafeWriteAllBytes(Filename, Bytes);
+		}
+
+		/// <summary>
+		/// Gets a character representing the specified separator type.
+		/// </summary>
+		/// <param name="SeparatorType">Separator type.</param>
+		/// <returns>Separator character</returns>
+		public static char GetPathSeparatorChar(PathSeparator SeparatorType)
+		{
+			char Separator;
+			switch (SeparatorType)
+			{
+				case PathSeparator.Slash:
+				case PathSeparator.Depot:
+					Separator = '/';
+					break;
+				case PathSeparator.Backslash:
+					Separator = '\\';
+					break;
+				default:
+					Separator = Path.DirectorySeparatorChar;
+					break;
+			}
+			return Separator;
+		}
+
+		/// <summary>
+		/// Checks if the character is one of the two sperator types ('\' or '/')
+		/// </summary>
+		/// <param name="Character">Character to check.</param>
+		/// <returns>True if the character is a separator, false otherwise.</returns>
+		public static bool IsPathSeparator(char Character)
+		{
+			return (Character == '/' || Character == '\\');
+		}
+
+		/// <summary>
+		/// Combines paths and replaces all path separators with the system default separator.
+		/// </summary>
+		/// <param name="Paths"></param>
+		/// <returns>Combined Path</returns>
+		public static string CombinePaths(params string[] Paths)
+		{
+			return CombinePaths(PathSeparator.Default, Paths);
+		}
+
+		/// <summary>
+		/// Combines paths and replaces all path separators wth the system specified separator.
+		/// </summary>
+		/// <param name="SeparatorType">Type of separartor to use when combining paths.</param>
+		/// <param name="Paths"></param>
+		/// <returns>Combined Path</returns>
+		public static string CombinePaths(PathSeparator SeparatorType, params string[] Paths)
+		{
+			// Pick a separator to use.
+			var SeparatorToUse = GetPathSeparatorChar(SeparatorType);
+			var SeparatorToReplace = SeparatorToUse == '/' ? '\\' : '/';
+
+			// Allocate string builder
+			int CombinePathMaxLength = 0;
+			foreach (var PathPart in Paths)
+			{
+				CombinePathMaxLength += (PathPart != null) ? PathPart.Length : 0;
+			}
+			CombinePathMaxLength += Paths.Length;
+			var CombinedPath = new StringBuilder(CombinePathMaxLength);
+
+			// Combine all paths
+			CombinedPath.Append(Paths[0]);
+			for (int PathIndex = 1; PathIndex < Paths.Length; ++PathIndex)
+			{
+				var NextPath = Paths[PathIndex];
+				if (String.IsNullOrEmpty(NextPath) == false)
+				{
+					int NextPathStartIndex = 0;
+					if (CombinedPath.Length != 0)
+					{
+						var LastChar = CombinedPath[CombinedPath.Length - 1];
+						var NextChar = NextPath[0];
+						var IsLastCharPathSeparator = IsPathSeparator(LastChar);
+						var IsNextCharPathSeparator = IsPathSeparator(NextChar);
+						// Check if a separator between paths is required
+						if (!IsLastCharPathSeparator && !IsNextCharPathSeparator)
+						{
+							CombinedPath.Append(SeparatorToUse);
+						}
+						// Check if one of the saprators needs to be skipped.
+						else if (IsLastCharPathSeparator && IsNextCharPathSeparator)
+						{
+							NextPathStartIndex = 1;
+						}
+					}
+					CombinedPath.Append(NextPath, NextPathStartIndex, NextPath.Length - NextPathStartIndex);
+				}
+			}
+			// Make sure there's only one separator type used.
+			CombinedPath.Replace(SeparatorToReplace, SeparatorToUse);
+			return CombinedPath.ToString();
+		}
+
+		/// <summary>
+		/// Converts all separators in path to the specified separator type.
+		/// </summary>
+		/// <param name="ToSperatorType">Desired separator type.</param>
+		/// <param name="PathToConvert">Path</param>
+		/// <returns>Path where all separators have been converted to the specified type.</returns>
+		public static string ConvertSeparators(PathSeparator ToSperatorType, string PathToConvert)
+		{
+			return CombinePaths(ToSperatorType, PathToConvert);
+		}
+
+		/// <summary>
+		/// Copies a file.
+		/// </summary>
+		/// <param name="Source"></param>
+		/// <param name="Dest"></param>
+		/// <returns>True if the operation was successful, false otherwise.</returns>
+        public static void CopyFile(string Source, string Dest, bool bQuiet = false)
+		{
+			Source = ConvertSeparators(PathSeparator.Default, Source);
+			Dest = ConvertSeparators(PathSeparator.Default, Dest);
+			if (InternalUtils.SafeFileExists(Dest, true))
+			{
+				InternalUtils.SafeDeleteFile(Dest, bQuiet);
+			}
+			else if (!InternalUtils.SafeDirectoryExists(Path.GetDirectoryName(Dest), true))
+			{
+				if (!InternalUtils.SafeCreateDirectory(Path.GetDirectoryName(Dest), bQuiet))
+				{
+					throw new AutomationException("Failed to create directory {0} for copy", Path.GetDirectoryName(Dest));
+				}
+			}
+			if (InternalUtils.SafeFileExists(Dest, true))
+			{
+				throw new AutomationException("Failed to delete {0} for copy", Dest);
+			}
+			if (!InternalUtils.SafeCopyFile(Source, Dest, bQuiet))
+			{
+				throw new AutomationException("Failed to copy {0} to {1}", Source, Dest);
+			}
+		}
+
+		/// <summary>
+		/// Copies a file. Does not throw exceptions.
+		/// </summary>
+		/// <param name="Source"></param>
+		/// <param name="Dest"></param>
+		/// <returns>True if the operation was successful, false otherwise.</returns>
+        public static bool CopyFile_NoExceptions(string Source, string Dest, bool bQuiet = false)
+		{
+			Source = ConvertSeparators(PathSeparator.Default, Source);
+			Dest = ConvertSeparators(PathSeparator.Default, Dest);
+            if (InternalUtils.SafeFileExists(Dest, true))
+			{
+                InternalUtils.SafeDeleteFile(Dest, bQuiet);
+			}
+			else if (!InternalUtils.SafeDirectoryExists(Path.GetDirectoryName(Dest), true))
+			{
+				if (!InternalUtils.SafeCreateDirectory(Path.GetDirectoryName(Dest)))
+				{
+					return false;
+				}
+			}
+			if (InternalUtils.SafeFileExists(Dest, true))
+			{
+				return false;
+			}
+			return InternalUtils.SafeCopyFile(Source, Dest, bQuiet);
+		}
+
+		/// <summary>
+		/// Copies a file if the dest doesn't exist or (optionally) if the dest timestamp is different; after a copy, copies the timestamp
+		/// </summary>
+		/// <param name="Source">The full path to the source file</param>
+		/// <param name="Dest">The full path to the destination file</param>
+		/// <param name="bAllowDifferingTimestamps">If true, will always skip a file if the destination exists, even if timestamp differs; defaults to false</param>
+		/// <returns>True if the operation was successful, false otherwise.</returns>
+		public static void CopyFileIncremental(string Source, string Dest, bool bAllowDifferingTimestamps = false, bool bFilterSpecialLinesFromIniFiles = false)
+		{
+			Source = ConvertSeparators(PathSeparator.Default, Source);
+			Dest = ConvertSeparators(PathSeparator.Default, Dest);
+			if (InternalUtils.SafeFileExists(Dest, true))
+			{
+				if (bAllowDifferingTimestamps == true)
+				{
+					Log("CopyFileIncremental Skipping {0}, already exists", Dest);
+					return;
+				}
+				TimeSpan Diff = File.GetLastWriteTimeUtc(Dest) - File.GetLastWriteTimeUtc(Source);
+				if (Diff.TotalSeconds > -1 && Diff.TotalSeconds < 1)
+				{
+					Log("CopyFileIncremental Skipping {0}, up to date.", Dest);
+					return;
+				}
+				InternalUtils.SafeDeleteFile(Dest);
+			}
+			else if (!InternalUtils.SafeDirectoryExists(Path.GetDirectoryName(Dest), true))
+			{
+				if (!InternalUtils.SafeCreateDirectory(Path.GetDirectoryName(Dest)))
+				{
+					throw new AutomationException("Failed to create directory {0} for copy", Path.GetDirectoryName(Dest));
+				}
+			}
+			if (InternalUtils.SafeFileExists(Dest, true))
+			{
+				throw new AutomationException("Failed to delete {0} for copy", Dest);
+			}
+			if (!InternalUtils.SafeCopyFile(Source, Dest, bFilterSpecialLinesFromIniFiles:bFilterSpecialLinesFromIniFiles))
+			{
+				throw new AutomationException("Failed to copy {0} to {1}", Source, Dest);
+			}
+			FileAttributes Attributes = File.GetAttributes(Dest);
+			if ((Attributes & FileAttributes.ReadOnly) != 0)
+			{
+				File.SetAttributes(Dest, Attributes & ~FileAttributes.ReadOnly);
+			}
+			File.SetLastWriteTimeUtc(Dest, File.GetLastWriteTimeUtc(Source));
+		}
+
+		/// <summary>
+		/// Copies a directory and all of it's contents recursively. Does not throw exceptions.
+		/// </summary>
+		/// <param name="Source"></param>
+		/// <param name="Dest"></param>
+		/// <returns>True if the operation was successful, false otherwise.</returns>
+		public static bool CopyDirectory_NoExceptions(string Source, string Dest, bool bQuiet = false)
+		{
+			Source = ConvertSeparators(PathSeparator.Default, Source);
+			Dest = ConvertSeparators(PathSeparator.Default, Dest);
+			Dest.TrimEnd(PathSeparator.Default.ToString().ToCharArray());
+
+			if (InternalUtils.SafeDirectoryExists(Dest))
+			{
+				InternalUtils.SafeDeleteDirectory(Dest, bQuiet);
+				if (InternalUtils.SafeDirectoryExists(Dest, true))
+				{
+					return false;
+				}
+			}
+
+			if (!InternalUtils.SafeCreateDirectory(Dest, bQuiet))
+			{
+				return false;
+			}
+			foreach (var SourceSubDirectory in Directory.GetDirectories(Source))
+			{
+				string DestPath = Dest + GetPathSeparatorChar(PathSeparator.Default) + GetLastDirectoryName(SourceSubDirectory);
+				if (!CopyDirectory_NoExceptions(SourceSubDirectory, DestPath, bQuiet))
+				{
+					return false;
+				}
+			}
+			foreach (var SourceFile in Directory.GetFiles(Source))
+			{
+				int FilenameStart = SourceFile.LastIndexOf(GetPathSeparatorChar(PathSeparator.Default));
+				string DestPath = Dest + SourceFile.Substring(FilenameStart);
+				if (!CopyFile_NoExceptions(SourceFile, DestPath, bQuiet))
+				{
+					return false;
+				}
+			}
+
+			return true;
+		}
+
+		/// <summary>
+		/// Returns directory name without filename. 
+		/// The difference between this and Path.GetDirectoryName is that this
+		/// function will not throw away the last name if it doesn't have an extension, for example:
+		/// D:\Project\Data\Asset -> D:\Project\Data\Asset
+		/// D:\Project\Data\Asset.ussset -> D:\Project\Data
+		/// </summary>
+		/// <param name="FilePath"></param>
+		/// <returns></returns>
+		public static string GetDirectoryName(string FilePath)
+		{
+			var LastSeparatorIndex = Math.Max(FilePath.LastIndexOf('/'), FilePath.LastIndexOf('\\'));
+			var ExtensionIndex = FilePath.LastIndexOf('.');
+			if (ExtensionIndex > LastSeparatorIndex || LastSeparatorIndex == (FilePath.Length - 1))
+			{
+				return FilePath.Substring(0, LastSeparatorIndex);
+			}
+			else
+			{
+				return FilePath;
+			}
+		}
+
+		/// <summary>
+		/// Returns the last directory name in the path string.
+		/// For example: D:\Temp\Project\File.txt -> Project, Data\Samples -> Samples
+		/// </summary>
+		/// <param name="FilePath"></param>
+		/// <returns></returns>
+		public static string GetLastDirectoryName(string FilePath)
+		{
+			var LastDir = GetDirectoryName(FilePath);
+			var LastSeparatorIndex = Math.Max(LastDir.LastIndexOf('/'), LastDir.LastIndexOf('\\'));
+			if (LastSeparatorIndex >= 0)
+			{
+				LastDir = LastDir.Substring(LastSeparatorIndex + 1);
+			}
+			return LastDir;
+		}
+
+		/// <summary>
+		/// Removes multi-dot extensions from a filename (i.e. *.automation.csproj)
+		/// </summary>
+		/// <param name="Filename">Filename to remove the extensions from</param>
+		/// <returns>Clean filename.</returns>
+		public static string GetFilenameWithoutAnyExtensions(string Filename)
+		{
+			do
+			{
+				Filename = Path.GetFileNameWithoutExtension(Filename);
+			}
+			while (Filename.IndexOf('.') >= 0);
+			return Filename;
+		}
+
+		/// <summary>
+		/// Reads a file manifest and returns it
+		/// </summary>
+		/// <param name="ManifestName">ManifestName</param>
+		/// <returns></returns>
+		public static UnrealBuildTool.BuildManifest ReadManifest(string ManifestName)
+		{
+			return XmlHandler.ReadXml<UnrealBuildTool.BuildManifest>(ManifestName);
+		}
+
+		private static void CloneDirectoryRecursiveWorker(string SourcePathBase, string TargetPathBase, List<string> ClonedFiles)
+		{
+            if (!InternalUtils.SafeCreateDirectory(TargetPathBase))
+            {
+                throw new AutomationException("Failed to create directory {0} for copy", TargetPathBase);
+            }
+
+			DirectoryInfo SourceDirectory = new DirectoryInfo(SourcePathBase);
+			DirectoryInfo[] SourceSubdirectories = SourceDirectory.GetDirectories();
+
+			// Copy the files
+			FileInfo[] SourceFiles = SourceDirectory.GetFiles();
+			foreach (FileInfo SourceFI in SourceFiles)
+			{
+				string TargetFilename = CommandUtils.CombinePaths(TargetPathBase, SourceFI.Name);
+				SourceFI.CopyTo(TargetFilename);
+
+				if (ClonedFiles != null)
+				{
+					ClonedFiles.Add(TargetFilename);
+				}
+			}
+
+			// Recurse into subfolders
+			foreach (DirectoryInfo SourceSubdir in SourceSubdirectories)
+			{
+				string NewSourcePath = CommandUtils.CombinePaths(SourcePathBase, SourceSubdir.Name);
+				string NewTargetPath = CommandUtils.CombinePaths(TargetPathBase, SourceSubdir.Name);
+				CloneDirectoryRecursiveWorker(NewSourcePath, NewTargetPath, ClonedFiles);
+			}
+		}
+
+		/// <summary>
+		/// Clones a directory.
+		/// Warning: Will delete all of the existing files in TargetPath
+		/// This is recursive, copying subfolders too.
+		/// </summary>
+		/// <param name="SourcePath">Source directory.</param>
+		/// <param name="TargetPath">Target directory.</param>
+		/// <param name="ClonedFiles">List of cloned files.</param>
+		public static void CloneDirectory(string SourcePath, string TargetPath, List<string> ClonedFiles = null)
+		{
+			DeleteDirectory_NoExceptions(TargetPath);
+
+			CloneDirectoryRecursiveWorker(SourcePath, TargetPath, ClonedFiles);
+		}
+
+		#endregion
+
+		#region Threaded Copy
+
+		/// <summary>
+		/// Single batch of files to copy on a thread
+		/// </summary>
+		class CopyRequest
+		{
+			public string[] Source;
+			public string[] Dest;
+			public Exception Result;
+			public Thread Worker;
+		}
+
+		/// <summary>
+		/// Main thread procedure for copying files
+		/// </summary>
+		private static void CopyThreadProc(object Request)
+		{
+			const bool bQuiet = true;
+			var FileToCopy = (CopyRequest)Request;
+			try
+			{
+				for (int Index = 0; Index < FileToCopy.Source.Length; ++Index)
+				{
+					CopyFile(FileToCopy.Source[Index], FileToCopy.Dest[Index], bQuiet);
+				}
+			}
+			catch (Exception Ex)
+			{
+				FileToCopy.Result = Ex;
+			}
+		}
+
+		/// <summary>
+		/// Copies files using miltiple threads
+		/// </summary>
+		/// <param name="SourceDirectory"></param>
+		/// <param name="DestDirectory"></param>
+		/// <param name="MaxThreads"></param>
+		public static void ThreadedCopyFiles(string SourceDirectory, string DestDirectory, int MaxThreads = 64)
+		{
+			CreateDirectory(DestDirectory);
+
+			var SourceFiles = Directory.GetFiles(SourceDirectory, "*", SearchOption.AllDirectories);
+			var SourceBase = GetDirectoryName(SourceDirectory) + GetPathSeparatorChar(PathSeparator.Default);
+			var DestBase = GetDirectoryName(DestDirectory) + GetPathSeparatorChar(PathSeparator.Default);
+			var DestFiles = new string[SourceFiles.Length];
+			for (int Index = 0; Index < SourceFiles.Length; ++Index)
+			{
+				DestFiles[Index] = SourceFiles[Index].Replace(SourceBase, DestBase);
+			}
+			ThreadedCopyFiles(SourceFiles, DestFiles, MaxThreads);
+		}
+
+		/// <summary>
+		/// Copies files using multiple threads
+		/// </summary>
+		/// <param name="SourceDirectory"></param>
+		/// <param name="DestDirectory"></param>
+		/// <param name="MaxThreads"></param>
+		public static void ThreadedCopyFiles(string[] Source, string[] Dest, int MaxThreads = 64)
+		{
+			Log("Copying {0} file(s) using max {1} thread(s)", Source.Length, MaxThreads);
+
+			if (Source.Length != Dest.Length)
+			{
+				throw new AutomationException("Source count ({0}) does not match Dest count ({1})", Source.Length, Dest.Length);
+			}
+			const int MinFilesPerThread = 10;
+			if (MaxThreads > 1 && Source.Length > MinFilesPerThread)
+			{
+				// Split evenly across the threads
+				int FilesPerThread = Math.Max(Source.Length / MaxThreads, MinFilesPerThread);
+				int ThreadCount = Math.Min(Source.Length / FilesPerThread + 1, MaxThreads);
+				FilesPerThread = Source.Length / ThreadCount + 1;
+
+				// Divide and copy
+				var WorkerThreads = new List<CopyRequest>(ThreadCount);
+				var FilesToCopy = Source.Length;
+				var FirstFileIndex = 0;
+				while (FilesToCopy > 0)
+				{
+					var Request = new CopyRequest();
+					WorkerThreads.Add(Request);
+
+					Request.Source = new string[Math.Min(FilesToCopy, FilesPerThread)];
+					Request.Dest = new string[Request.Source.Length];
+					Array.Copy(Source, FirstFileIndex, Request.Source, 0, Request.Source.Length);
+					Array.Copy(Dest, FirstFileIndex, Request.Dest, 0, Request.Dest.Length);
+
+					Request.Worker = new Thread(new ParameterizedThreadStart(CopyThreadProc));
+					Request.Worker.Start(Request);
+
+					FirstFileIndex += Request.Source.Length;
+					FilesToCopy -= Request.Source.Length;
+				}
+
+				// Wait for completion
+				foreach (var Request in WorkerThreads)
+				{
+					if (Request.Worker.IsAlive)
+					{
+						Request.Worker.Join();
+					}
+					if (Request.Result != null)
+					{
+						throw new AutomationException(String.Format("Failed to thread-copy files: {0}", Request.Result.Message), Request.Result);
+					}
+				}
+			}
+			else
+			{
+				const bool bQuiet = true;
+				for (int Index = 0; Index < Source.Length; ++Index)
+				{
+					CopyFile(Source[Index], Dest[Index], bQuiet);
+				}
+			}
+		}
+
+		/// <summary>
+		/// Copies a set of files from one folder to another
+		/// </summary>
+		/// <param name="SourceDir">Source directory</param>
+		/// <param name="TargetDir">Target directory</param>
+		/// <param name="RelativePaths">Paths relative to the source directory to copy</param>
+		/// <param name="bIgnoreSymlinks">Whether to ignore symlinks during the copy</param>
+		/// <param name="MaxThreads">Maximum number of threads to create</param>
+		/// <returns>Array of filenames copied to the target directory</returns>
+		public static string[] ThreadedCopyFiles(string SourceDir, string TargetDir, string[] RelativePaths, int MaxThreads = 64)
+		{
+			string[] SourceFileNames = new string[RelativePaths.Length];
+			string[] TargetFileNames = new string[RelativePaths.Length];
+			for(int Idx = 0; Idx < RelativePaths.Length; Idx++)
+			{
+				SourceFileNames[Idx] = CommandUtils.CombinePaths(SourceDir, RelativePaths[Idx]);
+				TargetFileNames[Idx] = CommandUtils.CombinePaths(TargetDir, RelativePaths[Idx]);
+			}
+			CommandUtils.ThreadedCopyFiles(SourceFileNames, TargetFileNames, MaxThreads);
+			return TargetFileNames;
+		}
+
+		/// <summary>
+		/// Copies a set of files from one folder to another
+		/// </summary>
+		/// <param name="SourceDir">Source directory</param>
+		/// <param name="TargetDir">Target directory</param>
+		/// <param name="Filter">Filter which selects files from the source directory to copy</param>
+		/// <param name="bIgnoreSymlinks">Whether to ignore symlinks during the copy</param>
+		/// <param name="MaxThreads">Maximum number of threads to create</param>
+		/// <returns>Array of filenames copied to the target directory</returns>
+		public static string[] ThreadedCopyFiles(string SourceDir, string TargetDir, FileFilter Filter, bool bIgnoreSymlinks, int MaxThreads = 64)
+		{
+			// Filter all the relative paths
+			CommandUtils.Log("Applying filter to {0}...", SourceDir);
+			string[] RelativePaths = Filter.ApplyToDirectory(SourceDir, bIgnoreSymlinks).ToArray();
+			return ThreadedCopyFiles(SourceDir, TargetDir, RelativePaths);
+		}
+
+		#endregion
+
+		#region Environment variables
+
+		/// <summary>
+		/// Gets environment variable value.
+		/// </summary>
+		/// <param name="Name">Name of the environment variable</param>
+		/// <returns>Environment variable value as string.</returns>
+		public static string GetEnvVar(string Name)
+		{
+			return InternalUtils.GetEnvironmentVariable(Name, "");
+		}
+
+		/// <summary>
+		/// Gets environment variable value.
+		/// </summary>
+		/// <param name="Name">Name of the environment variable</param>
+		/// <param name="DefaultValue">Default value of the environment variable if the variable is not set.</param>
+		/// <returns>Environment variable value as string.</returns>
+		public static string GetEnvVar(string Name, string DefaultValue)
+		{
+			return InternalUtils.GetEnvironmentVariable(Name, DefaultValue);
+		}
+
+		/// <summary>
+		/// Sets environment variable.
+		/// </summary>
+		/// <param name="Name">Variable name.</param>
+		/// <param name="Value">Variable value.</param>
+		/// <returns>True if the value has been set, false otherwise.</returns>
+		public static void SetEnvVar(string Name, object Value)
+		{
+			try
+			{
+				Log("SetEnvVar {0}={1}", Name, Value);
+				Environment.SetEnvironmentVariable(Name, Value.ToString());
+			}
+			catch (Exception Ex)
+			{
+				throw new AutomationException(String.Format("Failed to set environment variable {0} to {1}", Name, Value), Ex);
+			}
+		}
+
+		/// <summary>
+		/// Sets the environment variable if it hasn't been already.
+		/// </summary>
+		/// <param name="VarName">Environment variable name</param>
+		/// <param name="Value">New value</param>
+		public static void ConditionallySetEnvVar(string VarName, string Value)
+		{
+			if (String.IsNullOrEmpty(CommandUtils.GetEnvVar(VarName)))
+			{
+				Environment.SetEnvironmentVariable(VarName, Value);
+			}
+		}
+
+		/// <summary>
+		/// Gets environment variables set by a batch file.
+		/// </summary>
+		/// <param name="BatchFileName">Filename that sets the environment variables</param>
+		/// <param name="AlsoSet">True if found variables should be automatically set withing this process.</param>
+		/// <returns>Dictionary of environment variables set by the batch file.</returns>
+		public static CaselessDictionary<string> GetEnvironmentVariablesFromBatchFile(string BatchFileName, bool AlsoSet = false)
+		{
+			CaselessDictionary<string> Result;
+			try
+			{
+				Result = HarvestEnvVars.HarvestEnvVarsFromBatchFile(BatchFileName, "", HarvestEnvVars.EPathOverride.User);
+			}
+			catch (Exception Ex)
+			{
+				throw new AutomationException("Failed to harvest environment variables", Ex);
+			}
+
+			if (AlsoSet)
+			{
+				// Set the environment variables
+				foreach (var Envvar in Result)
+				{
+					Environment.SetEnvironmentVariable(Envvar.Key, Envvar.Value);
+				}
+			}
+
+			return Result;
+		}
+
+		#endregion
+
+		#region CommandLine
+
+		/// <summary>
+		/// Converts a list of arguments to a string where each argument is separated with a space character.
+		/// </summary>
+		/// <param name="Args">Arguments</param>
+		/// <returns>Single string containing all arguments separated with a space.</returns>
+		public static string ArgsToCommandLine(params object[] Args)
+		{
+			string Arguments = "";
+			if (Args != null)
+			{
+				for (int Index = 0; Index < Args.Length; ++Index)
+				{
+					Arguments += Args[Index].ToString();
+					if (Index < (Args.Length - 1))
+					{
+						Arguments += " ";
+					}
+				}
+			}
+			return Arguments;
+		}
+
+		/// <summary>
+		/// Parses the argument list for a parameter and returns whether it is defined or not.
+		/// </summary>
+		/// <param name="ArgList">Argument list.</param>
+		/// <param name="Param">Param to check for.</param>
+		/// <returns>True if param was found, false otherwise.</returns>
+		public static bool ParseParam(object[] ArgList, string Param)
+		{
+			foreach (object Arg in ArgList)
+			{
+				if (Arg.ToString().Equals(Param, StringComparison.InvariantCultureIgnoreCase))
+				{
+					return true;
+				}
+			}
+			return false;
+		}
+
+		/// <summary>
+		/// Parses the command's Params list for a parameter and returns whether it is defined or not.
+		/// </summary>
+		/// <param name="Param">Param to check for.</param>
+		/// <returns>True if param was found, false otherwise.</returns>
+		public bool ParseParam(string Param)
+		{
+			return ParseParam(Params, Param);
+		}
+
+		/// <summary>
+		/// Parses the argument list for a parameter and reads its value. 
+		/// Ex. ParseParamValue(Args, "map=")
+		/// </summary>
+		/// <param name="ArgList">Argument list.</param>
+		/// <param name="Param">Param to read its value.</param>
+		/// <returns>Returns the value or Default if the parameter was not found.</returns>
+		public static string ParseParamValue(object[] ArgList, string Param, string Default = null)
+		{
+			if (!Param.EndsWith("="))
+			{
+				Param += "=";
+			}
+
+			foreach (object Arg in ArgList)
+			{
+				string ArgStr = Arg.ToString();
+				if (ArgStr.StartsWith(Param, StringComparison.InvariantCultureIgnoreCase))
+				{
+					return ArgStr.Substring(Param.Length);
+				}
+			}
+			return Default;
+		}
+
+		/// <summary>
+		/// Parses the command's Params list for a parameter and reads its value. 
+		/// Ex. ParseParamValue(Args, "map=")
+		/// </summary>
+		/// <param name="Param">Param to read its value.</param>
+		/// <returns>Returns the value or Default if the parameter was not found.</returns>
+		public string ParseParamValue(string Param, string Default = null)
+		{
+			return ParseParamValue(Params, Param, Default);
+		}
+
+		/// <summary>
+		/// Parses the command's Params list for a parameter and reads its value. 
+		/// Ex. ParseParamValue(Args, "map=")
+		/// </summary>
+		/// <param name="Param">Param to read its value.</param>
+		/// <returns>Returns the value or Default if the parameter was not found.</returns>
+		public int ParseParamInt(string Param, int Default = 0)
+		{
+			string num = ParseParamValue(Params, Param, Default.ToString());
+			return int.Parse(num);
+		}
+
+		/// <summary>
+		/// Makes sure path can be used as a command line param (adds quotes if it contains spaces)
+		/// </summary>
+		/// <param name="InPath">Path to convert</param>
+		/// <returns></returns>
+		public static string MakePathSafeToUseWithCommandLine(string InPath)
+		{
+			if (InPath.Contains(' ') && InPath[0] != '\"')
+			{
+				InPath = "\"" + InPath + "\"";
+			}
+			return InPath;
+		}
+
+		#endregion
+
+		#region Other
+
+		public static string EscapePath(string InPath)
+		{
+			return InPath.Replace(":", "").Replace("/", "+").Replace("\\", "+").Replace(" ", "+");
+		}
+
+		/// <summary>
+		/// Checks if collection is either null or empty.
+		/// </summary>
+		/// <param name="Collection">Collection to check.</param>
+		/// <returns>True if the collection is either nur or empty.</returns>
+		public static bool IsNullOrEmpty(ICollection Collection)
+		{
+			return Collection == null || Collection.Count == 0;
+		}
+
+		/// <summary>
+		/// List of available target platforms.
+		/// </summary>
+		public static UnrealBuildTool.UnrealTargetPlatform[] KnownTargetPlatforms
+		{
+			get
+			{
+				if (UBTTargetPlatforms == null || UBTTargetPlatforms.Length == 0)
+				{
+					UBTTargetPlatforms = new UnrealBuildTool.UnrealTargetPlatform[UnrealBuildTool.UEBuildPlatform.BuildPlatformDictionary.Count];
+					int Index = 0;
+					foreach (var Platform in UnrealBuildTool.UEBuildPlatform.BuildPlatformDictionary)
+					{
+						UBTTargetPlatforms[Index++] = Platform.Key;
+					}
+				}
+				return UBTTargetPlatforms;
+			}
+		}
+		private static UnrealBuildTool.UnrealTargetPlatform[] UBTTargetPlatforms; 
+        public static void PrintCSVFile(string Input)
+        {
+            if (IsBuildMachine && CmdEnv.CSVFile != "")
+            {
+                var CSVBuilder = new StringBuilder();
+                var CSVLineToAppend = String.Format("{0}{1}", Input, Environment.NewLine);
+                CSVBuilder.Append(CSVLineToAppend);
+                try
+                {
+                    File.AppendAllText(CmdEnv.CSVFile, CSVBuilder.ToString());
+                }
+                catch (Exception e)
+                {
+                    Log(System.Diagnostics.TraceEventType.Warning, "Could not append to csv file ({0}) : {1}", CmdEnv.CSVFile, e.ToString());
+                }
+            }
+        }
+
+		#endregion
+
+		#region Properties
+
+		/// <summary>
+		/// Command line parameters for this command (empty by non-null by default)
+		/// </summary>
+		private object[] CommandLineParams = new object[0];
+		public object[] Params
+		{
+			get { return CommandLineParams; }
+			set { CommandLineParams = value; }
+		}
+
+		/// <summary>
+		/// Path to the AutomationTool executable.
+		/// </summary>
+		public static string ExeFilename
+		{
+			get { return InternalUtils.ExecutingAssemblyLocation; }
+		}
+
+		/// <summary>
+		/// Directory where the AutomationTool executable sits.
+		/// </summary>
+		public static string ExeDirectory
+		{
+			get { return InternalUtils.ExecutingAssemblyDirectory; }
+		}
+
+		/// <summary>
+		/// Current directory.
+		/// </summary>
+		public static string CurrentDirectory
+		{
+			get { return Environment.CurrentDirectory; }
+			set { ChDir(value); }
+		}
+
+		/// <summary>
+		/// Checks if this command is running on a build machine.
+		/// </summary>
+		public static bool IsBuildMachine
+		{
+			get { return Automation.IsBuildMachine; }
+		}
+
+		#endregion
+
+
+        public static string RootSharedTempStorageDirectory()
+        {
+            string StorageDirectory = "";
+            if (UnrealBuildTool.Utils.IsRunningOnMono)
+            {
+                StorageDirectory = "/Volumes/Builds";
+            }
+            else
+            {
+                StorageDirectory = CombinePaths("P:", "Builds");
+            }
+            return StorageDirectory;
+        }
+
+        public static bool DirectoryExistsAndIsWritable_NoExceptions(string Dir)
+        {
+            if (!DirectoryExists_NoExceptions(Dir))
+            {
+				Log(System.Diagnostics.TraceEventType.Information, "Directory {0} does not exist", Dir);
+				return false;
+			}
+
+            try
+            {
+				string Filename = CombinePaths(Dir, Guid.NewGuid().ToString() + ".Temp.txt");
+				string NativeFilename = ConvertSeparators(PathSeparator.Default, Filename);
+				using(StreamWriter Writer = new StreamWriter(NativeFilename))
+				{
+					Writer.Write("Test");
+				}
+				if(File.Exists(NativeFilename))
+				{
+		            DeleteFile_NoExceptions(Filename, true);
+		            Log(System.Diagnostics.TraceEventType.Information, "Directory {0} is writable", Dir);
+					return true;
+				}
+			}
+			catch(Exception)
+			{
+			}
+
+			Log(System.Diagnostics.TraceEventType.Information, "Directory {0} is not writable", Dir);
+			return false;
+		}
+
+        public static void CleanFormalBuilds(string DirectoryForThisBuild, string CLString = "", int MaximumDaysToKeepTempStorage = 4)
+        {
+            if (CLString == "" && (!IsBuildMachine || !DirectoryForThisBuild.StartsWith(RootSharedTempStorageDirectory()) || !P4Enabled))
+            {
+                return;
+            }
+            try
+            {
+                if (P4Enabled && CLString == "")
+                {
+                    CLString = P4Env.ChangelistString;
+                }
+                string ParentDir = Path.GetDirectoryName(CombinePaths(DirectoryForThisBuild));
+                if (!DirectoryExists_NoExceptions(ParentDir))
+                {
+                    throw new AutomationException("Not cleaning formal builds, because the parent directory {0} does not exist.", ParentDir);
+                }
+                string MyDir = Path.GetFileName(CombinePaths(DirectoryForThisBuild));
+                int CLStart = MyDir.IndexOf(CLString);
+                if (CLStart < 0)
+                {
+                    throw new AutomationException("Not cleaning formal builds, because the directory {0} does not contain the CL {1}.", DirectoryForThisBuild, CLString);
+                }
+                string StartString = MyDir.Substring(0, CLStart);
+                string EndString = MyDir.Substring(CLStart + CLString.Length);
+
+
+                DirectoryInfo DirInfo = new DirectoryInfo(ParentDir);
+                var TopLevelDirs = DirInfo.GetDirectories();
+                Log("Looking for directories to delete in {0}   {1} dirs", ParentDir, TopLevelDirs.Length);
+                foreach (var TopLevelDir in TopLevelDirs)
+                {
+                    if (DirectoryExists_NoExceptions(TopLevelDir.FullName))
+                    {
+                        var JustDir = Path.GetFileName(CombinePaths(TopLevelDir.FullName));
+                        if (JustDir.StartsWith(StartString, StringComparison.InvariantCultureIgnoreCase) && (String.IsNullOrEmpty(EndString) || JustDir.EndsWith(EndString, StringComparison.InvariantCultureIgnoreCase)))
+                        {
+                            string CLPart = JustDir.Substring(StartString.Length, JustDir.Length - StartString.Length - EndString.Length);
+                            if (CLPart.Contains("-PF-") || // always delete preflights
+                                (!CLPart.Contains("-") && !CLPart.Contains("+"))) // never delete anything else that is weird, this is probably another branch
+                            {
+                                DirectoryInfo ThisDirInfo = new DirectoryInfo(TopLevelDir.FullName);
+                                bool bOld = false;
+
+                                if ((DateTime.UtcNow - ThisDirInfo.CreationTimeUtc).TotalDays > MaximumDaysToKeepTempStorage)
+                                {
+                                    bOld = true;
+                                }
+                                if (bOld)
+                                {
+                                    Log("Deleting temp storage directory {0}, because it is more than {1} days old.", TopLevelDir.FullName, MaximumDaysToKeepTempStorage);
+                                    DeleteDirectory_NoExceptions(true, TopLevelDir.FullName);
+                                }
+                                else
+                                {
+                                    Log("Not Deleteing temp storage directory {0}, because it is less than {1} days old.", TopLevelDir.FullName, MaximumDaysToKeepTempStorage);
+                                }
+                            }
+                            else
+                            {
+                                Log("skipping {0}, because the CL part {1} had weird characters", JustDir, CLPart);
+                            }
+                        }
+                        else
+                        {
+                            Log("skipping {0}, because it didn't start with {1} or end with {2}", JustDir, StartString, EndString);
+                        }
+                    }
+                }
+            }
+            catch (Exception Ex)
+            {
+                Log(System.Diagnostics.TraceEventType.Warning, "Unable to Clean Directory with DirectoryForThisBuild {0}", DirectoryForThisBuild);
+                Log(System.Diagnostics.TraceEventType.Warning, " Exception was {0}", LogUtils.FormatException(Ex));
+            }
+        }
+
+		/// <summary>
+		/// Returns the generic name for a given platform (eg. "Windows" for Win32/Win64)
+		/// </summary>
+		/// <param name="Platform">Specific platform</param>
+		public static string GetGenericPlatformName(UnrealBuildTool.UnrealTargetPlatform Platform)
+		{
+			if(Platform == UnrealTargetPlatform.Win32 || Platform == UnrealTargetPlatform.Win64)
+			{
+				return "Windows";
+			}
+			else
+			{
+				return Enum.GetName(typeof(UnrealBuildTool.UnrealTargetPlatform), Platform);
+			}
+		}
+
+		/// <summary>
+		/// Creates a zip file containing the given input files
+		/// </summary>
+		/// <param name="ZipFileName">Filename for the zip</param>
+		/// <param name="Filter">Filter which selects files to be included in the zip</param>
+		/// <param name="BaseDirectory">Base directory to store relative paths in the zip file to</param>
+		public static void ZipFiles(string ZipFileName, string BaseDirectory, FileFilter Filter)
+		{
+			Ionic.Zip.ZipFile Zip = new Ionic.Zip.ZipFile();
+			Zip.UseZip64WhenSaving = Ionic.Zip.Zip64Option.Always;
+			foreach(string FilteredFile in Filter.ApplyToDirectory(BaseDirectory, true))
+			{
+				Zip.AddFile(Path.Combine(BaseDirectory, FilteredFile), Path.GetDirectoryName(FilteredFile));
+			}
+			CommandUtils.CreateDirectory(Path.GetDirectoryName(ZipFileName));
+			Zip.Save(ZipFileName);
+		}
+
+		/// <summary>
+		/// Extracts the contents of a zip file
+		/// </summary>
+		/// <param name="ZipFileName">Name of the zip file</param>
+		/// <param name="BaseDirectory">Output directory</param>
+		/// <returns>List of files written</returns>
+		public static IEnumerable<string> UnzipFiles(string ZipFileName, string BaseDirectory)
+		{
+			using(Ionic.Zip.ZipFile Zip = new Ionic.Zip.ZipFile(ZipFileName))
+			{
+				// For some reason, calling ExtractAll() under Mono throws an exception from a call to Directory.CreateDirectory(). Manually extract all the files instead.
+				List<string> OutputFileNames = new List<string>();
+				foreach(Ionic.Zip.ZipEntry Entry in Zip.Entries)
+				{
+					string OutputFileName = Path.Combine(BaseDirectory, Entry.FileName);
+					Directory.CreateDirectory(Path.GetDirectoryName(OutputFileName));
+					using(FileStream OutputStream = new FileStream(OutputFileName, FileMode.Create, FileAccess.Write))
+					{
+						Entry.Extract(OutputStream);
+					}
+					if (UnrealBuildTool.Utils.IsRunningOnMono && CommandUtils.IsProbablyAMacOrIOSExe(OutputFileName))
+					{
+						FixUnixFilePermissions(OutputFileName);
+					}
+					OutputFileNames.Add(OutputFileName);
+				}
+				return OutputFileNames;
+			}
+		}
+	}
+
+	/// <summary>
+	/// Use with "using" syntax to push and pop directories in a convenient, exception-safe way
+	/// </summary>
+	public class PushedDirectory : IDisposable
+	{
+		public PushedDirectory(string DirectoryName)
+		{
+			CommandUtils.PushDir(DirectoryName);
+		}
+
+		public void Dispose()
+		{
+			CommandUtils.PopDir();
+			GC.SuppressFinalize(this);
+		}
+	}
+
+    /// <summary>
+    /// Helper class to associate a file and its contents
+    /// </summary>
+    public class EMSFileInfo
+    {
+        public string FileName { get; set; }
+        public byte[] Bytes { get; set; }
+    }
+
+    /// <summary>
+    /// Wrapper class for the enumerate files JSON response from MCP
+    /// </summary>
+    [DataContract]
+    public sealed class EnumerationResponse
+    {
+        [DataMember(Name = "doNotCache", IsRequired = true)]
+        public Boolean DoNotCache { get; set; }
+
+        [DataMember(Name = "uniqueFilename", IsRequired = true)]
+        public string UniqueFilename { get; set; }
+
+        [DataMember(Name = "filename", IsRequired = true)]
+        public string Filename { get; set; }
+
+        [DataMember(Name = "hash", IsRequired = true)]
+        public string Hash { get; set; }
+
+        [DataMember(Name = "length", IsRequired = true)]
+        public long Length { get; set; }
+
+        [DataMember(Name = "uploaded", IsRequired = true)]
+        public string Uploaded { get; set; }
+    }
+
+	/// <summary>
+	/// Code signing
+	/// </summary>
+	[Help("NoSign", "Skips signing of code/content files.")]
+	public class CodeSign
+	{
+		/// <summary>
+		/// If so, what is the signing identity to search for?
+		/// </summary>
+		public static string SigningIdentity = "Epic Games";
+
+		/// <summary>
+		/// Should we use the machine store?
+		/// </summary>
+		public static bool bUseMachineStoreInsteadOfUserStore = false;
+
+		/// <summary>
+		/// How long to keep re-trying code signing for
+		/// </summary>
+		public static TimeSpan CodeSignTimeOut = new TimeSpan(0, 3, 0); // Keep trying to sign one file for up to 3 minutes
+
+		/// <summary>
+		/// Code signs the specified file
+		/// </summary>
+		public static void SignSingleExecutableIfEXEOrDLL(string Filename, bool bIgnoreExtension = false)
+		{
+            if (UnrealBuildTool.Utils.IsRunningOnMono)
+            {
+                CommandUtils.Log(TraceEventType.Information, String.Format("Can't sign '{0}', we are running under mono.", Filename));
+                return;
+            }
+            if (!CommandUtils.FileExists(Filename))
+			{
+				throw new AutomationException("Can't sign '{0}', file does not exist.", Filename);
+			}
+			// Make sure the file isn't read-only
+			FileInfo TargetFileInfo = new FileInfo(Filename);
+
+			// Executable extensions
+			List<string> Extensions = new List<string>();
+			Extensions.Add(".dll");
+			Extensions.Add(".exe");
+
+			bool IsExecutable = bIgnoreExtension;
+
+			foreach (var Ext in Extensions)
+			{
+				if (TargetFileInfo.FullName.EndsWith(Ext, StringComparison.InvariantCultureIgnoreCase))
+				{
+					IsExecutable = true;
+					break;
+				}
+			}
+			if (!IsExecutable)
+			{
+				CommandUtils.Log(TraceEventType.Verbose, String.Format("Won't sign '{0}', not an executable.", TargetFileInfo.FullName));
+				return;
+			}
+
+			string SignToolName = null;
+			if (WindowsPlatform.Compiler == WindowsCompiler.VisualStudio2013)
+			{
+				SignToolName = "C:/Program Files (x86)/Windows Kits/8.1/bin/x86/SignTool.exe";
+			}
+			else if (WindowsPlatform.Compiler == WindowsCompiler.VisualStudio2012)
+			{
+				SignToolName = "C:/Program Files (x86)/Windows Kits/8.0/bin/x86/SignTool.exe";
+			}
+
+
+			if (!File.Exists(SignToolName))
+			{
+				throw new AutomationException("SignTool not found at '{0}' (are you missing the Windows SDK?)", SignToolName);
+			}
+
+			TargetFileInfo.IsReadOnly = false;
+
+			// Code sign the executable
+			string TimestampServer = "http://timestamp.verisign.com/scripts/timestamp.dll";
+
+			string SpecificStoreArg = bUseMachineStoreInsteadOfUserStore ? " /sm" : "";
+
+			//@TODO: Verbosity choosing
+			//  /v will spew lots of info
+			//  /q does nothing on success and minimal output on failure
+			string CodeSignArgs = String.Format("sign{0} /a /n \"{1}\" /t {2} /d \"{3}\" /v \"{4}\"", SpecificStoreArg, SigningIdentity, TimestampServer, TargetFileInfo.Name, TargetFileInfo.FullName);
+
+			DateTime StartTime = DateTime.Now;
+
+			int NumTrials = 0;
+			for (; ; )
+			{
+				ProcessResult Result = CommandUtils.Run(SignToolName, CodeSignArgs, null, CommandUtils.ERunOptions.AllowSpew);
+				++NumTrials;
+
+				if (Result.ExitCode != 1)
+				{
+					if (Result.ExitCode == 2)
+					{
+						CommandUtils.Log(TraceEventType.Error, String.Format("Signtool returned a warning."));
+					}
+					// Success!
+					break;
+				}
+				else
+				{
+					// Keep retrying until we run out of time
+					TimeSpan RunTime = DateTime.Now - StartTime;
+					if (RunTime > CodeSignTimeOut)
+					{
+						throw new AutomationException("Failed to sign executable '{0}' {1} times over a period of {2}", TargetFileInfo.FullName, NumTrials, RunTime);
+					}
+				}
+			}
+		}
+
+		/// <summary>
+		/// Code signs the specified file or folder
+		/// </summary>
+		public static void SignMacFileOrFolder(string InPath, bool bIgnoreExtension = false)
+		{
+			bool bExists = CommandUtils.FileExists(InPath) || CommandUtils.DirectoryExists(InPath);
+			if (!bExists)
+			{
+				throw new AutomationException("Can't sign '{0}', file or folder does not exist.", InPath);
+			}
+
+			// @todo: OS X 10.9.5/Xcode 6 can't sign libsteam_api.dylib, so we temporarily disable signing of the editor and games,
+			// which code sign individual files (contrary to the Launcher, which signs the whole app bundle)
+			if (CommandUtils.FileExists(InPath))
+			{
+				return;
+			}
+
+			// Executable extensions
+			List<string> Extensions = new List<string>();
+			Extensions.Add(".dylib");
+			Extensions.Add(".app");
+
+			bool IsExecutable = bIgnoreExtension || (Path.GetExtension(InPath) == "" && !InPath.EndsWith("PkgInfo"));
+
+			foreach (var Ext in Extensions)
+			{
+				if (InPath.EndsWith(Ext, StringComparison.InvariantCultureIgnoreCase))
+				{
+					IsExecutable = true;
+					break;
+				}
+			}
+			if (!IsExecutable)
+			{
+				CommandUtils.Log(TraceEventType.Verbose, String.Format("Won't sign '{0}', not an executable.", InPath));
+				return;
+			}
+
+			string SignToolName = "codesign";
+
+			string CodeSignArgs = String.Format("-f --deep -s \"{0}\" -v \"{1}\"", "Developer ID Application", InPath);
+
+			DateTime StartTime = DateTime.Now;
+
+			int NumTrials = 0;
+			for (; ; )
+			{
+				ProcessResult Result = CommandUtils.Run(SignToolName, CodeSignArgs, null, CommandUtils.ERunOptions.AllowSpew);
+				int ExitCode = Result.ExitCode;
+				++NumTrials;
+
+				if (ExitCode == 0)
+				{
+					// Success!
+					break;
+				}
+				else
+				{
+					// Keep retrying until we run out of time
+					TimeSpan RunTime = DateTime.Now - StartTime;
+					if (RunTime > CodeSignTimeOut)
+					{
+						throw new AutomationException("Failed to sign '{0}' {1} times over a period of {2}", InPath, NumTrials, RunTime);
+					}
+				}
+			}
+		}
+
+		/// <summary>
+		/// Codesigns multiple files, but skips anything that's not an EXE or DLL file
+		/// Will automatically skip signing if -NoSign is specified in the command line.
+		/// </summary>
+		/// <param name="Files">List of files to sign</param>
+		public static void SignMultipleIfEXEOrDLL(BuildCommand Command, List<string> Files)
+		{
+			if (!Command.ParseParam("NoSign"))
+			{
+				CommandUtils.Log("Signing up to {0} files...", Files.Count);
+				UnrealBuildTool.UnrealTargetPlatform TargetPlatform = UnrealBuildTool.BuildHostPlatform.Current.Platform;
+				if (TargetPlatform == UnrealBuildTool.UnrealTargetPlatform.Mac)
+				{
+					foreach (var File in Files)
+					{
+						SignMacFileOrFolder(File);
+					}
+				}
+				else
+				{
+					List<string> FilesToSign = new List<string>();
+					foreach (string File in Files)
+					{
+						if (!(Path.GetDirectoryName(File).Replace("\\", "/")).Contains("Binaries/XboxOne"))
+						{
+							FilesToSign.Add(File);
+						}						
+					}
+					SignMultipleFilesIfEXEOrDLL(FilesToSign);
+				}
+			}
+			else
+			{
+				CommandUtils.Log("Skipping signing {0} files due to -nosign.", Files.Count);
+			}
+		}
+
+		public static void SignListFilesIfEXEOrDLL(string FilesToSign)
+		{
+			string SignToolName = null;
+			if (WindowsPlatform.Compiler == WindowsCompiler.VisualStudio2013)
+			{
+				SignToolName = "C:/Program Files (x86)/Windows Kits/8.1/bin/x86/SignTool.exe";
+			}
+			else if (WindowsPlatform.Compiler == WindowsCompiler.VisualStudio2012)
+			{
+				SignToolName = "C:/Program Files (x86)/Windows Kits/8.0/bin/x86/SignTool.exe";
+			}
+
+
+			if (!File.Exists(SignToolName))
+			{
+				throw new AutomationException("SignTool not found at '{0}' (are you missing the Windows SDK?)", SignToolName);
+			}
+
+			// Code sign the executable
+			string TimestampServer = "http://timestamp.verisign.com/scripts/timestamp.dll";
+
+			string SpecificStoreArg = bUseMachineStoreInsteadOfUserStore ? " /sm" : "";	
+			//@TODO: Verbosity choosing
+			//  /v will spew lots of info
+			//  /q does nothing on success and minimal output on failure
+			string CodeSignArgs = String.Format("sign{0} /a /n \"{1}\" /t {2} /d /v {3}", SpecificStoreArg, SigningIdentity, TimestampServer, FilesToSign);
+
+			DateTime StartTime = DateTime.Now;
+
+			int NumTrials = 0;
+			for (; ; )
+			{
+				ProcessResult Result = CommandUtils.Run(SignToolName, CodeSignArgs, null, CommandUtils.ERunOptions.AllowSpew);
+				++NumTrials;
+
+				if (Result.ExitCode != 1)
+				{
+					if (Result.ExitCode == 2)
+					{
+						CommandUtils.Log(TraceEventType.Error, String.Format("Signtool returned a warning."));
+					}
+					// Success!
+					break;
+				}
+				else
+				{
+					// Keep retrying until we run out of time
+					TimeSpan RunTime = DateTime.Now - StartTime;
+					if (RunTime > CodeSignTimeOut)
+					{
+						throw new AutomationException("Failed to sign executables {0} times over a period of {1}", NumTrials, RunTime);
+					}
+				}
+			}
+		}
+
+		public static void SignMultipleFilesIfEXEOrDLL(List<string> Files, bool bIgnoreExtension = false)
+		{
+			if (UnrealBuildTool.Utils.IsRunningOnMono)
+			{
+				CommandUtils.Log(TraceEventType.Information, String.Format("Can't sign we are running under mono."));
+				return;
+			}
+			List<string> FinalFiles = new List<string>();
+			foreach (string Filename in Files)
+			{
+				// Make sure the file isn't read-only
+				FileInfo TargetFileInfo = new FileInfo(Filename);
+
+				// Executable extensions
+				List<string> Extensions = new List<string>();
+				Extensions.Add(".dll");
+				Extensions.Add(".exe");
+
+				bool IsExecutable = bIgnoreExtension;
+
+				foreach (var Ext in Extensions)
+				{
+					if (TargetFileInfo.FullName.EndsWith(Ext, StringComparison.InvariantCultureIgnoreCase))
+					{
+						IsExecutable = true;
+						break;
+					}
+				}
+				if (IsExecutable && CommandUtils.FileExists(Filename))
+				{
+					FinalFiles.Add(Filename);
+				}
+			}			
+
+			StringBuilder FilesToSignBuilder = new StringBuilder();
+			List<string> FinalListSignStrings = new List<string>();
+			foreach(string File in FinalFiles)
+			{
+				FilesToSignBuilder.Append("\"" + File + "\" ");				
+				if(FilesToSignBuilder.Length > 1900)
+				{
+					string AddFilesToFinalList = FilesToSignBuilder.ToString();
+					FinalListSignStrings.Add(AddFilesToFinalList);
+					FilesToSignBuilder.Clear();
+				}
+			}
+			foreach(string FilesToSign in FinalListSignStrings)
+			{
+				SignListFilesIfEXEOrDLL(FilesToSign);
+			}
+
+		}
+	}
+
+}
diff --git a/Engine/Source/Programs/AutomationTool/AutomationUtils/CommandletUtils.cs b/Engine/Source/Programs/AutomationTool/AutomationUtils/CommandletUtils.cs
new file mode 100644
--- /dev/null
+++ b/Engine/Source/Programs/AutomationTool/AutomationUtils/CommandletUtils.cs
@@ -0,0 +1,247 @@
+﻿// Copyright 1998-2015 Epic Games, Inc. All Rights Reserved.
+using System;
+using System.Collections.Generic;
+using System.Linq;
+using System.Text;
+using System.IO;
+
+
+namespace AutomationTool
+{
+	public partial class CommandUtils
+	{
+		#region Commandlets
+
+		/// <summary>
+		/// Runs Cook commandlet.
+		/// </summary>
+		/// <param name="ProjectName">Project name.</param>
+		/// <param name="UE4Exe">The name of the UE4 Editor executable to use.</param>
+		/// <param name="Maps">List of maps to cook, can be null in which case -MapIniSection=AllMaps is used.</param>
+		/// <param name="Dirs">List of directories to cook, can be null</param>
+		/// <param name="TargetPlatform">Target platform.</param>
+		/// <param name="Parameters">List of additional parameters.</param>
+		public static void CookCommandlet(string ProjectName, string UE4Exe = "UE4Editor-Cmd.exe", string[] Maps = null, string[] Dirs = null, string InternationalizationPreset = "", string[] Cultures = null, string TargetPlatform = "WindowsNoEditor", string Parameters = "-Unversioned")
+		{
+			string MapsToCook = "";
+			if (IsNullOrEmpty(Maps))
+			{
+				// MapsToCook = "-MapIniSection=AllMaps";
+			}
+			else
+			{
+				MapsToCook = "-Map=" + CombineCommandletParams(Maps);
+				MapsToCook.Trim();
+			}
+
+			string DirsToCook = "";
+			if (!IsNullOrEmpty(Dirs))
+			{
+				DirsToCook = "-CookDir=" + CombineCommandletParams(Dirs);
+				DirsToCook.Trim();		
+			}
+
+            string CulturesToCook = "";
+            if (!IsNullOrEmpty(Cultures))
+            {
+                CulturesToCook = "-CookCultures=" + CombineCommandletParams(Cultures);
+                CulturesToCook.Trim();
+            }
+
+            RunCommandlet(ProjectName, UE4Exe, "Cook", String.Format("{0} {1} -I18NPreset={2} {3} -TargetPlatform={4} {5}", MapsToCook, DirsToCook, InternationalizationPreset, CulturesToCook, TargetPlatform, Parameters));
+		}
+
+        /// <summary>
+        /// Runs DDC commandlet.
+        /// </summary>
+        /// <param name="ProjectName">Project name.</param>
+        /// <param name="UE4Exe">The name of the UE4 Editor executable to use.</param>
+        /// <param name="Maps">List of maps to cook, can be null in which case -MapIniSection=AllMaps is used.</param>
+        /// <param name="TargetPlatform">Target platform.</param>
+        /// <param name="Parameters">List of additional parameters.</param>
+        public static void DDCCommandlet(string ProjectName, string UE4Exe = "UE4Editor-Cmd.exe", string[] Maps = null, string TargetPlatform = "WindowsNoEditor", string Parameters = "")
+        {
+            string MapsToCook = "";
+            if (!IsNullOrEmpty(Maps))
+            {
+                MapsToCook = "-Map=" + CombineCommandletParams(Maps);
+                MapsToCook.Trim();
+            }
+
+            RunCommandlet(ProjectName, UE4Exe, "DerivedDataCache", String.Format("{0} -TargetPlatform={1} {2}", MapsToCook, TargetPlatform, Parameters));
+        }
+
+        /// <summary>
+        /// Runs GenerateDistillFileSets commandlet.
+        /// </summary>
+        /// <param name="ProjectName">Project name.</param>
+        /// <param name="UE4Exe">The name of the UE4 Editor executable to use.</param>
+        /// <param name="Maps">List of maps to cook, can be null in which case -MapIniSection=AllMaps is used.</param>
+        /// <param name="TargetPlatform">Target platform.</param>
+        /// <param name="Parameters">List of additional parameters.</param>
+        public static List<string> GenerateDistillFileSetsCommandlet(string ProjectName, string ManifestFile, string UE4Exe = "UE4Editor-Cmd.exe", string[] Maps = null, string Parameters = "")
+        {
+            string MapsToCook = "";
+            if (!IsNullOrEmpty(Maps))
+            {
+                MapsToCook = CombineCommandletParams(Maps, " ");
+                MapsToCook.Trim();
+            }
+            var Dir = Path.GetDirectoryName(ManifestFile);
+            var Filename = Path.GetFileName(ManifestFile);
+            if (String.IsNullOrEmpty(Dir) || String.IsNullOrEmpty(Filename))
+            {
+                throw new AutomationException("GenerateDistillFileSets should have a full path and file for {0}.", ManifestFile);
+            }
+            CreateDirectory_NoExceptions(Dir);
+            if (FileExists_NoExceptions(ManifestFile))
+            {
+                DeleteFile(ManifestFile);
+            }
+
+            RunCommandlet(ProjectName, UE4Exe, "GenerateDistillFileSets", String.Format("{0} -OutputFolder={1} -Output={2} {3}", MapsToCook, CommandUtils.MakePathSafeToUseWithCommandLine(Dir), Filename, Parameters));
+
+            if (!FileExists_NoExceptions(ManifestFile))
+            {
+                throw new AutomationException("GenerateDistillFileSets did not produce a manifest for {0}.", ProjectName);
+            }
+            var Lines = new List<string>(ReadAllLines(ManifestFile));
+            if (Lines.Count < 1)
+            {
+                throw new AutomationException("GenerateDistillFileSets for {0} did not produce any files.", ProjectName);
+            }
+            var Result = new List<string>();
+            foreach (var ThisFile in Lines)
+            {
+                var TestFile = CombinePaths(ThisFile);
+                if (!FileExists_NoExceptions(TestFile))
+                {
+                    throw new AutomationException("GenerateDistillFileSets produced {0}, but {1} doesn't exist.", ThisFile, TestFile);
+                }
+                // we correct the case here
+                var TestFileInfo = new FileInfo(TestFile);
+                var FinalFile = CombinePaths(TestFileInfo.FullName);
+                if (!FileExists_NoExceptions(FinalFile))
+                {
+                    throw new AutomationException("GenerateDistillFileSets produced {0}, but {1} doesn't exist.", ThisFile, FinalFile);
+                }
+                Result.Add(FinalFile);
+            }
+            return Result;
+        }
+
+        /// <summary>
+        /// Runs UpdateGameProject commandlet.
+        /// </summary>
+        /// <param name="ProjectName">Project name.</param>
+        /// <param name="UE4Exe">The name of the UE4 Editor executable to use.</param>
+        /// <param name="Parameters">List of additional parameters.</param>
+        public static void UpdateGameProjectCommandlet(string ProjectName, string UE4Exe = "UE4Editor-Cmd.exe", string Parameters = "")
+        {
+            RunCommandlet(ProjectName, UE4Exe, "UpdateGameProject", Parameters);
+        }
+		/// <summary>
+		/// Runs a commandlet using Engine/Binaries/Win64/UE4Editor-Cmd.exe.
+		/// </summary>
+		/// <param name="ProjectName">Project name.</param>
+		/// <param name="UE4Exe">The name of the UE4 Editor executable to use.</param>
+		/// <param name="Commandlet">Commandlet name.</param>
+		/// <param name="Parameters">Command line parameters (without -run=)</param>
+		public static void RunCommandlet(string ProjectName, string UE4Exe, string Commandlet, string Parameters = null)
+		{
+			Log("Running UE4Editor {0} for project {1}", Commandlet, ProjectName);
+
+            var CWD = Path.GetDirectoryName(UE4Exe);
+
+            string EditorExe = UE4Exe;
+
+            if (String.IsNullOrEmpty(CWD))
+            {
+                EditorExe = HostPlatform.Current.GetUE4ExePath(UE4Exe);
+                CWD = CombinePaths(CmdEnv.LocalRoot, HostPlatform.Current.RelativeBinariesFolder);
+            }
+			
+			PushDir(CWD);
+
+			string LocalLogFile = LogUtils.GetUniqueLogName(CombinePaths(CmdEnv.EngineSavedFolder, Commandlet));
+			Log("Commandlet log file is {0}", LocalLogFile);
+			string Args = String.Format(
+				"{0} -run={1} {2} -abslog={3} -stdout -FORCELOGFLUSH -CrashForUAT -unattended -AllowStdOutLogVerbosity {4}",
+				(ProjectName == null) ? "" : CommandUtils.MakePathSafeToUseWithCommandLine(ProjectName),
+				Commandlet,
+				String.IsNullOrEmpty(Parameters) ? "" : Parameters,
+				CommandUtils.MakePathSafeToUseWithCommandLine(LocalLogFile),
+				IsBuildMachine ? "-buildmachine" : ""
+			);
+			ERunOptions Opts = ERunOptions.Default;
+			if (GlobalCommandLine.UTF8Output)
+			{
+				Args += " -UTF8Output";
+				Opts |= ERunOptions.UTF8Output;
+			}
+			var RunResult = Run(EditorExe, Args, Options: Opts);
+			PopDir();
+
+			// Copy the local commandlet log to the destination folder.
+			string DestLogFile = LogUtils.GetUniqueLogName(CombinePaths(CmdEnv.LogFolder, Commandlet));
+			if (!CommandUtils.CopyFile_NoExceptions(LocalLogFile, DestLogFile))
+			{
+				CommandUtils.LogWarning("Commandlet {0} failed to copy the local log file from {1} to {2}. The log file will be lost.", Commandlet, LocalLogFile, DestLogFile);
+			}
+
+			// Whether it was copied correctly or not, delete the local log as it was only a temporary file. 
+			CommandUtils.DeleteFile_NoExceptions(LocalLogFile);
+
+			if (RunResult.ExitCode != 0)
+			{
+				throw new AutomationException("BUILD FAILED: Failed while running {0} for {1}; see log {2}", Commandlet, ProjectName, DestLogFile);
+			}
+		}
+
+		/// <summary>
+		/// Returns the default path of the editor executable to use for running commandlets.
+		/// </summary>
+		/// <param name="BuildRoot">Root directory for the build</param>
+		/// <param name="HostPlatform">Platform to get the executable for</param>
+		/// <returns>Path to the editor executable</returns>
+		public static string GetEditorCommandletExe(string BuildRoot, UnrealBuildTool.UnrealTargetPlatform HostPlatform)
+		{
+			switch(HostPlatform)
+			{
+				case UnrealBuildTool.UnrealTargetPlatform.Mac:
+					return CommandUtils.CombinePaths(BuildRoot, "Engine/Binaries/Mac/UE4Editor.app/Contents/MacOS/UE4Editor");
+				case UnrealBuildTool.UnrealTargetPlatform.Win64:
+					return CommandUtils.CombinePaths(BuildRoot, "Engine/Binaries/Win64/UE4Editor-Cmd.exe");
+				case UnrealBuildTool.UnrealTargetPlatform.Linux:
+					return CommandUtils.CombinePaths(BuildRoot, "Engine/Binaries/Linux/UE4Editor");
+				default:
+					throw new NotImplementedException();
+			}
+		}
+
+		/// <summary>
+		/// Combines a list of parameters into a single commandline param separated with '+':
+		/// For example: Map1+Map2+Map3
+		/// </summary>
+		/// <param name="ParamValues">List of parameters (must not be empty)</param>
+		/// <returns>Combined param</returns>
+		public static string CombineCommandletParams(string[] ParamValues, string Separator = "+")
+		{
+			if (!IsNullOrEmpty(ParamValues))
+			{
+				var CombinedParams = ParamValues[0];
+				for (int Index = 1; Index < ParamValues.Length; ++Index)
+				{
+                    CombinedParams += Separator + ParamValues[Index];
+				}
+				return CombinedParams;
+			}
+			else
+			{
+				return String.Empty;
+			}
+		}
+
+		#endregion
+	}
+}
