UEB-260 - Break AutomationTool into AutomationUtils that all automation projects depend on, and AutomationTool, which essentially only contains the startup code.

diff --git a/Engine/Source/Programs/AutomationTool/AutomationUtils/DeploymentContext.cs b/Engine/Source/Programs/AutomationTool/AutomationUtils/DeploymentContext.cs
new file mode 100644
--- /dev/null
+++ b/Engine/Source/Programs/AutomationTool/AutomationUtils/DeploymentContext.cs
@@ -0,0 +1,713 @@
+﻿// Copyright 1998-2015 Epic Games, Inc. All Rights Reserved.
+using System;
+using System.Collections.Generic;
+using System.IO;
+using System.Threading;
+using System.Reflection;
+using AutomationTool;
+using UnrealBuildTool;
+
+public enum StagedFileType
+{
+	UFS,
+	NonUFS,
+	DebugNonUFS
+}
+
+public class DeploymentContext //: ProjectParams
+{
+	/// <summary>
+	/// Full path where the project exists (For uprojects this should include the uproj filename, otherwise just project folder)
+	/// </summary>
+	public string RawProjectPath;
+
+	/// <summary>
+	///  true if we should stage crash reporter
+	/// </summary>
+	public bool bStageCrashReporter;
+
+    /// <summary>
+    ///  CookPlatform, where to get the cooked data from and use for sandboxes
+    /// </summary>
+    public string CookPlatform;
+
+	/// <summary>
+    ///  FinalCookPlatform, directory to stage and archive the final result to
+	/// </summary>
+    public string FinalCookPlatform;
+
+    /// <summary>
+    ///  Source platform to get the cooked data from
+    /// </summary>
+    public Platform CookSourcePlatform;
+
+	/// <summary>
+	///  Target platform used for sandboxes and stage directory names
+	/// </summary>
+	public Platform StageTargetPlatform;
+
+	/// <summary>
+	///  Configurations to stage. Used to determine which ThirdParty configurations to copy.
+	/// </summary>
+	public List<UnrealTargetConfiguration> StageTargetConfigurations;
+
+	/// <summary>
+	/// Receipts for the build targets that should be staged
+	/// </summary>
+	public List<BuildReceipt> StageTargetReceipts;
+
+	/// <summary>
+	///  this is the root directory that contains the engine: d:\a\UE4\
+	/// </summary>
+	public string LocalRoot;
+
+	/// <summary>
+	///  this is the directory that contains the "game": d:\a\UE4\ShooterGame
+	/// </summary>
+	public string ProjectRoot;
+
+	/// <summary>
+	///  raw name used for platform subdirectories Win32
+	/// </summary>
+	public string PlatformDir;
+
+	/// <summary>
+	///  this is the directory that contains the "game", staged: d:\stagedir\WindowsNoEditor\ShooterGame
+	/// </summary>
+	public string StageProjectRoot;
+
+	/// <summary>
+	///  Directory to put all of the files in: d:\stagedir\WindowsNoEditor
+	/// </summary>
+	public string StageDirectory;
+
+	/// <summary>
+	///  The relative staged project root, what would be tacked on to StageDirectory
+	/// </summary>
+    public string RelativeProjectRootForStage;
+
+    /// <summary>
+    ///  The relative source project root, this will be the short project name for foreign projects, otherwise it is ProjectRoot minus the LocalRoot prefix
+    /// </summary>
+    public string SourceRelativeProjectRoot;
+
+    /// <summary>
+	///  The relative staged project root, used inside the pak file list (different slash convention than above)
+	/// </summary>
+	public string RelativeProjectRootForUnrealPak;
+
+	/// <summary>
+	///  This is what you use to test the engine which uproject you want. Many cases.
+	/// </summary>
+	public string ProjectArgForCommandLines;
+
+    /// <summary>
+    ///  This is the root that the cook source platform would run from
+    /// </summary>
+    public string CookSourceRuntimeRootDir;
+
+	/// <summary>
+	///  This is the root that we are going to run from. Many cases.
+	/// </summary>
+	public string RuntimeRootDir;
+
+	/// <summary>
+	///  This is the project root that we are going to run from. Many cases.
+	/// </summary>
+	public string RuntimeProjectRootDir;
+
+	/// <summary>
+	///  This is the executable we are going to run and stage. Filled in by the platform abstraction.
+	/// </summary>
+	public string RuntimeExecutable = "";
+
+	/// <summary>
+	///  List of executables we are going to stage
+	/// </summary>
+	public List<string> StageExecutables;
+
+	/// <summary>
+	///  Probably going away, used to construct ProjectArgForCommandLines in the case that we are running staged
+	/// </summary>
+	public string UProjectCommandLineArgInternalRoot = "../../../";
+
+	/// <summary>
+	///  Probably going away, used to construct the pak file list
+	/// </summary>
+	public string PakFileInternalRoot = "../../../";
+
+	/// <summary>
+	///  Probably going away, and currently unused, this would be used to build the list of files for the UFS server
+	/// </summary>
+	public string UnrealFileServerInternalRoot = "../../../";
+
+	/// <summary>
+	///  After staging, this is a map from source file to relative file in the stage
+	/// These file are binaries, etc and can't go into a pak file
+	/// </summary>
+	public Dictionary<string, string> NonUFSStagingFiles = new Dictionary<string, string>();
+	/// <summary>
+	///  After staging, this is a map from source file to relative file in the stage
+	/// These file are debug, and can't go into a pak file
+	/// </summary>
+	public Dictionary<string, string> NonUFSStagingFilesDebug = new Dictionary<string, string>();
+	/// <summary>
+	///  After staging, this is a map from source file to relative file in the stage
+	/// These file are content, and can go into a pak file
+	/// </summary>
+	public Dictionary<string, string> UFSStagingFiles = new Dictionary<string, string>();
+	/// <summary>
+	///  List of files to be archived
+	/// </summary>
+	public Dictionary<string, string> ArchivedFiles = new Dictionary<string, string>();
+	/// <summary>
+	///  Directory to archive all of the files in: d:\archivedir\WindowsNoEditor
+	/// </summary>
+	public string ArchiveDirectory;
+	/// <summary>
+	///  Directory to project binaries
+	/// </summary>
+	public string ProjectBinariesFolder;
+
+	/// <summary>
+	/// Filename for the manifest of file changes for iterative deployment.
+	/// </summary>
+	static public readonly string UFSDeployDeltaFileName			= "Manifest_DeltaUFSFiles.txt";	
+	static public readonly string NonUFSDeployDeltaFileName			= "Manifest_DeltaNonUFSFiles.txt";
+
+	/// <summary>
+	/// Filename for the manifest of files currently deployed on a device.
+	/// </summary>
+	static public readonly string UFSDeployedManifestFileName		= "Manifest_UFSFiles.txt";
+	static public readonly string NonUFSDeployedManifestFileName	= "Manifest_NonUFSFiles.txt";
+
+	
+
+
+	/// <summary>
+	/// The client connects to dedicated server to get data
+	/// </summary>
+	public bool DedicatedServer;
+
+	/// <summary>
+	/// True if this build is staged
+	/// </summary>
+	public bool Stage;
+
+	/// <summary>
+	/// True if this build is archived
+	/// </summary>
+	public bool Archive;
+
+	/// <summary>
+	/// True if this project has code
+	/// </summary>	
+	public bool IsCodeBasedProject;
+
+	/// <summary>
+	/// Project name (name of the uproject file without extension or directory name where the project is localed)
+	/// </summary>
+	public string ShortProjectName;
+
+	/// <summary>
+	/// If true, multiple platforms are being merged together - some behavior needs to change (but not much)
+	/// </summary>
+	public bool bIsCombiningMultiplePlatforms = false;
+
+	public DeploymentContext(
+		string RawProjectPathOrName,
+		string InLocalRoot,
+		string BaseStageDirectory,
+		string BaseArchiveDirectory,
+		string CookFlavor,
+		Platform InSourcePlatform,
+        Platform InTargetPlatform,
+		List<UnrealTargetConfiguration> InTargetConfigurations,
+		IEnumerable<BuildReceipt> InStageTargetReceipts,
+		List<String> InStageExecutables,
+		bool InServer,
+		bool InCooked,
+		bool InStageCrashReporter,
+		bool InStage,
+		bool InCookOnTheFly,
+		bool InArchive,
+		bool InProgram,
+		bool bHasDedicatedServerAndClient
+		)
+	{
+		bStageCrashReporter = InStageCrashReporter;
+		RawProjectPath = RawProjectPathOrName;
+		DedicatedServer = InServer;
+		LocalRoot = CommandUtils.CombinePaths(InLocalRoot);
+        CookSourcePlatform = InSourcePlatform;
+		StageTargetPlatform = InTargetPlatform;
+		StageTargetConfigurations = new List<UnrealTargetConfiguration>(InTargetConfigurations);
+		StageTargetReceipts = new List<BuildReceipt>(InStageTargetReceipts);
+		StageExecutables = InStageExecutables;
+        IsCodeBasedProject = ProjectUtils.IsCodeBasedUProjectFile(RawProjectPath);
+		ShortProjectName = ProjectUtils.GetShortProjectName(RawProjectPath);
+		Stage = InStage;
+		Archive = InArchive;
+
+        if (CookSourcePlatform != null && InCooked)
+        {
+            CookPlatform = CookSourcePlatform.GetCookPlatform(DedicatedServer, bHasDedicatedServerAndClient, CookFlavor);
+        }
+        else if (CookSourcePlatform != null && InProgram)
+        {
+            CookPlatform = CookSourcePlatform.GetCookPlatform(false, false, "");
+        }
+        else
+        {
+            CookPlatform = "";
+        }
+
+		if (StageTargetPlatform != null && InCooked)
+		{
+            FinalCookPlatform = StageTargetPlatform.GetCookPlatform(DedicatedServer, bHasDedicatedServerAndClient, CookFlavor);
+		}
+		else if (StageTargetPlatform != null && InProgram)
+		{
+            FinalCookPlatform = StageTargetPlatform.GetCookPlatform(false, false, "");
+		}
+		else
+		{
+            FinalCookPlatform = "";
+		}
+
+		PlatformDir = StageTargetPlatform.PlatformType.ToString();
+
+        StageDirectory = CommandUtils.CombinePaths(BaseStageDirectory, FinalCookPlatform);
+        ArchiveDirectory = CommandUtils.CombinePaths(BaseArchiveDirectory, FinalCookPlatform);
+
+		if (!CommandUtils.FileExists(RawProjectPath))
+		{
+			throw new AutomationException("Can't find uproject file {0}.", RawProjectPathOrName);
+		}
+
+		ProjectRoot = CommandUtils.CombinePaths(CommandUtils.GetDirectoryName(Path.GetFullPath(RawProjectPath)));
+
+		if (!CommandUtils.DirectoryExists(ProjectRoot))
+		{
+			throw new AutomationException("Project Directory {0} doesn't exist.", ProjectRoot);
+		}
+
+        RelativeProjectRootForStage = ShortProjectName;
+
+		ProjectArgForCommandLines = CommandUtils.MakePathSafeToUseWithCommandLine(RawProjectPath);
+		CookSourceRuntimeRootDir = RuntimeRootDir = LocalRoot;
+		RuntimeProjectRootDir = ProjectRoot;
+
+        RelativeProjectRootForUnrealPak = CommandUtils.CombinePaths(RelativeProjectRootForStage).Replace("\\", "/");
+		if (RelativeProjectRootForUnrealPak.StartsWith("/"))
+		{
+			RelativeProjectRootForUnrealPak = RelativeProjectRootForUnrealPak.Substring(1);
+            RelativeProjectRootForStage = RelativeProjectRootForStage.Substring(1);
+		}
+
+        SourceRelativeProjectRoot = RelativeProjectRootForStage; // for foreign projects this doesn't make much sense, but it turns into a noop on staging files
+        if (ProjectRoot.StartsWith(LocalRoot, StringComparison.InvariantCultureIgnoreCase))
+        {
+            SourceRelativeProjectRoot = ProjectRoot.Substring(LocalRoot.Length);
+        }
+        if (SourceRelativeProjectRoot.StartsWith("/") || SourceRelativeProjectRoot.StartsWith("\\"))
+        {
+            SourceRelativeProjectRoot = SourceRelativeProjectRoot.Substring(1);
+        }        
+        
+        if (Stage)
+		{
+			CommandUtils.CreateDirectory(StageDirectory);
+            StageProjectRoot = CommandUtils.CombinePaths(StageDirectory, RelativeProjectRootForStage);
+
+			RuntimeRootDir = StageDirectory;
+            CookSourceRuntimeRootDir = CommandUtils.CombinePaths(BaseStageDirectory, CookPlatform);
+			RuntimeProjectRootDir = StageProjectRoot;
+			ProjectArgForCommandLines = CommandUtils.MakePathSafeToUseWithCommandLine(UProjectCommandLineArgInternalRoot + RelativeProjectRootForStage + "/" + ShortProjectName + ".uproject");
+		}
+		if (Archive)
+		{
+			CommandUtils.CreateDirectory(ArchiveDirectory);
+		}
+		ProjectArgForCommandLines = ProjectArgForCommandLines.Replace("\\", "/");
+		ProjectBinariesFolder = CommandUtils.CombinePaths(ProjectUtils.GetClientProjectBinariesRootPath(RawProjectPath, TargetRules.TargetType.Game, IsCodeBasedProject), PlatformDir);
+	}
+
+	public void StageFile(StagedFileType FileType, string InputPath, string OutputPath = null, bool bRemap = true)
+	{
+		InputPath = InputPath.Replace(Path.AltDirectorySeparatorChar, Path.DirectorySeparatorChar);
+
+		if(OutputPath == null)
+		{
+			if(InputPath.StartsWith(ProjectRoot, StringComparison.InvariantCultureIgnoreCase))
+			{
+				OutputPath = CommandUtils.CombinePaths(RelativeProjectRootForStage, InputPath.Substring(ProjectRoot.Length).TrimStart('/', '\\'));
+			}
+			else if (InputPath.StartsWith(LocalRoot, StringComparison.InvariantCultureIgnoreCase))
+			{
+				OutputPath = CommandUtils.CombinePaths(InputPath.Substring(LocalRoot.Length).TrimStart('/', '\\'));
+			}
+			else
+			{
+				throw new AutomationException("Can't deploy {0} because it doesn't start with {1} or {2}", InputPath, ProjectRoot, LocalRoot);
+			}
+		}
+
+		if(bRemap)
+		{
+			OutputPath = StageTargetPlatform.Remap(OutputPath);
+		}
+
+		if (FileType == StagedFileType.UFS)
+		{
+			AddUniqueStagingFile(UFSStagingFiles, InputPath, OutputPath);
+		}
+		else if (FileType == StagedFileType.NonUFS)
+		{
+			AddUniqueStagingFile(NonUFSStagingFiles, InputPath, OutputPath);
+		}
+		else if (FileType == StagedFileType.DebugNonUFS)
+		{
+			AddUniqueStagingFile(NonUFSStagingFilesDebug, InputPath, OutputPath);
+		}
+	}
+
+	public void StageBuildProductsFromReceipt(BuildReceipt Receipt)
+	{
+		// Stage all the build products needed at runtime
+		foreach(BuildProduct BuildProduct in Receipt.BuildProducts)
+		{
+			// allow missing files if needed
+			if (Receipt.bRequireDependenciesToExist == false && File.Exists(BuildProduct.Path) == false)
+			{
+				continue;
+			}
+
+			if(BuildProduct.Type == BuildProductType.Executable || BuildProduct.Type == BuildProductType.DynamicLibrary || BuildProduct.Type == BuildProductType.RequiredResource)
+			{
+				StageFile(StagedFileType.NonUFS, BuildProduct.Path);
+			}
+			else if(BuildProduct.Type == BuildProductType.SymbolFile)
+			{
+				StageFile(StagedFileType.DebugNonUFS, BuildProduct.Path);
+			}
+		}
+	}
+
+	public void StageRuntimeDependenciesFromReceipt(BuildReceipt Receipt)
+	{
+		// Also stage any additional runtime dependencies, like ThirdParty DLLs
+		foreach(RuntimeDependency RuntimeDependency in Receipt.RuntimeDependencies)
+		{
+			// allow missing files if needed
+			if (Receipt.bRequireDependenciesToExist == false && File.Exists(RuntimeDependency.Path) == false)
+			{
+				continue;
+			}
+
+			StageFile(StagedFileType.NonUFS, RuntimeDependency.Path, RuntimeDependency.StagePath);
+		}
+	}
+
+    public int StageFiles(StagedFileType FileType, string InPath, string Wildcard = "*", bool bRecursive = true, string[] ExcludeWildcard = null, string NewPath = null, bool bAllowNone = false, bool bRemap = true, string NewName = null, bool bAllowNotForLicenseesFiles = true, bool bStripFilesForOtherPlatforms = true)
+	{
+		int FilesAdded = 0;
+		// make sure any ..'s are removed
+		Utils.CollapseRelativeDirectories(ref InPath);
+
+		if (CommandUtils.DirectoryExists(InPath))
+		{
+			var All = CommandUtils.FindFiles(Wildcard, bRecursive, InPath);
+
+			var Exclude = new HashSet<string>();
+			if (ExcludeWildcard != null)
+			{
+				foreach (var Excl in ExcludeWildcard)
+				{
+					var Remove = CommandUtils.FindFiles(Excl, bRecursive, InPath);
+					foreach (var File in Remove)
+					{
+                        Exclude.Add(CommandUtils.CombinePaths(File));
+					}
+				}
+			}
+			foreach (var AllFile in All)
+			{
+				var FileToCopy = CommandUtils.CombinePaths(AllFile);
+				if (Exclude.Contains(FileToCopy))
+				{
+					continue;
+				}
+                
+                if (!bAllowNotForLicenseesFiles && (FileToCopy.Contains("NotForLicensees") || FileToCopy.Contains("NoRedist")))
+                {
+                    continue;
+                }
+
+				if (bStripFilesForOtherPlatforms && !bIsCombiningMultiplePlatforms)
+                {
+                    bool OtherPlatform = false;
+                    foreach (UnrealTargetPlatform Plat in Enum.GetValues(typeof(UnrealTargetPlatform)))
+                    {
+                        if (Plat != StageTargetPlatform.PlatformType && Plat != UnrealTargetPlatform.Unknown)
+                        {
+                            var Search = FileToCopy;
+                            if (Search.StartsWith(LocalRoot, StringComparison.InvariantCultureIgnoreCase))
+                            {
+                                if (LocalRoot.EndsWith("\\") || LocalRoot.EndsWith("/"))
+                                {
+                                    Search = Search.Substring(LocalRoot.Length - 1);
+                                }
+                                else
+                                {
+                                    Search = Search.Substring(LocalRoot.Length);
+                                }
+                            }
+							if (Search.StartsWith(ProjectRoot, StringComparison.InvariantCultureIgnoreCase))
+							{
+								if (ProjectRoot.EndsWith("\\") || ProjectRoot.EndsWith("/"))
+								{
+									Search = Search.Substring(ProjectRoot.Length - 1);
+								}
+								else
+								{
+									Search = Search.Substring(ProjectRoot.Length);
+								}
+							}
+                            if (Search.IndexOf(CommandUtils.CombinePaths("/" + Plat.ToString() + "/"), 0, StringComparison.InvariantCultureIgnoreCase) >= 0)
+                            {
+                                OtherPlatform = true;
+                                break;
+                            }
+                        }
+                    }
+                    if (OtherPlatform)
+                    {
+                        continue;
+                    }
+                }
+
+				string Dest;
+				if (!FileToCopy.StartsWith(InPath))
+				{
+					throw new AutomationException("Can't deploy {0}; it was supposed to start with {1}", FileToCopy, InPath);
+				}
+				string FileToRemap = FileToCopy;
+
+                // If the specified a new directory, first we deal with that, then apply the other things
+                // this is used to collapse the sandbox, among other things
+				if (NewPath != null)
+				{
+					Dest = FileToRemap.Substring(InPath.Length);
+					if (Dest.StartsWith("/") || Dest.StartsWith("\\"))
+					{
+						Dest = Dest.Substring(1);
+					}
+					Dest = CommandUtils.CombinePaths(NewPath, Dest);
+#if false // if the cooker rebases, I don't think we need to ever rebase while staging to a new path
+                    if (Dest.StartsWith("/") || Dest.StartsWith("\\"))
+                    {
+                        Dest = Dest.Substring(1);
+                    }
+                    // project relative stuff in a collapsed sandbox
+                    if (Dest.StartsWith(SourceRelativeProjectRoot, StringComparison.InvariantCultureIgnoreCase))
+                    {
+                        Dest = Dest.Substring(SourceRelativeProjectRoot.Length);
+                        if (Dest.StartsWith("/") || Dest.StartsWith("\\"))
+                        {
+                            Dest = Dest.Substring(1);
+                        }
+                        Dest = CommandUtils.CombinePaths(RelativeProjectRootForStage, Dest);
+                    }
+#endif
+                }
+
+                // project relative file
+				else if (FileToRemap.StartsWith(ProjectRoot, StringComparison.InvariantCultureIgnoreCase))
+				{
+					Dest = FileToRemap.Substring(ProjectRoot.Length);
+					if (Dest.StartsWith("/") || Dest.StartsWith("\\"))
+					{
+						Dest = Dest.Substring(1);
+					}
+                    Dest = CommandUtils.CombinePaths(RelativeProjectRootForStage, Dest);
+				}
+                // engine relative file
+				else if (FileToRemap.StartsWith(LocalRoot, StringComparison.InvariantCultureIgnoreCase))
+				{
+					Dest = CommandUtils.CombinePaths(FileToRemap.Substring(LocalRoot.Length));
+				}
+				else
+				{
+					throw new AutomationException("Can't deploy {0} because it doesn't start with {1} or {2}", FileToRemap, ProjectRoot, LocalRoot);
+				}
+
+				if (Dest.StartsWith("/") || Dest.StartsWith("\\"))
+				{
+					Dest = Dest.Substring(1);
+				}
+
+				if (NewName != null)
+				{
+					Dest = CommandUtils.CombinePaths(Path.GetDirectoryName(Dest), NewName);
+				}
+
+				if (bRemap)
+				{
+					Dest = StageTargetPlatform.Remap(Dest);
+				}
+
+				if (FileType == StagedFileType.UFS)
+				{
+					AddUniqueStagingFile(UFSStagingFiles, FileToCopy, Dest);
+				}
+				else if (FileType == StagedFileType.NonUFS)
+				{
+					AddUniqueStagingFile(NonUFSStagingFiles, FileToCopy, Dest);
+				}
+				else if (FileType == StagedFileType.DebugNonUFS)
+				{
+					AddUniqueStagingFile(NonUFSStagingFilesDebug, FileToCopy, Dest);
+				}
+				FilesAdded++;
+			}
+		}
+
+		if (FilesAdded == 0 && !bAllowNone && !bIsCombiningMultiplePlatforms)
+		{
+			throw new AutomationException(ErrorCodes.Error_StageMissingFile, "No files found to deploy for {0} with wildcard {1} and exclusions {2}", InPath, Wildcard, ExcludeWildcard);
+		}
+
+		return FilesAdded;
+	}
+
+	private void AddUniqueStagingFile(Dictionary<string, string> FilesToStage, string FileToCopy, string Dest)
+	{
+		string ExistingDest;
+		if (FilesToStage.TryGetValue(FileToCopy, out ExistingDest))
+		{
+			// If the file already exists, it must have the same dest
+			if (String.Compare(ExistingDest, Dest, true) != 0)
+			{
+				throw new AutomationException("Attempting to add \"{0}\" to staging map but it already exists with a different destination (existing: \"{1}\", new: \"{2}\"",
+					FileToCopy, ExistingDest, Dest);
+			}
+		}
+		else
+		{
+			FilesToStage.Add(FileToCopy, Dest);
+		}
+	}
+
+	public int ArchiveFiles(string InPath, string Wildcard = "*", bool bRecursive = true, string[] ExcludeWildcard = null)
+	{
+		int FilesAdded = 0;
+
+		if (CommandUtils.DirectoryExists(InPath))
+		{
+			var All = CommandUtils.FindFiles(Wildcard, bRecursive, InPath);
+
+			var Exclude = new HashSet<string>();
+			if (ExcludeWildcard != null)
+			{
+				foreach (var Excl in ExcludeWildcard)
+				{
+					var Remove = CommandUtils.FindFiles(Excl, bRecursive, InPath);
+					foreach (var File in Remove)
+					{
+						Exclude.Add(CommandUtils.CombinePaths(File));
+					}
+				}
+			}
+			foreach (var AllFile in All)
+			{
+				var FileToCopy = CommandUtils.CombinePaths(AllFile);
+				if (Exclude.Contains(FileToCopy))
+				{
+					continue;
+				}
+
+				if (!bIsCombiningMultiplePlatforms)
+				{
+					bool OtherPlatform = false;
+					foreach (UnrealTargetPlatform Plat in Enum.GetValues(typeof(UnrealTargetPlatform)))
+					{
+                        if (Plat != StageTargetPlatform.PlatformType && Plat != UnrealTargetPlatform.Unknown)
+                        {
+                            var Search = FileToCopy;
+                            if (Search.StartsWith(LocalRoot, StringComparison.InvariantCultureIgnoreCase))
+                            {
+                                if (LocalRoot.EndsWith("\\") || LocalRoot.EndsWith("/"))
+                                {
+                                    Search = Search.Substring(LocalRoot.Length - 1);
+                                }
+                                else
+                                {
+                                    Search = Search.Substring(LocalRoot.Length);
+                                }
+                            }
+							if (Search.StartsWith(ProjectRoot, StringComparison.InvariantCultureIgnoreCase))
+							{
+								if (ProjectRoot.EndsWith("\\") || ProjectRoot.EndsWith("/"))
+								{
+									Search = Search.Substring(ProjectRoot.Length - 1);
+								}
+								else
+								{
+									Search = Search.Substring(ProjectRoot.Length);
+								}
+							}
+                            if (Search.IndexOf(CommandUtils.CombinePaths("/" + Plat.ToString() + "/"), 0, StringComparison.InvariantCultureIgnoreCase) >= 0)
+                            {
+                                OtherPlatform = true;
+                                break;
+                            }
+                        }
+					}
+					if (OtherPlatform)
+					{
+						continue;
+					}
+				}
+
+				string Dest;
+				if (!FileToCopy.StartsWith(InPath))
+				{
+					throw new AutomationException("Can't archive {0}; it was supposed to start with {1}", FileToCopy, InPath);
+				}
+				Dest = FileToCopy.Substring(InPath.Length);
+
+				if (Dest.StartsWith("/") || Dest.StartsWith("\\"))
+				{
+					Dest = Dest.Substring(1);
+				}
+
+				if (ArchivedFiles.ContainsKey(FileToCopy))
+				{
+					if (ArchivedFiles[FileToCopy] != Dest)
+					{
+						throw new AutomationException("Can't archive {0}: it was already in the files to archive with a different destination '{1}'", FileToCopy, Dest);
+					}
+				}
+				else
+				{
+					ArchivedFiles.Add(FileToCopy, Dest);
+				}
+
+				FilesAdded++;
+			}
+		}
+
+		return FilesAdded;
+	}
+
+	public String GetUFSDeploymentDeltaPath()
+	{
+		return Path.Combine(StageDirectory, UFSDeployDeltaFileName);
+	}
+
+	public String GetNonUFSDeploymentDeltaPath()
+	{
+		return Path.Combine(StageDirectory, NonUFSDeployDeltaFileName);
+	}
+}
diff --git a/Engine/Source/Programs/AutomationTool/AutomationUtils/Distiller.cs b/Engine/Source/Programs/AutomationTool/AutomationUtils/Distiller.cs
new file mode 100644
--- /dev/null
+++ b/Engine/Source/Programs/AutomationTool/AutomationUtils/Distiller.cs
@@ -0,0 +1,264 @@
+﻿// Copyright 1998-2015 Epic Games, Inc. All Rights Reserved.
+using System;
+using System.Collections;
+using System.Collections.Generic;
+using System.Diagnostics;
+using System.IO;
+using UnrealBuildTool;
+
+namespace AutomationTool
+{
+    public partial class CommandUtils
+    {
+        /// <summary>
+        /// Given a path to a file, strips off the base directory part of the path
+        /// </summary>
+        /// <param name="FilePath">The full path</param>
+        /// <param name="BaseDirectory">The base directory, which must be the first part of the path</param>
+        /// <returns>The part of the path after the base directory</returns>
+        public static string StripBaseDirectory(string InFilePath, string InBaseDirectory)
+        {
+            var FilePath = CombinePaths(InFilePath);
+            var BaseDirectory = CombinePaths(InBaseDirectory);
+            if (!FilePath.StartsWith(BaseDirectory, StringComparison.InvariantCultureIgnoreCase))
+            {
+                throw new AutomationException("Cannot strip the base directory {0} from {1} because it doesn't start with the base directory.", BaseDirectory, FilePath);
+            }
+            if (BaseDirectory.EndsWith("/") || BaseDirectory.EndsWith("\\"))
+            {
+                return FilePath.Substring(BaseDirectory.Length);
+            }
+            return FilePath.Substring(BaseDirectory.Length + 1);
+        }
+
+
+        /// <summary>
+        /// Given a path to a "source" file, re-roots the file path to be located under the "destination" folder.  The part of the source file's path after the root folder is unchanged.
+        /// </summary>
+        /// <param name="FilePath"></param>
+        /// <param name="BaseDirectory"></param>
+        /// <param name="NewBaseDirectory"></param>
+        /// <returns></returns>
+        public static string MakeRerootedFilePath(string FilePath, string BaseDirectory, string NewBaseDirectory)
+        {
+            var RelativeFile = StripBaseDirectory(FilePath, BaseDirectory);
+            var DestFile = CombinePaths(NewBaseDirectory, RelativeFile);
+            return DestFile;
+        }
+    }
+    public class Distiller
+    {
+        public string SourceBaseDir;
+        public string DestBaseDir;
+        public string DestBaseDirSymbols;
+        public DateTime DestinationTime;
+        public List<UnrealTargetPlatform> AllLegalPlatforms;
+        public List<string> RejectStrings;
+
+        public Distiller(string InDestBaseDir, string InDestBaseDirSymbols = null, string InSourceBaseDir = null, DateTime? InDestinationTime = null, List<UnrealTargetPlatform> InAllLegalPlatforms = null, bool AllowNoRedist = false)
+        {
+            DestBaseDir = CommandUtils.CombinePaths(InDestBaseDir);
+            SourceBaseDir = InSourceBaseDir;
+            if (SourceBaseDir == null)
+            {
+                SourceBaseDir = CommandUtils.CmdEnv.LocalRoot;
+            }
+            SourceBaseDir = CommandUtils.CombinePaths(SourceBaseDir);
+            if (InDestinationTime == null)
+            {
+                DestinationTime = DateTime.Now;
+            }
+            else
+            {
+                DestinationTime = InDestinationTime.Value;
+            }
+            DestBaseDirSymbols = InDestBaseDirSymbols;
+            if (DestBaseDirSymbols != null)
+            {
+                DestBaseDirSymbols = CommandUtils.CombinePaths(DestBaseDirSymbols);
+            }
+            AllLegalPlatforms = InAllLegalPlatforms;
+            if (AllLegalPlatforms == null)
+            {
+                AllLegalPlatforms = new List<UnrealTargetPlatform>();
+                foreach (UnrealTargetPlatform Plat in Enum.GetValues(typeof(UnrealTargetPlatform)))
+                {
+                    AllLegalPlatforms.Add(Plat);
+                }
+            }
+            RejectStrings = new List<string>
+            {
+                CommandUtils.CombinePaths("/CarefullyRedist/"),
+                CommandUtils.CombinePaths("/NotForLicensees/"),
+            };
+            if (!AllowNoRedist)
+            {
+                RejectStrings.Add(CommandUtils.CombinePaths("/NoRedist/"));
+            }
+            foreach (UnrealTargetPlatform Plat in Enum.GetValues(typeof(UnrealTargetPlatform)))
+            {
+                if (!AllLegalPlatforms.Contains(Plat))
+                {
+                    RejectStrings.Add(CommandUtils.CombinePaths("/" + Plat.ToString() + "/"));
+                }
+            }
+            if (!AllLegalPlatforms.Contains(UnrealTargetPlatform.Win64) && !AllLegalPlatforms.Contains(UnrealTargetPlatform.Win32))
+            {
+                RejectStrings.Add(CommandUtils.CombinePaths("/Windows/"));
+            }
+        }
+
+        public static bool RejectFileOnSubstring(string FileToCopy, string RejectString, string Root = null)
+        {
+            var Search = CommandUtils.CombinePaths(FileToCopy);
+            if (Root == null)
+            {
+                Root = CommandUtils.CmdEnv.LocalRoot;
+            }
+            if (Search.StartsWith(CommandUtils.CombinePaths(Root), StringComparison.InvariantCultureIgnoreCase))
+            {
+                Search = CommandUtils.CombinePaths("/", CommandUtils.StripBaseDirectory(FileToCopy, Root)); // lots of rejects are like "/path/" so we need to be sure we start with a slash
+            }
+            if (Search.IndexOf(CommandUtils.CombinePaths(RejectString), 0, StringComparison.InvariantCultureIgnoreCase) >= 0)
+            {
+                return true;
+            }
+            return false;
+        }
+
+        public bool RejectFile(string FileToCopy, string Root = null)
+        {
+            bool Result = false;
+            foreach (var RejectOn in RejectStrings)
+            {
+                if (RejectFileOnSubstring(FileToCopy, RejectOn, Root == null ? SourceBaseDir : Root))
+                {
+                    Result = true;
+                    break;
+                }
+            }
+            if (Result)
+            {
+                CommandUtils.Log("Rejecting file {0}", FileToCopy);
+            }
+            return Result;
+        }
+
+        public static bool IsSymbolFile(UnrealTargetPlatform ForPlatform, string FileToCopy)
+        {
+            var DebugExts = Platform.Platforms[ForPlatform].GetDebugFileExtentions();
+            foreach (var DebugExt in DebugExts)
+            {
+                if (Path.GetExtension(FileToCopy).Equals(DebugExt, StringComparison.InvariantCultureIgnoreCase))
+                {
+                    return true;
+                }
+            }
+            return false;
+        }
+
+        public bool IsSymbolFile(string FileToCopy)
+        {
+            foreach (var ForPlatform in AllLegalPlatforms)
+            {
+                if (IsSymbolFile(ForPlatform, FileToCopy))
+                {
+                    return true;
+                }
+            }
+            return false;
+        }
+
+        public string CopyFileToDest(string FileToCopy, bool MoveSymbols = true)
+        {
+            if (!CommandUtils.FileExists_NoExceptions(FileToCopy))
+            {
+                throw new AutomationException("Can't distill {0}; it doesn't exist", FileToCopy);
+            }
+            {
+                FileInfo FileInfo = new FileInfo(FileToCopy);
+                FileToCopy = CommandUtils.CombinePaths(FileInfo.FullName);
+            }
+
+            string Dest;
+            if (MoveSymbols && !String.IsNullOrEmpty(DestBaseDirSymbols) && IsSymbolFile(FileToCopy))
+            {
+                Dest = CommandUtils.MakeRerootedFilePath(FileToCopy, SourceBaseDir, DestBaseDirSymbols);
+            }
+            else
+            {
+                Dest = CommandUtils.MakeRerootedFilePath(FileToCopy, SourceBaseDir, DestBaseDir);
+            }
+            if (CommandUtils.FileExists_NoExceptions(Dest))
+            {
+                throw new AutomationException("Can't distill to {0}; it already exists", Dest);
+            }
+            CommandUtils.CopyFile(FileToCopy, Dest);
+
+            {
+                FileInfo FileInfo = new FileInfo(Dest);
+                FileInfo.IsReadOnly = false;
+                FileInfo.CreationTime = DestinationTime;
+                FileInfo.LastWriteTime = DestinationTime;
+                Dest = CommandUtils.CombinePaths(FileInfo.FullName);
+            }
+            if (!CommandUtils.FileExists_NoExceptions(Dest))
+            {
+                throw new AutomationException("Couldn't distill {0} to {1}", FileToCopy, Dest);
+            }
+            return Dest;
+        }
+
+        public List<string> Distill(string FullPathAndWildcard, bool bRecursive = false, bool bAllowMissing = false, bool MoveSymbols = true, params string[] Exclusions)
+        {
+            var FilesToCopy = new List<string>();
+
+            string AbsFile = CommandUtils.CombinePaths(FullPathAndWildcard);
+            CommandUtils.Log("Distilling {0}", AbsFile);
+            string PathOnly = Path.GetDirectoryName(AbsFile);
+            if (!CommandUtils.DirectoryExists_NoExceptions(PathOnly))
+            {
+                if (bAllowMissing)
+                {
+                    return FilesToCopy;
+                }
+                throw new AutomationException("Path {0} was used to distill {1} but it doesn't exist.", PathOnly, AbsFile);
+            }
+            var Files = CommandUtils.FindFiles(Path.GetFileName(AbsFile), bRecursive, PathOnly);
+            var Exclude = new List<string>();
+            foreach (var Excl in Exclusions)
+            {
+                if (!Excl.Contains("/") && !Excl.Contains("\\"))
+                {
+                    var ExFiles = CommandUtils.FindFiles(Excl, bRecursive, PathOnly);
+                    Exclude.AddRange(ExFiles);
+                }
+            }
+
+            foreach (var FileToCopy in Files)
+            {
+                bool bOk = !Exclude.Contains(FileToCopy) && !RejectFile(FileToCopy, PathOnly);
+                foreach (var Excl in Exclusions)
+                {
+                    if (Excl.Contains("/") || Excl.Contains("\\"))
+                    {
+                        if (Excl.Contains("*") || Excl.Contains("?"))
+                        {
+                            throw new AutomationException("Exclusion {0} is illegal, must either be a substring or a wildcard without a path", Excl);
+                        }
+                        bOk = bOk && !RejectFileOnSubstring(FileToCopy, Excl, PathOnly);
+                    }
+                }
+                if (bOk)
+                {
+                    FilesToCopy.Add(CopyFileToDest(FileToCopy, MoveSymbols));
+                }
+            }
+            if (FilesToCopy.Count < 1 && !bAllowMissing)
+            {
+                throw new AutomationException("Distill {0} did not produce any files. Recursive option was {1}.", AbsFile, bRecursive.ToString());
+            }
+            return FilesToCopy;
+        }
+    }
+}
diff --git a/Engine/Source/Programs/AutomationTool/AutomationUtils/HelpUtils.cs b/Engine/Source/Programs/AutomationTool/AutomationUtils/HelpUtils.cs
new file mode 100644
--- /dev/null
+++ b/Engine/Source/Programs/AutomationTool/AutomationUtils/HelpUtils.cs
@@ -0,0 +1,369 @@
+﻿// Copyright 1998-2015 Epic Games, Inc. All Rights Reserved.
+using System;
+using System.Collections;
+using System.Collections.Generic;
+using System.Diagnostics;
+using System.IO;
+using System.Net;
+using System.Runtime.CompilerServices;
+using System.Runtime.Serialization;
+using System.Runtime.Serialization.Json;
+using System.Text;
+using System.Threading;
+using Tools.DotNETCommon.XmlHandler;
+
+using System.ComponentModel;
+using System.Reflection;
+
+namespace AutomationTool
+{
+	#region Help
+
+	/// <summary>
+	/// Help Attribute.
+	/// </summary>
+	[AttributeUsage(AttributeTargets.All, AllowMultiple = true)]
+	public class HelpAttribute : DescriptionAttribute
+	{
+		private bool IsParamHelp;
+
+		/// <summary>
+		/// Basic constructor for class descriptions.
+		/// </summary>
+		/// <param name="Description">Class description</param>
+		public HelpAttribute(string Description)
+			: base(Description)
+		{
+			IsParamHelp = false;
+		}
+
+		/// <summary>
+		/// Constructor for parameter descriptions.
+		/// </summary>
+		/// <param name="ParamName">Paramter name</param>
+		/// <param name="Description">Paramter description</param>
+		public HelpAttribute(string ParamName, string Description)
+			: base(FormatDescription(ParamName, Description))
+		{
+			IsParamHelp = true;
+		}
+
+		/// <summary>
+		/// Additional type to display help for.
+		/// </summary>
+		/// <param name="AdditionalType"></param>
+		public HelpAttribute(Type AdditionalType)
+		{
+			IsParamHelp = false;
+			AdditionalHelp = AdditionalType;
+		}
+
+		private static string FormatDescription(string Name, string Description)
+		{
+			return String.Format("-{0} {1}", Name, Description);
+		}
+
+		public bool IsParam
+		{
+			get { return IsParamHelp; }
+		}
+
+		public Type AdditionalHelp
+		{
+			get;
+			private set;
+		}
+	}
+
+	#endregion
+
+	/// <summary>
+	/// Base utility function for script commands.
+	/// </summary>
+	public partial class CommandUtils
+	{
+		#region Help
+
+		/// <summary>
+		/// Displays help for the specified command.
+		/// <param name="Command">Command type.</param>
+		/// </summary>
+		public static void Help(Type Command)
+		{
+			string Description;
+			List<string> Params;
+			GetTypeHelp(Command, out Description, out Params);
+			LogHelp(Command, Description, Params);
+		}
+
+		/// <summary>
+		/// Displays help for the specified Type.
+		/// </summary>
+		/// <param name="Command">Type to display help for.</param>
+		public static void LogHelp(Type Command)
+		{
+			string Description;
+			List<string> Params;
+			GetTypeHelp(Command, out Description, out Params);
+			LogHelp(Command, Description, Params);
+		}
+
+		/// <summary>
+		/// Displays a formatted help for the specified Command type.
+		/// </summary>
+		/// <param name="Command">Command Type</param>
+		/// <param name="Description">Command Desscription</param>
+		/// <param name="Params">Command Parameters</param>
+		public static void LogHelp(Type Command, string Description, List<string> Params)
+		{
+			string HelpMessage = Environment.NewLine;
+
+			bool bFirstLine = true;
+			if (Command != null)
+			{
+				HelpMessage += String.Format("{0} Help:{1}", Command.Name, Environment.NewLine);
+				bFirstLine = false;
+			}
+
+			if (String.IsNullOrEmpty(Description) == false)
+			{
+				if (!bFirstLine)
+				{
+					HelpMessage += Environment.NewLine;
+				}
+				HelpMessage += Description + Environment.NewLine;
+				bFirstLine = false;
+			}
+			
+			if (IsNullOrEmpty(Params) == false)
+			{
+				if (!bFirstLine)
+				{
+					HelpMessage += Environment.NewLine;
+				}
+				HelpMessage += "Parameters: " + Environment.NewLine;
+				HelpMessage += Environment.NewLine;
+
+				HelpMessage += FormatParams(Params, 4, 24);
+			}
+            
+			Log(HelpMessage);
+		}
+
+		/// <summary>
+		/// Formats the given parameters as so:
+		///     -Param1     Param1 Description
+		///
+		///     -Param2      Param2 Description, this description is
+		///                  longer and splits onto a separate line. 
+		///
+		///     -Param3      Param3 Description continues as before. 
+		/// </summary>
+		/// <param name="Params">List of parameters arranged as "-ParamName Param Description"</param>
+		/// <param name="Indent">Indent from the left hand side</param>
+		/// <param name="DefaultRightPadding">The minimum padding from the start of the param name to the start of the description (resizes with larger param names)</param>
+		/// <returns></returns>
+		static string FormatParams(List<string> Params, int Indent, int DefaultRightPadding)
+		{
+			Dictionary<string, string> ParamDict = new Dictionary<string, string>(StringComparer.InvariantCultureIgnoreCase);
+
+			// Extract Params/Descriptions into Key/Value pairs
+			foreach (var Param in Params)
+			{
+				// Find the first space (should be following the param name)
+				if (!String.IsNullOrWhiteSpace(Param))
+				{
+					var ParamName = String.Empty;
+					var ParamDesc = String.Empty;
+
+					var SplitPoint = Param.IndexOf(' ');
+					if (SplitPoint > 0)
+					{
+						// Extract the name and description seperately
+						ParamName = Param.Substring(0, SplitPoint);
+						ParamDesc = Param.Substring(SplitPoint + 1, Param.Length - (SplitPoint + 1));
+					}
+					else
+					{
+						ParamName = Param;
+					}
+
+					// build dictionary using Name and Desc as Key and Value
+					if (!ParamDict.ContainsKey(ParamName))
+					{
+						ParamDict.Add(ParamName, ParamDesc);
+					}
+					else
+					{
+						LogWarning("Duplicated help parameter \"{0}\"", ParamName);
+					}
+				}
+			}
+
+			// string used to intent the param
+			string IndentString = string.Empty.PadRight(Indent);
+
+			// default the padding value
+			int RightPadding = DefaultRightPadding;
+
+			// If (Padding < longest param name) padding = longest param name + 1
+			foreach (var ParamName in ParamDict.Keys)
+			{
+				if (ParamName.Length + 1 > RightPadding)
+				{
+					RightPadding = ParamName.Length + 1;
+				}
+			}
+
+			// results string
+			string FormattedParams = string.Empty;
+
+			// Build the formatted params
+			foreach (var ParamName in ParamDict.Keys)
+			{
+				// build the param first, including intend and padding on the rights size
+				string ParamString = IndentString + ParamName.PadRight(RightPadding);
+
+				// Build the description line by line, adding the same amount of intending each time. 
+				foreach (var DescriptionLine in WordWrap(ParamDict[ParamName], Console.WindowWidth - ParamString.Length))
+				{
+					// Formatting as following:
+					// <Indent>-param<Right Padding>Description<New line>
+					FormattedParams += ParamString + DescriptionLine + Environment.NewLine;
+
+					// we replace the param string on subsequent lines with white space of the same length
+					ParamString = string.Empty.PadRight(IndentString.Length + RightPadding);
+				}
+
+				// new line after each param/desc combo
+				FormattedParams += Environment.NewLine;
+			}
+
+			return FormattedParams;
+		}
+
+		/// <summary>
+		/// Takes a given sentence and wraps it on a word by word basis so that no line exceeds the
+		/// set maximum line length. Words longer than a line are broken up. Returns the sentence as
+		/// a list of individual lines.
+		/// </summary>
+		/// <param name="sentence">The sentence to be wrapped</param>
+		/// <param name="MaximumLineLength">The maximum (non negative) length of the returned sentences</param>
+		/// <returns>The provided sentence in a list of individual lines no longer than MaximumLineLength</returns>
+		static List<string> WordWrap(string Sentence, int MaximumLineLength)
+		{
+			// Early out
+			if (Sentence.Length == 0)
+			{
+				return new List<string>();
+			}
+
+			string[] words = Sentence.Split(' ');
+			List<string> WrappedWords = new List<string>();
+
+			string CurrentSentence = string.Empty;
+			foreach (var word in words)
+			{
+				// if this is a very large word, split it
+				if (word.Length > MaximumLineLength)
+				{
+					// Top up the current line
+					WrappedWords.Add(CurrentSentence + word.Substring(0, MaximumLineLength - CurrentSentence.Length));
+
+					int length = MaximumLineLength - CurrentSentence.Length;
+					while (length + MaximumLineLength < word.Length)
+					{
+						// Place the starting lengths into their own lines
+						WrappedWords.Add(word.Substring(length, Math.Min(MaximumLineLength, word.Length - length)));
+						length += MaximumLineLength;
+					}
+
+					// then the trailing end into the next line
+					CurrentSentence += word.Substring(length, Math.Min(MaximumLineLength, word.Length - length)) + " ";
+				}
+				else
+				{
+					if (CurrentSentence.Length + word.Length > MaximumLineLength)
+					{
+						// next line and reset sentence
+						WrappedWords.Add(CurrentSentence);
+						CurrentSentence = string.Empty;
+					}
+
+					// Add the word to the current sentence.
+					CurrentSentence += word + " ";
+				}
+			}
+
+			if (CurrentSentence.Length > 0)
+			{
+				WrappedWords.Add(CurrentSentence);
+			}
+
+			return WrappedWords;
+		}
+
+		/// <summary>
+		/// Gets the description and a list of parameters for the specified type.
+		/// </summary>
+		/// <param name="ObjType">Type to get help for.</param>
+		/// <param name="ObjectDescription">Description</param>
+		/// <param name="Params">List of paramters</param>
+		public static void GetTypeHelp(Type ObjType, out string ObjectDescription, out List<string> Params)
+		{
+			ObjectDescription = String.Empty;
+			Params = new List<string>();
+
+			var AllAttributes = ObjType.GetCustomAttributes(false);			
+			foreach (var CustomAttribute in AllAttributes)
+			{
+				var HelpAttribute = CustomAttribute as HelpAttribute;
+				if (HelpAttribute != null)
+				{
+					if (HelpAttribute.IsParam)
+					{
+						Params.Add(HelpAttribute.Description);
+					}
+					else if (HelpAttribute.AdditionalHelp != null)
+					{
+						string DummyDescription;
+						List<string> AdditionalParams;
+						GetTypeHelp(HelpAttribute.AdditionalHelp, out DummyDescription, out AdditionalParams);
+						Params.AddRange(AdditionalParams);
+					}
+					else
+					{
+						ObjectDescription = HelpAttribute.Description;
+					}
+				}
+			}
+
+			var AllMembers = GetAllMemebers(ObjType);
+			foreach (var Member in AllMembers)
+			{
+				var AllMemeberAtributes = Member.GetCustomAttributes(false);
+				foreach (var CustomAttribute in AllMemeberAtributes)
+				{
+					var HelpAttribute = CustomAttribute as HelpAttribute;
+					if (HelpAttribute != null && HelpAttribute.IsParam)
+					{
+						Params.Add(HelpAttribute.Description);
+					}
+				}
+			}
+		}
+
+		private static List<MemberInfo> GetAllMemebers(Type ObjType)
+		{
+			var Members = new List<MemberInfo>();
+			for (; ObjType != null; ObjType = ObjType.BaseType)
+			{
+				var ObjMembers = ObjType.GetMembers(BindingFlags.Public | BindingFlags.NonPublic | BindingFlags.Instance | BindingFlags.Static);
+				Members.AddRange(ObjMembers);
+			}
+			return Members;
+		}
+
+		#endregion
+	}
+}
diff --git a/Engine/Source/Programs/AutomationTool/AutomationUtils/HostPlatform.cs b/Engine/Source/Programs/AutomationTool/AutomationUtils/HostPlatform.cs
new file mode 100644
--- /dev/null
+++ b/Engine/Source/Programs/AutomationTool/AutomationUtils/HostPlatform.cs
@@ -0,0 +1,155 @@
+﻿// Copyright 1998-2015 Epic Games, Inc. All Rights Reserved.
+
+using System;
+using System.Collections.Generic;
+using System.IO;
+using System.Linq;
+using System.Text;
+using System.Diagnostics;
+
+namespace AutomationTool
+{
+	/// <summary>
+	/// Host platform abstraction
+	/// </summary>
+	public abstract class HostPlatform
+	{
+		private static HostPlatform RunningPlatform;
+		/// <summary>
+		/// Current running host platform.
+		/// </summary>
+		public static HostPlatform Current
+		{
+			get 
+			{ 
+				if (RunningPlatform == null)
+				{
+					throw new AutomationException("UnrealAutomationTool host platform not initialized.");
+				}
+				return RunningPlatform; 
+			}
+		}
+
+		/// <summary>
+		/// Initializes the current platform.
+		/// </summary>
+		public static void Initialize()
+		{
+			PlatformID Platform = Environment.OSVersion.Platform;
+			switch (Platform)
+			{
+				case PlatformID.Win32NT:
+				case PlatformID.Win32S:
+				case PlatformID.Win32Windows:
+				case PlatformID.WinCE:
+					RunningPlatform = new WindowsHostPlatform();
+					break;
+
+				case PlatformID.Unix:
+					if (File.Exists ("/System/Library/CoreServices/SystemVersion.plist"))
+					{
+						RunningPlatform = new MacHostPlatform();
+					} 
+					else 
+					{
+						RunningPlatform = new LinuxHostPlatform();
+					}
+					break;
+
+				case PlatformID.MacOSX:
+					RunningPlatform = new MacHostPlatform();
+					break;
+
+				default:
+					throw new Exception ("Unhandled runtime platform " + Platform);
+			}
+		}
+
+		/// <summary>
+		/// Sets .Net framework environment variables.
+		/// </summary>
+		abstract public void SetFrameworkVars();
+
+		/// <summary>
+		/// Gets the build executable filename.
+		/// </summary>
+		/// <returns></returns>
+		abstract public string GetMsBuildExe();
+
+		/// <summary>
+		/// Gets the dev environment executable name.
+		/// </summary>
+		/// <returns></returns>
+		abstract public string GetMsDevExe();
+
+		/// <summary>
+		/// Folder under UE4/ to the platform's binaries.
+		/// </summary>
+		abstract public string RelativeBinariesFolder { get; }
+
+		/// <summary>
+		/// Full path to the UE4 Editor executable for the current platform.
+		/// </summary>
+		/// <param name="UE4Exe"></param>
+		/// <returns></returns>
+		abstract public string GetUE4ExePath(string UE4Exe);
+
+		/// <summary>
+		/// Log folder for local builds.
+		/// </summary>
+		abstract public string LocalBuildsLogFolder { get; }
+
+		/// <summary>
+		/// Name of the p4 executable.
+		/// </summary>
+		abstract public string P4Exe { get; }
+
+		/// <summary>
+		/// Creates a process and sets it up for the current platform.
+		/// </summary>
+		/// <param name="LogName"></param>
+		/// <returns></returns>
+		abstract public Process CreateProcess(string AppName);
+
+		/// <summary>
+		/// Sets any additional options for running an executable.
+		/// </summary>
+		/// <param name="AppName"></param>
+		/// <param name="Options"></param>
+		/// <param name="CommandLine"></param>
+		abstract public void SetupOptionsForRun(ref string AppName, ref CommandUtils.ERunOptions Options, ref string CommandLine);
+
+		/// <summary>
+		/// Sets the console control handler for the current platform.
+		/// </summary>
+		/// <param name="Handler"></param>
+		abstract public void SetConsoleCtrlHandler(ProcessManager.CtrlHandlerDelegate Handler);
+
+		/// <summary>
+		/// Platform specific override to skip loading/compiling unsupported modules
+		/// </summary>
+		/// <param name="ModuleName">Module name</param>
+		/// <returns>True if module should be compiled or loaded</returns>
+		abstract public bool IsScriptModuleSupported(string ModuleName);
+
+		/// <summary>
+		/// Gets UBT project name for the current platform.
+		/// </summary>
+		abstract public string UBTProjectName { get; }
+
+		/// <summary>
+		/// Returns the type of the host editor platform.
+		/// </summary>
+		abstract public UnrealBuildTool.UnrealTargetPlatform HostEditorPlatform { get; }
+
+		/// <summary>
+		/// Returns the pdb file extenstion for the host platform.
+		/// </summary>
+		abstract public string PdbExtension { get; }
+
+		/// <summary>
+		/// List of processes that can't not be killed
+		/// </summary>
+		abstract public string[] DontKillProcessList { get; }
+	}
+}
diff --git a/Engine/Source/Programs/AutomationTool/AutomationUtils/LocalBuildEnvironment.cs b/Engine/Source/Programs/AutomationTool/AutomationUtils/LocalBuildEnvironment.cs
new file mode 100644
--- /dev/null
+++ b/Engine/Source/Programs/AutomationTool/AutomationUtils/LocalBuildEnvironment.cs
@@ -0,0 +1,160 @@
+// Copyright 1998-2015 Epic Games, Inc. All Rights Reserved.
+using System;
+using System.Collections.Generic;
+using System.Linq;
+using System.Text;
+using System.IO;
+using System.Reflection;
+using System.Diagnostics;
+using UnrealBuildTool;
+
+namespace AutomationTool
+{
+	/// <summary>
+	/// Local Build Environment. Used when not running on build machine.
+	/// </summary>
+	class LocalCommandEnvironment : CommandEnvironment
+	{
+		internal LocalCommandEnvironment()
+			: base()
+		{
+		}
+
+		/// <summary>
+		/// Initializes the environment. Tries to autodetect all source control settings.
+		/// </summary>
+		/// <param name="CompilationEnv">Compilation environment</param>
+		protected override void InitEnvironment()
+		{
+			SetUATLocation();
+			SetLocalRootPath();
+			SetUATSavedPath();
+			SetAndCreateLogFolder();
+
+			base.InitEnvironment();
+		}
+
+		private void SetLocalRootPath()
+		{
+			var LocalRootPath = CommandUtils.ConvertSeparators(PathSeparator.Slash, RootFromUATLocation(UATExe));
+			CommandUtils.ConditionallySetEnvVar(EnvVarNames.LocalRoot, LocalRootPath);
+		}
+
+		/// <summary>
+		/// Finds the root path of the branch by looking in progressively higher folder locations until a file known to exist in the branch is found.
+		/// </summary>
+		/// <param name="UATLocation">Location of the currently executing assembly.</param>
+		/// <returns></returns>
+		private static string RootFromUATLocation(string UATLocation)
+		{
+			// pick a known file in the depot and try to build a relative path to it. This will tell us where the root path to the branch is.
+			var KnownFilePathFromRoot = CommandEnvironment.KnownFileRelativeToRoot;
+			var CurrentPath = CommandUtils.ConvertSeparators(PathSeparator.Slash, Path.GetDirectoryName(UATLocation));
+			var UATPathRoot = CommandUtils.CombinePaths(Path.GetPathRoot(CurrentPath), KnownFilePathFromRoot);
+
+			// walk parent dirs until we find the file or hit the path root, which means we aren't in a known location.
+			while (true)
+			{
+				var PathToCheck = Path.GetFullPath(CommandUtils.CombinePaths(CurrentPath, KnownFilePathFromRoot));
+				Log.TraceVerbose("Checking for {0}", PathToCheck);
+				if (!File.Exists(PathToCheck))
+				{
+					var LastSeparatorIndex = CurrentPath.LastIndexOf('/');
+					if (String.Compare(PathToCheck, UATPathRoot, true) == 0 || LastSeparatorIndex < 0)
+					{
+						throw new AutomationException("{0} does not appear to exist at a valid location with the branch, so the local root cannot be determined.", UATLocation);
+					}
+					
+					// Step back one directory
+					CurrentPath = CurrentPath.Substring(0, LastSeparatorIndex);
+				}
+				else
+				{
+					return CurrentPath;
+				}
+			}
+		}
+
+		/// <summary>
+		/// Gets the log folder used when running from installed build.
+		/// </summary>
+		/// <returns></returns>
+		private string GetInstalledLogFolder()
+		{
+			var LocalRootPath = CommandUtils.GetEnvVar(EnvVarNames.LocalRoot);
+			var LocalRootEscaped = CommandUtils.ConvertSeparators(PathSeparator.Slash, LocalRootPath.Replace(":", ""));
+			if (LocalRootEscaped[LocalRootEscaped.Length - 1] == '/')
+			{
+				LocalRootEscaped = LocalRootEscaped.Substring(0, LocalRootEscaped.Length - 1);
+			}
+			LocalRootEscaped = CommandUtils.EscapePath(LocalRootEscaped);
+			return CommandUtils.CombinePaths(PathSeparator.Slash, HostPlatform.Current.LocalBuildsLogFolder, LocalRootEscaped);
+		}
+
+		/// <summary>
+		/// Sets up log output folder.
+		/// </summary>
+		private void SetAndCreateLogFolder()
+		{			
+			var LocalLogFolder = CommandUtils.GetEnvVar(EnvVarNames.LogFolder);
+			if (String.IsNullOrEmpty(LocalLogFolder))
+			{
+				string RootSavedPath;
+				if (!GlobalCommandLine.Installed)
+				{					
+					RootSavedPath = CommandUtils.CombinePaths(PathSeparator.Slash, CommandUtils.GetEnvVar(EnvVarNames.EngineSavedFolder), "Logs");
+				}
+				else
+				{
+					RootSavedPath = GetInstalledLogFolder(); 
+				}
+				if (TryCreateLogFolder(RootSavedPath) == false)
+				{
+					var Result = false;
+					if (!GlobalCommandLine.Installed)
+					{
+						RootSavedPath = GetInstalledLogFolder();
+						Result = TryCreateLogFolder(RootSavedPath);
+					}
+					if (!Result)
+					{
+						throw new AutomationException("Unable to find a usable log folder.");
+					}
+				}				
+			}
+		}
+
+		private bool TryCreateLogFolder(string PossibleLogFolder)
+		{
+			bool Result = true;
+			try
+			{
+				if (Directory.Exists(PossibleLogFolder))
+				{
+					CommandUtils.DeleteDirectoryContents(PossibleLogFolder);
+
+					// Since the directory existed and might have been empty, we need to make sure
+					// we can actually write to it, so create and delete a temporary file.
+					var RandomFilename = Path.GetRandomFileName();
+					var RandomFilePath = CommandUtils.CombinePaths(PossibleLogFolder, RandomFilename);
+					File.WriteAllText(RandomFilePath, RandomFilename);
+					File.Delete(RandomFilePath);
+				}
+				else
+				{
+					Directory.CreateDirectory(PossibleLogFolder);
+				}
+			}
+			catch (UnauthorizedAccessException Ex)
+			{
+				Log.TraceInformation("No permission to use LogFolder={0} ({1})", PossibleLogFolder, Ex.Message);
+				Result = false;
+			}
+			if (Result)
+			{
+				Environment.SetEnvironmentVariable(EnvVarNames.LogFolder, PossibleLogFolder);
+			}
+			return Result;
+		}
+	}
+}
diff --git a/Engine/Source/Programs/AutomationTool/AutomationUtils/LocalP4Environment.cs b/Engine/Source/Programs/AutomationTool/AutomationUtils/LocalP4Environment.cs
new file mode 100644
--- /dev/null
+++ b/Engine/Source/Programs/AutomationTool/AutomationUtils/LocalP4Environment.cs
@@ -0,0 +1,311 @@
+﻿// Copyright 1998-2015 Epic Games, Inc. All Rights Reserved.
+
+using System;
+using System.Collections.Generic;
+using System.Linq;
+using System.Text;
+using UnrealBuildTool;
+
+namespace AutomationTool
+{
+	/// <summary>
+	/// Local P4 environment settings
+	/// </summary>
+	class LocalP4Environment : P4Environment
+	{
+		internal LocalP4Environment(P4Connection Connection, CommandEnvironment CmdEnv)
+			: base(Connection, CmdEnv)
+		{
+		}
+
+		/// <summary>
+		/// Initializes the environment. Tries to autodetect all source control settings.
+		/// </summary>
+		/// <param name="CompilationEnv">Compilation environment</param>
+		protected override void InitEnvironment(P4Connection Connection, CommandEnvironment CmdEnv)
+		{
+			var HostName = Environment.MachineName.ToLower();
+			var P4PortEnv = Environment.GetEnvironmentVariable("P4PORT");
+			if (String.IsNullOrEmpty(P4PortEnv))
+			{
+				P4PortEnv = DetectP4Port();
+			}
+
+			var UserName = CommandUtils.GetEnvVar(EnvVarNames.User);
+			if (String.IsNullOrEmpty(UserName))
+			{
+				UserName = DetectUserName(Connection);
+			}
+
+			var CommandLineClient = CommandUtils.GetEnvVar(EnvVarNames.Client);
+			P4ClientInfo ThisClient = null;
+			if (String.IsNullOrEmpty(CommandLineClient) == false)
+			{
+				ThisClient = Connection.GetClientInfo(CommandLineClient);
+				if (ThisClient == null)
+				{
+					throw new AutomationException("Unable to find client {0}", CommandLineClient);
+				}
+				if (String.Compare(ThisClient.Owner, UserName, true) != 0)
+				{
+					throw new AutomationException("Client specified with {0}={1} has a different owner then the detected user name (has: {2}, expected: {3})",
+						EnvVarNames.Client, CommandLineClient, ThisClient.Owner, UserName);
+				}
+			}
+			else
+			{
+				ThisClient = DetectClient(Connection, UserName, HostName, CmdEnv.UATExe);
+			}
+
+			Log.TraceInformation("Using user {0} clientspec {1} {2}", UserName, ThisClient.Name, ThisClient.RootPath);
+			Environment.SetEnvironmentVariable("P4CLIENT", ThisClient.Name);
+
+			string BuildRootPath;
+			string ClientRootPath;
+			DetectRootPaths(Connection, CmdEnv.LocalRoot, ThisClient, out BuildRootPath, out ClientRootPath);
+
+			CommandUtils.ConditionallySetEnvVar(EnvVarNames.P4Port, P4PortEnv);
+			CommandUtils.ConditionallySetEnvVar(EnvVarNames.User, UserName);
+			CommandUtils.ConditionallySetEnvVar(EnvVarNames.Client, ThisClient.Name);
+			CommandUtils.ConditionallySetEnvVar(EnvVarNames.BuildRootP4, BuildRootPath);
+			CommandUtils.ConditionallySetEnvVar(EnvVarNames.ClientRoot, ClientRootPath);
+
+			var CLString = CommandUtils.GetEnvVar(EnvVarNames.Changelist, null);
+			if (String.IsNullOrEmpty(CLString) && CommandUtils.P4CLRequired)
+			{
+                CLString = DetectCurrentCL(Connection, ClientRootPath);
+			}
+			if (!String.IsNullOrEmpty(CLString))
+			{
+				CommandUtils.ConditionallySetEnvVar(EnvVarNames.Changelist, CLString);
+			}
+
+			CommandUtils.ConditionallySetEnvVar(EnvVarNames.LabelToSync, "");
+			CommandUtils.ConditionallySetEnvVar("P4USER", UserName);
+			CommandUtils.ConditionallySetEnvVar("P4CLIENT", ThisClient.Name);
+
+			var P4Password = Environment.GetEnvironmentVariable(EnvVarNames.P4Password);
+			if (!String.IsNullOrEmpty(P4Password))
+			{
+				CommandUtils.ConditionallySetEnvVar("P4PASSWD", P4Password);
+			}
+
+			SetBuildRootEscaped();
+
+			base.InitEnvironment(Connection, CmdEnv);
+		}
+
+		/// <summary>
+		/// Sets the escaped build root environment variable. If the build root is not set, UAT's location UE4 root will be used.
+		/// </summary>
+		private void SetBuildRootEscaped()
+		{
+			var BuildRoot = CommandUtils.GetEnvVar(EnvVarNames.BuildRootP4);
+			if (String.IsNullOrEmpty(BuildRoot))
+			{
+				throw new AutomationException("Build root is empty");
+			}
+			BuildRoot = CommandUtils.EscapePath(BuildRoot);
+			CommandUtils.ConditionallySetEnvVar(EnvVarNames.BuildRootEscaped, BuildRoot);
+		}
+
+		/// <summary>
+		/// Detects the current changelist the workspace is synced to.
+		/// </summary>
+		/// <param name="ClientRootPath">Workspace path.</param>
+		/// <returns>Changelist number as a string.</returns>
+		private static string DetectCurrentCL(P4Connection Connection, string ClientRootPath)
+		{
+			CommandUtils.Log("uebp_CL not set, detecting 'have' CL...");
+
+			// Retrieve the current changelist 
+			var P4Result = Connection.P4("changes -m 1 " + CommandUtils.CombinePaths(PathSeparator.Depot, ClientRootPath, "/...#have"), AllowSpew: false);
+			var CLTokens = P4Result.Output.Split(new char[] { ' ' }, StringSplitOptions.RemoveEmptyEntries);
+			var CLString = CLTokens[1];
+			var CL = Int32.Parse(CLString);
+			if (CLString != CL.ToString())
+			{
+				throw new AutomationException("Failed to retrieve current changelist.");
+			}
+			return CLString;
+		}
+
+		/// <summary>
+		/// Detects root paths for the specified client.
+		/// </summary>
+		/// <param name="UATLocation">AutomationTool.exe location</param>
+		/// <param name="ThisClient">Client to detect the root paths for</param>
+		/// <param name="BuildRootPath">Build root</param>
+		/// <param name="LocalRootPath">Local root</param>
+		/// <param name="ClientRootPath">Client root</param>
+		private static void DetectRootPaths(P4Connection Connection, string LocalRootPath, P4ClientInfo ThisClient, out string BuildRootPath, out string ClientRootPath)
+		{
+			// Figure out the build root
+			string KnownFilePathFromRoot = CommandEnvironment.KnownFileRelativeToRoot;
+			string KnownLocalPath = CommandUtils.MakePathSafeToUseWithCommandLine(CommandUtils.CombinePaths(PathSeparator.Slash, LocalRootPath, KnownFilePathFromRoot));
+			ProcessResult P4Result = Connection.P4(String.Format("files -m 1 {0}", KnownLocalPath), AllowSpew: false);
+
+			string KnownFileDepotMapping = P4Result.Output;
+
+			// Get the build root
+			Log.TraceVerbose("Looking for {0} in {1}", KnownFilePathFromRoot, KnownFileDepotMapping);
+			int EndIdx = KnownFileDepotMapping.IndexOf(KnownFilePathFromRoot, StringComparison.CurrentCultureIgnoreCase);
+			if (EndIdx < 0)
+			{
+				EndIdx = KnownFileDepotMapping.IndexOf(CommandUtils.ConvertSeparators(PathSeparator.Slash, KnownFilePathFromRoot), StringComparison.CurrentCultureIgnoreCase);
+			}
+			// Get the root path without the trailing path separator
+			BuildRootPath = KnownFileDepotMapping.Substring(0, EndIdx - 1);
+
+			// Get the client root
+			if (LocalRootPath.StartsWith(CommandUtils.CombinePaths(PathSeparator.Slash, ThisClient.RootPath, "/"), StringComparison.InvariantCultureIgnoreCase) || LocalRootPath == ThisClient.RootPath)
+			{
+				ClientRootPath = CommandUtils.CombinePaths(PathSeparator.Depot, String.Format("//{0}/", ThisClient.Name), LocalRootPath.Substring(ThisClient.RootPath.Length));
+			}
+			else
+			{
+				throw new AutomationException("LocalRootPath ({0}) does not start with the client root path ({1})", LocalRootPath, ThisClient.RootPath);
+			}
+		}
+
+		/// <summary>
+		/// Detects a workspace given the current user name, host name and depot path.
+		/// </summary>
+		/// <param name="UserName">User name</param>
+		/// <param name="HostName">Host</param>
+		/// <param name="UATLocation">Path to UAT exe, this will be checked agains the root path.</param>
+		/// <returns>Client to use.</returns>
+		private static P4ClientInfo DetectClient(P4Connection Connection, string UserName, string HostName, string UATLocation)
+		{
+			CommandUtils.Log("uebp_CLIENT not set, detecting current client...");
+
+			var MatchingClients = new List<P4ClientInfo>();
+			P4ClientInfo[] P4Clients = Connection.GetClientsForUser(UserName, UATLocation);
+			foreach (var Client in P4Clients)
+			{
+				if (!String.IsNullOrEmpty(Client.Host) && String.Compare(Client.Host, HostName, true) != 0)
+				{
+					Log.TraceInformation("Rejecting client because of different Host {0} \"{1}\" != \"{2}\"", Client.Name, Client.Host, HostName);
+					continue;
+				}
+				
+				MatchingClients.Add(Client);
+			}
+
+			P4ClientInfo ClientToUse = null;
+			if (MatchingClients.Count == 0)
+			{
+				throw new AutomationException("No matching clientspecs found!");
+			}
+			else if (MatchingClients.Count == 1)
+			{
+				ClientToUse = MatchingClients[0];
+			}
+			else
+			{
+				// We may have empty host clients here, so pick the first non-empty one if possible
+				foreach (var Client in MatchingClients)
+				{
+					if (!String.IsNullOrEmpty(Client.Host) && String.Compare(Client.Host, HostName, true) == 0)
+					{
+						ClientToUse = Client;
+						break;
+					}
+				}
+				if (ClientToUse == null)
+				{
+					Log.TraceWarning("{0} clients found that match the current host and root path. The most recently accessed client will be used.", MatchingClients.Count);
+					ClientToUse = GetMostRecentClient(MatchingClients);
+				}
+			}
+
+			return ClientToUse;
+		}
+
+		/// <summary>
+		/// Given a list of clients with the same owner and root path, tries to find the most recently accessed one.
+		/// </summary>
+		/// <param name="Clients">List of clients with the same owner and path.</param>
+		/// <returns>The most recent client from the list.</returns>
+		private static P4ClientInfo GetMostRecentClient(List<P4ClientInfo> Clients)
+		{
+			Log.TraceVerbose("Detecting the most recent client.");
+			P4ClientInfo MostRecentClient = null;
+			var MostRecentAccessTime = DateTime.MinValue;
+			foreach (var ClientInfo in Clients)
+			{
+				if (ClientInfo.Access.Ticks > MostRecentAccessTime.Ticks)
+				{
+					MostRecentAccessTime = ClientInfo.Access;
+					MostRecentClient = ClientInfo;
+				}
+			}
+			if (MostRecentClient == null)
+			{
+				throw new AutomationException("Failed to determine the most recent client in {0}", Clients[0].RootPath);
+			}
+			return MostRecentClient;
+		}
+
+		/// <summary>
+		/// Detects current user name.
+		/// </summary>
+		/// <returns></returns>
+		private static string DetectUserName(P4Connection Connection)
+		{
+			var UserName = String.Empty;
+			var P4Result = Connection.P4("info", AllowSpew: false);
+			if (P4Result.ExitCode != 0)
+			{
+				throw new AutomationException("Perforce command failed: {0}. Please make sure your P4PORT or {1} is set properly.", P4Result.Output, EnvVarNames.P4Port);
+			}
+
+			// Retrieve the P4 user name			
+			var Tags = Connection.ParseTaggedP4Output(P4Result.Output);
+			Tags.TryGetValue("User name", out UserName);
+
+			if (String.IsNullOrEmpty(UserName))
+			{
+				UserName = Environment.GetEnvironmentVariable(EnvVarNames.User);
+
+				if (!String.IsNullOrEmpty(UserName))
+				{
+					Log.TraceWarning("Unable to retrieve perforce user name. Trying to fall back to {0} which is set to {1}.", EnvVarNames.User, UserName);
+				}
+				else
+				{
+					throw new AutomationException("Failed to retrieve user name.");
+				}
+			}
+
+			Environment.SetEnvironmentVariable("P4USER", UserName);
+
+			return UserName;
+		}
+
+		/// <summary>
+		/// Attempts to detect source control server address from environment variables.
+		/// </summary>
+		/// <returns>Source control server address.</returns>
+		private static string DetectP4Port()
+		{
+			// Try to read the P4PORT environment and check if it is set correctly
+			var P4PortEnv = Environment.GetEnvironmentVariable(EnvVarNames.P4Port);
+
+			// If not, try to fallback to Mapping.P4Port and set this as P4PORT before continueing
+			if (!String.IsNullOrEmpty(P4PortEnv))
+			{
+				Log.TraceWarning("P4PORT is not set. Falling back to {0} which is set to {1}.", EnvVarNames.P4Port, P4PortEnv);
+			}
+			else
+			{
+				// If that fails as well, we just give it a shot with perforce:1666 and hope that this works
+				Log.TraceWarning("P4PORT is not set. Trying to fallback to perforce:1666");
+				P4PortEnv = "perforce:1666";
+			}
+			Environment.SetEnvironmentVariable("P4PORT", P4PortEnv);
+
+			return P4PortEnv;
+		}
+	}
+}
diff --git a/Engine/Source/Programs/AutomationTool/AutomationUtils/Log.cs b/Engine/Source/Programs/AutomationTool/AutomationUtils/Log.cs
new file mode 100644
--- /dev/null
+++ b/Engine/Source/Programs/AutomationTool/AutomationUtils/Log.cs
@@ -0,0 +1,270 @@
+﻿// Copyright 1998-2015 Epic Games, Inc. All Rights Reserved.
+using System;
+using System.Collections.Generic;
+using System.Linq;
+using System.Text;
+using System.IO;
+using System.Runtime.CompilerServices;
+using System.Diagnostics;
+
+namespace AutomationTool
+{
+	#region LogUtils
+
+
+	public class LogUtils
+	{
+        private static string LogFilename;
+
+		/// <summary>
+		/// Initializes trace logging.
+		/// </summary>
+		/// <param name="CommandLine">Command line.</param>
+		public static void InitLogging(string[] CommandLine)
+		{
+            // ensure UTF8Output flag is respected, since we are initializing logging early in the program.
+            if (CommandLine.Any(Arg => Arg.Equals("-utf8output", StringComparison.InvariantCultureIgnoreCase)))
+            {
+                Console.OutputEncoding = new System.Text.UTF8Encoding(false, false);
+            }
+
+			UnrealBuildTool.Log.InitLogging(
+                bLogTimestamps: CommandUtils.ParseParam(CommandLine, "-Timestamps"),
+                bLogVerbose: CommandUtils.ParseParam(CommandLine, "-Verbose"),
+                bLogSeverity: true,
+                bLogSources: true,
+                bColorConsoleOutput: true,
+                TraceListeners: new TraceListener[]
+                {
+                    new UnrealBuildTool.FilteredConsoleTraceListener(),
+                    // could return null, but InitLogging handles this gracefully.
+                    CreateLogFileListener(out LogFilename),
+                    //@todo - this is only used by GUBP nodes. Ideally we don't waste this 20MB if we are not running GUBP.
+                    new AutomationMemoryLogListener(),
+                });
+        }
+
+        /// <summary>
+        /// Shuts down logging. This is a nothrow function as shutdown errors are ignored.
+        /// </summary>
+        public static void ShutdownLogging()
+        {
+            // from here we can copy the log file to its final resting place
+            try
+            {
+                // This closes all the output streams immediately, inside the Global Lock, so it's threadsafe.
+                Trace.Close();
+
+                // Try to copy the log file to the log folder. The reason why it's done here is that
+                // at the time the log file is being initialized the env var may not yet be set (this 
+                // applies to local log folder in particular)
+                var LogFolder = Environment.GetEnvironmentVariable(EnvVarNames.LogFolder);
+                if (!String.IsNullOrEmpty(LogFolder) && Directory.Exists(LogFolder) &&
+                        !String.IsNullOrEmpty(LogFilename) && File.Exists(LogFilename))
+                {
+                    var DestFilename = CommandUtils.CombinePaths(LogFolder, "UAT_" + Path.GetFileName(LogFilename));
+                    SafeCopyLogFile(LogFilename, DestFilename);
+                }
+            }
+            catch (Exception)
+            {
+                // Silently ignore, logging is pointless because eveything is shut down at this point
+            }
+        }
+
+        /// <summary>
+        /// Copies log file to the final log folder, does multiple attempts if the destination file could not be created.
+        /// </summary>
+        /// <param name="SourceFilename"></param>
+        /// <param name="DestFilename"></param>
+        private static void SafeCopyLogFile(string SourceFilename, string DestFilename)
+        {
+            const int MaxAttempts = 10;
+            int AttemptNo = 0;
+            var DestLogFilename = DestFilename;
+            bool Result = false;
+            do
+            {
+                try
+                {
+                    File.Copy(SourceFilename, DestLogFilename, true);
+                    Result = true;
+                }
+                catch (Exception)
+                {
+                    var ModifiedFilename = String.Format("{0}_{1}{2}", Path.GetFileNameWithoutExtension(DestFilename), AttemptNo, Path.GetExtension(DestLogFilename));
+                    DestLogFilename = CommandUtils.CombinePaths(Path.GetDirectoryName(DestFilename), ModifiedFilename);
+                    AttemptNo++;
+                }
+            }
+            while (Result == false && AttemptNo <= MaxAttempts);
+        }
+
+        /// <summary>
+        /// Creates the TraceListener used for file logging.
+        /// We cannot simply use a TextWriterTraceListener because we need more flexibility when the file cannot be created.
+        /// TextWriterTraceListener lazily creates the file, silently failing when it cannot.
+        /// </summary>
+        /// <returns>The newly created TraceListener, or null if it could not be created.</returns>
+        private static TraceListener CreateLogFileListener(out string LogFilename)
+        {
+            StreamWriter LogFile = null;
+            const int MaxAttempts = 10;
+            int Attempt = 0;
+            var TempLogFolder = Path.GetTempPath();
+            do
+            {
+                if (Attempt == 0)
+                {
+                    LogFilename = CommandUtils.CombinePaths(TempLogFolder, "Log.txt");
+                }
+                else
+                {
+                    LogFilename = CommandUtils.CombinePaths(TempLogFolder, String.Format("Log_{0}.txt", Attempt));
+                }
+                try
+                {
+                    // We do not need to set AutoFlush on the StreamWriter because we set Trace.AutoFlush, which calls it for us.
+                    // Not only would this be redundant, StreamWriter AutoFlush does not flush the encoder, while a direct call to 
+                    // StreamWriter.Flush() will, which is what the Trace system with AutoFlush = true will do.
+                    // Internally, FileStream constructor opens the file with good arguments for writing to log files.
+                    return new TextWriterTraceListener(new StreamWriter(LogFilename), "AutomationFileLogListener");
+                }
+                catch (Exception Ex)
+                {
+                    if (Attempt == (MaxAttempts - 1))
+                    {
+                        // Clear out the LogFilename to indicate we were not able to write one.
+                        LogFilename = null;
+                        UnrealBuildTool.Log.TraceWarning("Unable to create log file: {0}", LogFilename);
+                        UnrealBuildTool.Log.TraceWarning(LogUtils.FormatException(Ex));
+                    }
+                }
+            } while (LogFile == null && ++Attempt < MaxAttempts);
+            return null;
+        }
+
+        /// <summary>
+		/// Dumps exception info to log.
+		/// </summary>
+		/// <param name="Verbosity">Verbosity</param>
+		/// <param name="Ex">Exception</param>
+		public static string FormatException(Exception Ex)
+		{
+			var Message = String.Format("Exception in {0}: {1}{2}Stacktrace: {3}", Ex.Source, Ex.Message, Environment.NewLine, Ex.StackTrace);
+			if (Ex.InnerException != null)
+			{
+				Message += String.Format("InnerException in {0}: {1}{2}Stacktrace: {3}", Ex.InnerException.Source, Ex.InnerException.Message, Environment.NewLine, Ex.InnerException.StackTrace);
+			}
+			return Message;
+		}
+
+		/// <summary>
+		/// Returns a unique logfile name.
+		/// </summary>
+		/// <param name="Base">Base name for the logfile</param>
+		/// <returns>Unique logfile name.</returns>
+		public static string GetUniqueLogName(string Base)
+		{
+			const int MaxAttempts = 1000;
+
+			string LogFilename = string.Empty;
+
+			int Attempt = 0;
+			do
+			{
+				if (Attempt == 0)
+				{
+					LogFilename = String.Format("{0}.txt", Base);
+				}
+				else
+				{
+					LogFilename = String.Format("{0}.{1}.txt", Base, Attempt);
+				}
+			} while (File.Exists(LogFilename) && ++Attempt < MaxAttempts);
+
+			if (File.Exists(LogFilename))
+			{
+				throw new AutomationException(String.Format("Failed to create logfile {0}.", LogFilename));
+			}
+			return LogFilename;
+		}
+
+		public static string GetLogTail(string Filename = null, int NumLines = 250)
+		{
+			List<string> Lines;
+			if (Filename == null)
+			{
+				Lines = new List<string>(AutomationMemoryLogListener.GetAccumulatedLog().Split(new string[] { Environment.NewLine }, StringSplitOptions.RemoveEmptyEntries));
+			}
+			else
+			{
+				Lines = new List<string>(CommandUtils.ReadAllLines(Filename));
+			}
+
+			if (Lines.Count > NumLines)
+			{
+				Lines.RemoveRange(0, Lines.Count - NumLines);
+			}
+			string Result = "";
+			foreach (var Line in Lines)
+			{
+				Result += Line + Environment.NewLine;
+			}
+			return Result;
+		}
+
+
+	}
+
+	#endregion
+
+	#region AutomationMemoryLogListener
+
+	/// <summary>
+	/// Trace console listener.
+	/// </summary>
+	class AutomationMemoryLogListener : TraceListener
+	{
+		private static StringBuilder AccumulatedLog = new StringBuilder(1024 * 1024 * 20);
+		private static object SyncObject = new object();
+
+        public override bool IsThreadSafe { get { return true; } }
+
+        /// <summary>
+		/// Writes a formatted line to the console.
+		/// </summary>
+		private void WriteLinePrivate(string Message)
+		{
+			lock (SyncObject)
+			{
+				AccumulatedLog.AppendLine(Message);
+			}
+		}
+
+		public static string GetAccumulatedLog()
+		{
+			lock (SyncObject)
+			{
+				return AccumulatedLog.ToString();
+			}
+		}
+
+		#region TraceListener Interface
+
+		public override void Write(string message)
+		{
+            WriteLinePrivate(message);
+		}
+
+		public override void WriteLine(string message)
+		{
+            WriteLinePrivate(message);
+		}
+
+        #endregion
+	}
+
+	#endregion
+
+}
