UEB-260 - Break AutomationTool into AutomationUtils that all automation projects depend on, and AutomationTool, which essentially only contains the startup code.

diff --git a/Engine/Source/Programs/AutomationTool/AutomationUtils/ProcessUtils.cs b/Engine/Source/Programs/AutomationTool/AutomationUtils/ProcessUtils.cs
new file mode 100644
--- /dev/null
+++ b/Engine/Source/Programs/AutomationTool/AutomationUtils/ProcessUtils.cs
@@ -0,0 +1,991 @@
+// Copyright 1998-2015 Epic Games, Inc. All Rights Reserved.
+using System;
+using System.Collections.Generic;
+using System.Linq;
+using System.Text;
+using System.IO;
+using System.Threading;
+using System.Diagnostics;
+using System.Management;
+using System.Runtime.InteropServices;
+using System.Runtime.CompilerServices;
+
+namespace AutomationTool
+{
+	public enum CtrlTypes
+	{
+		CTRL_C_EVENT = 0,
+		CTRL_BREAK_EVENT,
+		CTRL_CLOSE_EVENT,
+		CTRL_LOGOFF_EVENT = 5,
+		CTRL_SHUTDOWN_EVENT
+	}
+
+	public interface IProcess
+	{
+		void StopProcess(bool KillDescendants = true);
+		bool HasExited { get; }
+		string GetProcessName();
+	}
+
+	/// <summary>
+	/// Tracks all active processes.
+	/// </summary>
+	public sealed class ProcessManager
+	{
+		public delegate bool CtrlHandlerDelegate(CtrlTypes EventType);
+
+		// @todo: Add mono support
+		[DllImport("Kernel32")]
+		public static extern bool SetConsoleCtrlHandler(CtrlHandlerDelegate Handler, bool Add);
+
+		/// <summary>
+		/// List of active (running) processes.
+		/// </summary>
+		private static List<IProcess> ActiveProcesses = new List<IProcess>();
+		/// <summary>
+		/// Synchronization object
+		/// </summary>
+		private static object SyncObject = new object();
+
+
+		/// <summary>
+		/// Creates a new process and adds it to the tracking list.
+		/// </summary>
+		/// <returns>New Process objects</returns>
+		public static ProcessResult CreateProcess(string AppName, bool bAllowSpew, string LogName, Dictionary<string, string> Env = null, TraceEventType SpewVerbosity = TraceEventType.Information)
+		{
+			var NewProcess = HostPlatform.Current.CreateProcess(LogName);
+			if (Env != null)
+			{
+				foreach (var EnvPair in Env)
+				{
+					if (NewProcess.StartInfo.EnvironmentVariables.ContainsKey(EnvPair.Key))
+					{
+						NewProcess.StartInfo.EnvironmentVariables.Remove(EnvPair.Key);
+					}
+					if (!String.IsNullOrEmpty(EnvPair.Value))
+					{
+						NewProcess.StartInfo.EnvironmentVariables.Add(EnvPair.Key, EnvPair.Value);
+					}
+				}
+			}
+			var Result = new ProcessResult(AppName, NewProcess, bAllowSpew, LogName, SpewVerbosity: SpewVerbosity);
+			AddProcess(Result);
+			return Result;
+		}
+
+		public static void AddProcess(IProcess Proc)
+		{
+			lock (SyncObject)
+			{
+				ActiveProcesses.Add(Proc);
+			}
+		}
+
+		public static void RemoveProcess(IProcess Proc)
+		{
+			lock (SyncObject)
+			{
+				ActiveProcesses.Remove(Proc);
+			}
+		}
+
+		public static bool CanBeKilled(string ProcessName)
+		{
+			return !HostPlatform.Current.DontKillProcessList.Contains(ProcessName, StringComparer.InvariantCultureIgnoreCase);
+		}
+
+		/// <summary>
+		/// Kills all running processes.
+		/// </summary>
+		public static void KillAll()
+		{
+			List<IProcess> ProcessesToKill = new List<IProcess>();
+			lock (SyncObject)
+			{
+				foreach (var ProcResult in ActiveProcesses)
+				{
+					if (!ProcResult.HasExited)
+					{
+						ProcessesToKill.Add(ProcResult);
+					}
+				}
+				ActiveProcesses.Clear();
+			}
+			// Remove processes that can't be killed
+			for (int ProcessIndex = ProcessesToKill.Count - 1; ProcessIndex >= 0; --ProcessIndex )
+			{
+				var ProcessName = ProcessesToKill[ProcessIndex].GetProcessName();
+				if (!String.IsNullOrEmpty(ProcessName) && !CanBeKilled(ProcessName))
+				{
+					CommandUtils.Log("Ignoring process \"{0}\" because it can't be killed.", ProcessName);
+					ProcessesToKill.RemoveAt(ProcessIndex);
+				}
+			}
+			CommandUtils.Log("Trying to kill {0} spawned processes.", ProcessesToKill.Count);
+			foreach (var Proc in ProcessesToKill)
+			{
+				CommandUtils.Log("  {0}", Proc.GetProcessName());
+			}
+            if (CommandUtils.IsBuildMachine)
+            {
+                for (int Cnt = 0; Cnt < 9; Cnt++)
+                {
+                    bool AllDone = true;
+                    foreach (var Proc in ProcessesToKill)
+                    {
+                        try
+                        {
+                            if (!Proc.HasExited)
+                            {
+                                AllDone = false;
+                                CommandUtils.Log(TraceEventType.Information, "Waiting for process: {0}", Proc.GetProcessName());
+                            }
+                        }
+                        catch (Exception)
+                        {
+                            CommandUtils.Log(TraceEventType.Information, "Exception Waiting for process");
+                            AllDone = false;
+                        }
+                    }
+                    try
+                    {
+                        if (ProcessResult.HasAnyDescendants(Process.GetCurrentProcess()))
+                        {
+                            AllDone = false;
+                            CommandUtils.Log(TraceEventType.Information, "Waiting for descendants of main process...");
+                        }
+                    }
+                    catch (Exception)
+                    {
+                        CommandUtils.Log(TraceEventType.Information, "Exception Waiting for descendants of main process");
+                        AllDone = false;
+                    }
+
+                    if (AllDone)
+                    {
+                        break;
+                    }
+                    Thread.Sleep(10000);
+                }
+            }
+			foreach (var Proc in ProcessesToKill)
+			{
+				var ProcName = Proc.GetProcessName();
+				try
+				{
+					if (!Proc.HasExited)
+					{
+						CommandUtils.Log(TraceEventType.Information, "Killing process: {0}", ProcName);
+						Proc.StopProcess(false);
+					}
+				}
+				catch (Exception Ex)
+				{
+					CommandUtils.Log(TraceEventType.Warning, "Exception while trying to kill process {0}:", ProcName);
+					CommandUtils.Log(TraceEventType.Warning, Ex);
+				}
+			}
+            try
+            {
+                if (CommandUtils.IsBuildMachine && ProcessResult.HasAnyDescendants(Process.GetCurrentProcess()))
+                {
+                    CommandUtils.Log(TraceEventType.Information, "current process still has descendants, trying to kill them...");
+                    ProcessResult.KillAllDescendants(Process.GetCurrentProcess());
+                }
+            }
+            catch (Exception)
+            {
+                CommandUtils.Log(TraceEventType.Information, "Exception killing descendants of main process");
+            }
+
+		}
+	}
+
+	#region ProcessResult Helper Class
+
+	/// <summary>
+	/// Class containing the result of process execution.
+	/// </summary>
+	public class ProcessResult : IProcess
+	{
+		private string Source = "";
+		private int ProcessExitCode = -1;
+		private StringBuilder ProcessOutput = new StringBuilder();
+		private bool AllowSpew = true;
+		private TraceEventType SpewVerbosity = TraceEventType.Information;
+		private string AppName = String.Empty;
+		private Process Proc = null;
+		private AutoResetEvent OutputWaitHandle = new AutoResetEvent(false);
+		private AutoResetEvent ErrorWaitHandle = new AutoResetEvent(false);
+		private object ProcSyncObject;
+
+		public ProcessResult(string InAppName, Process InProc, bool bAllowSpew, string LogName, TraceEventType SpewVerbosity = TraceEventType.Information)
+		{
+			AppName = InAppName;
+			ProcSyncObject = new object();
+			Proc = InProc;
+			Source = LogName;
+			AllowSpew = bAllowSpew;
+			this.SpewVerbosity = SpewVerbosity;
+			if (Proc != null)
+			{
+				Proc.EnableRaisingEvents = false;
+			}
+		}
+
+        ~ProcessResult()
+        {
+            if(Proc != null)
+            {
+                Proc.Dispose();
+            }
+        }
+
+		/// <summary>
+		/// Removes a process from the list of tracked processes.
+		/// </summary>
+		public void OnProcessExited()
+		{
+			ProcessManager.RemoveProcess(this);
+		}
+
+        /// <summary>
+        /// Log output of a remote process at a given severity.
+        /// To pretty up the output, we use a custom source so it will say the source of the process instead of this method name.
+        /// </summary>
+        /// <param name="Verbosity"></param>
+        /// <param name="Message"></param>
+        [MethodImplAttribute(MethodImplOptions.NoInlining)]
+        private void LogOutput(TraceEventType Verbosity, string Message)
+		{
+            UnrealBuildTool.Log.WriteLine(1, Source, Verbosity, Message);
+		}
+       
+		/// <summary>
+		/// Manually dispose of Proc and set it to null.
+		/// </summary>
+        public void DisposeProcess()
+        {
+			if(Proc != null)
+			{
+				Proc.Dispose();
+				Proc = null;
+			}
+        }
+
+		/// <summary>
+		/// Process.OutputDataReceived event handler.
+		/// </summary>
+		/// <param name="sender">Sender</param>
+		/// <param name="e">Event args</param>
+		public void StdOut(object sender, DataReceivedEventArgs e)
+		{
+			if (e.Data != null)
+			{
+				if (AllowSpew)
+				{
+					LogOutput(SpewVerbosity, e.Data);
+				}
+
+				ProcessOutput.Append(e.Data);
+				ProcessOutput.Append(Environment.NewLine);
+			}
+			else
+			{
+				OutputWaitHandle.Set();
+			}
+		}
+
+		/// <summary>
+		/// Process.ErrorDataReceived event handler.
+		/// </summary>
+		/// <param name="sender">Sender</param>
+		/// <param name="e">Event args</param>
+		public void StdErr(object sender, DataReceivedEventArgs e)
+		{
+			if (e.Data != null)
+			{
+				if (AllowSpew)
+				{
+                    LogOutput(SpewVerbosity, e.Data);
+				}
+
+				ProcessOutput.Append(e.Data);
+				ProcessOutput.Append(Environment.NewLine);
+			}
+			else
+			{
+				ErrorWaitHandle.Set();
+			}
+		}
+
+		/// <summary>
+		/// Convenience operator for getting the exit code value.
+		/// </summary>
+		/// <param name="Result"></param>
+		/// <returns>Process exit code.</returns>
+		public static implicit operator int(ProcessResult Result)
+		{
+			return Result.ExitCode;
+		}
+
+		/// <summary>
+		/// Gets or sets the process exit code.
+		/// </summary>
+		public int ExitCode
+		{
+			get { return ProcessExitCode; }
+			set { ProcessExitCode = value; }
+		}
+
+		/// <summary>
+		/// Gets all std output the process generated.
+		/// </summary>
+		public string Output
+		{
+			get { return ProcessOutput.ToString(); }
+		}
+
+		public bool HasExited
+		{
+			get 
+			{
+				bool bHasExited = true;
+				lock (ProcSyncObject)
+				{
+					if (Proc != null)
+					{
+						bHasExited = Proc.HasExited;
+						if (bHasExited)
+						{
+							ExitCode = Proc.ExitCode;
+						}
+					}
+				}
+				return bHasExited; 
+			}
+		}
+
+		public Process ProcessObject
+		{
+			get { return Proc; }
+		}
+
+		/// <summary>
+		/// Thread-safe way of getting the process name
+		/// </summary>
+		/// <returns>Name of the process this object represents</returns>
+		public string GetProcessName()
+		{
+			string Name = null;
+			lock (ProcSyncObject)
+			{
+				try
+				{
+					if (Proc != null && !Proc.HasExited)
+					{
+						Name = Proc.ProcessName;
+					}
+				}
+				catch
+				{
+					// Ignore all exceptions
+				}
+			}
+			if (String.IsNullOrEmpty(Name))
+			{
+				Name = "[EXITED] " + AppName;
+			}
+			return Name;
+		}
+
+		/// <summary>
+		/// Object iterface.
+		/// </summary>
+		/// <returns>String representation of this object.</returns>
+		public override string ToString()
+		{
+			return ExitCode.ToString();
+		}
+
+		public void WaitForExit()
+		{
+			bool bProcTerminated = false;
+			bool bStdOutSignalReceived = false;
+			bool bStdErrSignalReceived = false;
+			// Make sure the process objeect is valid.
+			lock (ProcSyncObject)
+			{
+				bProcTerminated = (Proc == null);
+			}
+			// Keep checking if we got all output messages until the process terminates.
+			if (!bProcTerminated)
+			{
+				// Check messages
+				int MaxWaitUntilMessagesReceived = 120;
+				while (MaxWaitUntilMessagesReceived > 0 && !(bStdOutSignalReceived && bStdErrSignalReceived))
+				{
+					if (!bStdOutSignalReceived)
+					{
+						bStdOutSignalReceived = OutputWaitHandle.WaitOne(500);
+					}
+					if (!bStdErrSignalReceived)
+					{
+						bStdErrSignalReceived = ErrorWaitHandle.WaitOne(500);
+					}
+					// Check if the process terminated
+					lock (ProcSyncObject)
+					{
+						bProcTerminated = (Proc == null) || Proc.HasExited;
+					}
+					if (bProcTerminated)
+					{
+						// Process terminated but make sure we got all messages, don't wait forever though
+						MaxWaitUntilMessagesReceived--;
+					}
+				}
+                if (!(bStdOutSignalReceived && bStdErrSignalReceived))
+                {
+                    CommandUtils.Log("Waited for a long time for output of {0}, some output may be missing; we gave up.", AppName);
+                }
+
+				// Double-check if the process terminated
+				lock (ProcSyncObject)
+				{
+					bProcTerminated = (Proc == null) || Proc.HasExited;
+				}
+				if (!bProcTerminated)
+				{
+					// The process did not terminate yet but we've read all output messages, wait until the process terminates
+					Proc.WaitForExit();
+				}
+			}
+		}
+
+		/// <summary>
+		/// Finds child processes of the current process.
+		/// </summary>
+		/// <param name="ProcessId"></param>
+		/// <param name="PossiblyRelatedId"></param>
+		/// <param name="VisitedPids"></param>
+		/// <returns></returns>
+		private static bool IsOurDescendant(Process ProcessToKill, int PossiblyRelatedId, HashSet<int> VisitedPids)
+		{
+			// check if we're the parent of it or its parent is our descendant
+			try
+			{
+				VisitedPids.Add(PossiblyRelatedId);
+				Process Parent = null;
+				using (ManagementObject ManObj = new ManagementObject(string.Format("win32_process.handle='{0}'", PossiblyRelatedId)))
+				{
+					ManObj.Get();
+					int ParentId = Convert.ToInt32(ManObj["ParentProcessId"]);
+					if (ParentId == 0 || VisitedPids.Contains(ParentId))
+					{
+						return false;
+					}
+					Parent = Process.GetProcessById(ParentId);  // will throw an exception if not spawned by us or not running
+				}
+				if (Parent != null)
+				{
+					return Parent.Id == ProcessToKill.Id || IsOurDescendant(ProcessToKill, Parent.Id, VisitedPids);  // could use ParentId, but left it to make the above var used
+				}
+				else
+				{
+					return false;
+				}
+			}
+			catch (Exception)
+			{
+				// This can happen if the pid is no longer valid which is ok.
+				return false;
+			}
+		}
+
+		/// <summary>
+		/// Kills all child processes of the specified process.
+		/// </summary>
+		/// <param name="ProcessId">Process id</param>
+		public static void KillAllDescendants(Process ProcessToKill)
+		{
+			bool bKilledAChild;
+			do
+			{
+				bKilledAChild = false;
+				// For some reason Process.GetProcesses() sometimes returns the same process twice
+				// So keep track of all processes we already tried to kill
+				var KilledPids = new HashSet<int>();
+				var AllProcs = Process.GetProcesses();
+				foreach (Process KillCandidate in AllProcs)
+				{
+					var VisitedPids = new HashSet<int>();
+					if (ProcessManager.CanBeKilled(KillCandidate.ProcessName) && 
+						!KilledPids.Contains(KillCandidate.Id) && 
+						IsOurDescendant(ProcessToKill, KillCandidate.Id, VisitedPids))
+					{
+						KilledPids.Add(KillCandidate.Id);
+						CommandUtils.Log("Trying to kill descendant pid={0}, name={1}", KillCandidate.Id, KillCandidate.ProcessName);
+						try
+						{
+							KillCandidate.Kill();
+							bKilledAChild = true;
+						}
+						catch (Exception Ex)
+						{
+							CommandUtils.Log(TraceEventType.Warning, "Failed to kill descendant:");
+							CommandUtils.Log(TraceEventType.Warning, Ex);
+						}
+						break;  // exit the loop as who knows what else died, so let's get processes anew
+					}
+				}
+			} while (bKilledAChild);
+		}
+
+        /// <summary>
+        /// returns true if this process has any descendants
+        /// </summary>
+        /// <param name="ProcessToCheck">Process to check</param>
+        public static bool HasAnyDescendants(Process ProcessToCheck)
+        {
+            Process[] AllProcs = Process.GetProcesses();
+            foreach (Process KillCandidate in AllProcs)
+            {
+                HashSet<int> VisitedPids = new HashSet<int>();
+                if (ProcessManager.CanBeKilled(KillCandidate.ProcessName) && IsOurDescendant(ProcessToCheck, KillCandidate.Id, VisitedPids))
+                {
+                    CommandUtils.Log("Descendant pid={0}, name={1}, filename={2}", KillCandidate.Id, KillCandidate.ProcessName,
+						KillCandidate.MainModule != null ? KillCandidate.MainModule.FileName : "unknown");
+                    return true;
+                }
+            }
+            return false;
+        }
+
+		public void StopProcess(bool KillDescendants = true)
+		{
+			if (Proc != null)
+			{
+				Process ProcToKill = null;
+				// At this point any call to Proc memebers will result in an exception
+				// so null the object.
+				var ProcToKillName = GetProcessName();
+				lock (ProcSyncObject)
+				{
+					ProcToKill = Proc;
+					Proc = null;
+				}
+				// Now actually kill the process and all its descendants if requested
+				if (KillDescendants)
+				{
+					KillAllDescendants(ProcToKill);
+				}
+				try
+				{
+					ProcToKill.Kill();
+					ProcToKill.WaitForExit(60000);
+					if (!ProcToKill.HasExited)
+					{
+						CommandUtils.Log("Process {0} failed to exit.", ProcToKillName);
+					}
+					else
+					{
+						CommandUtils.Log("Process {0} successfully exited.", ProcToKillName);
+						OnProcessExited();
+					}
+					ProcToKill.Close();					
+				}
+				catch (Exception Ex)
+				{
+					CommandUtils.Log(TraceEventType.Warning, "Exception while trying to kill process {0}:", ProcToKillName);
+					CommandUtils.Log(TraceEventType.Warning, Ex);
+				}
+			}
+		}
+	}
+
+	#endregion
+
+	public partial class CommandUtils
+	{
+		#region Statistics
+
+		private static Dictionary<string, int> ExeToTimeInMs = new Dictionary<string, int>();
+
+		public static void AddRunTime(string Exe, int TimeInMs)
+		{
+			string Base = Path.GetFileName(Exe);
+			if (!ExeToTimeInMs.ContainsKey(Base))
+			{
+				ExeToTimeInMs.Add(Base, TimeInMs);
+			}
+			else
+			{
+				ExeToTimeInMs[Base] += TimeInMs;
+			}
+		}
+
+		public static void PrintRunTime()
+		{
+			foreach (var Item in ExeToTimeInMs)
+			{
+				Log("Run: Total {0}s to run " + Item.Key, Item.Value / 1000);
+			}
+			ExeToTimeInMs.Clear();
+		}
+
+		#endregion
+
+		[Flags]
+		public enum ERunOptions
+		{
+			None = 0,
+			AllowSpew = 1 << 0,
+			AppMustExist = 1 << 1,
+			NoWaitForExit = 1 << 2,
+			NoStdOutRedirect = 1 << 3,
+            NoLoggingOfRunCommand = 1 << 4,
+            UTF8Output = 1 << 5,
+
+			/// When specified with AllowSpew, the output will be TraceEventType.Verbose instead of TraceEventType.Information
+			SpewIsVerbose = 1 << 6,
+
+			Default = AllowSpew | AppMustExist,
+		}
+
+		/// <summary>
+		/// Runs external program.
+		/// </summary>
+		/// <param name="App">Program filename.</param>
+		/// <param name="CommandLine">Commandline</param>
+		/// <param name="Input">Optional Input for the program (will be provided as stdin)</param>
+		/// <param name="Options">Defines the options how to run. See ERunOptions.</param>
+		/// <param name="Env">Environment to pass to program.</param>
+		/// <returns>Object containing the exit code of the program as well as it's stdout output.</returns>
+		public static ProcessResult Run(string App, string CommandLine = null, string Input = null, ERunOptions Options = ERunOptions.Default, Dictionary<string, string> Env = null)
+		{
+			App = ConvertSeparators(PathSeparator.Default, App);
+			HostPlatform.Current.SetupOptionsForRun(ref App, ref Options, ref CommandLine);
+			if (App == "ectool" || App == "zip" || App == "xcodebuild")
+			{
+				Options &= ~ERunOptions.AppMustExist;
+			}
+
+			if (Options.HasFlag(ERunOptions.AppMustExist) && !FileExists(Options.HasFlag(ERunOptions.NoLoggingOfRunCommand) ? true : false, App))
+			{
+				throw new AutomationException("BUILD FAILED: Couldn't find the executable to Run: {0}", App);
+			}
+			var StartTime = DateTime.UtcNow;
+
+			TraceEventType SpewVerbosity = Options.HasFlag(ERunOptions.SpewIsVerbose) ? TraceEventType.Verbose : TraceEventType.Information;
+            if (!Options.HasFlag(ERunOptions.NoLoggingOfRunCommand))
+            {
+                Log(SpewVerbosity,"Run: " + App + " " + (String.IsNullOrEmpty(CommandLine) ? "" : CommandLine));
+            }
+			ProcessResult Result = ProcessManager.CreateProcess(App, Options.HasFlag(ERunOptions.AllowSpew), Path.GetFileNameWithoutExtension(App), Env, SpewVerbosity:SpewVerbosity);
+			Process Proc = Result.ProcessObject;
+
+			bool bRedirectStdOut = (Options & ERunOptions.NoStdOutRedirect) != ERunOptions.NoStdOutRedirect;			
+			Proc.StartInfo.FileName = App;
+			Proc.StartInfo.Arguments = String.IsNullOrEmpty(CommandLine) ? "" : CommandLine;
+			Proc.StartInfo.UseShellExecute = false;
+			if (bRedirectStdOut)
+			{
+				Proc.StartInfo.RedirectStandardOutput = true;
+				Proc.StartInfo.RedirectStandardError = true;
+				Proc.OutputDataReceived += Result.StdOut;
+				Proc.ErrorDataReceived += Result.StdErr;
+			}
+			Proc.StartInfo.RedirectStandardInput = Input != null;
+			Proc.StartInfo.CreateNoWindow = true;
+			if ((Options & ERunOptions.UTF8Output) == ERunOptions.UTF8Output)
+			{
+				Proc.StartInfo.StandardOutputEncoding = new System.Text.UTF8Encoding(false, false);
+			}
+			Proc.Start();
+
+			if (bRedirectStdOut)
+			{
+				Proc.BeginOutputReadLine();
+				Proc.BeginErrorReadLine();
+			}
+
+			if (String.IsNullOrEmpty(Input) == false)
+			{
+				Proc.StandardInput.WriteLine(Input);
+				Proc.StandardInput.Close();
+			}
+
+			if (!Options.HasFlag(ERunOptions.NoWaitForExit))
+			{
+				Result.WaitForExit();
+				var BuildDuration = (DateTime.UtcNow - StartTime).TotalMilliseconds;
+				AddRunTime(App, (int)(BuildDuration));
+				Result.ExitCode = Proc.ExitCode;
+				if (!Options.HasFlag(ERunOptions.NoLoggingOfRunCommand))
+                {
+                    Log(SpewVerbosity,"Run: Took {0}s to run {1}, ExitCode={2}", BuildDuration / 1000, Path.GetFileName(App), Result.ExitCode);
+                }
+				Result.OnProcessExited();
+                Result.DisposeProcess();
+			}
+			else
+			{
+				Result.ExitCode = -1;
+			}
+
+			return Result;
+		}
+
+        /// <summary>
+        /// Gets a logfile name for a RunAndLog call
+        /// </summary>
+        /// <param name="Env">Environment to use.</param>
+        /// <param name="App">Executable to run</param>
+        /// <param name="LogName">Name of the logfile ( if null, executable name is used )</param>
+        /// <returns>The log file name.</returns>
+        public static string GetRunAndLogLogName(CommandEnvironment Env, string App, string LogName = null)
+        {
+            if (LogName == null)
+            {
+                LogName = Path.GetFileNameWithoutExtension(App);
+            }
+            return LogUtils.GetUniqueLogName(CombinePaths(Env.LogFolder, LogName));
+        }
+
+		/// <summary>
+		/// Runs external program and writes the output to a logfile.
+		/// </summary>
+		/// <param name="Env">Environment to use.</param>
+		/// <param name="App">Executable to run</param>
+		/// <param name="CommandLine">Commandline to pass on to the executable</param>
+		/// <param name="LogName">Name of the logfile ( if null, executable name is used )</param>
+		/// <param name="Input">Optional Input for the program (will be provided as stdin)</param>
+		/// <param name="Options">Defines the options how to run. See ERunOptions.</param>
+		public static void RunAndLog(CommandEnvironment Env, string App, string CommandLine, string LogName = null, int MaxSuccessCode = 0, string Input = null, ERunOptions Options = ERunOptions.Default, Dictionary<string, string> EnvVars = null)
+		{
+            RunAndLog(App, CommandLine, GetRunAndLogLogName(Env, App, LogName), MaxSuccessCode, Input, Options, EnvVars);
+		}
+        /// <summary>
+        /// Runs external program and writes the output to a logfile.
+        /// </summary>
+        /// <param name="App">Executable to run</param>
+        /// <param name="CommandLine">Commandline to pass on to the executable</param>
+        /// <param name="Logfile">Full path to the logfile, where the application output should be written to.</param>
+		/// <param name="Input">Optional Input for the program (will be provided as stdin)</param>
+		/// <param name="Options">Defines the options how to run. See ERunOptions.</param>
+        public static string RunAndLog(string App, string CommandLine, string Logfile = null, int MaxSuccessCode = 0, string Input = null, ERunOptions Options = ERunOptions.Default, Dictionary<string, string> EnvVars = null)
+        {
+            ProcessResult Result = Run(App, CommandLine, Input, Options, EnvVars);
+            if (Result.Output.Length > 0 && Logfile != null)
+            {
+                WriteToFile(Logfile, Result.Output);
+            }
+            else if (Logfile == null)
+            {
+                Logfile = "[No logfile specified]";
+            }
+            else
+            {
+                Logfile = "[None!, no output produced]";
+            }
+
+            if (Result > MaxSuccessCode || Result < 0)
+            {
+                throw new AutomationException((ErrorCodes)Result.ExitCode, String.Format("Command failed (Result:{3}): {0} {1}. See logfile for details: '{2}' ",
+                                                App, CommandLine, Path.GetFileName(Logfile), Result.ExitCode));
+            }
+            if (Result.Output.Length > 0)
+            {
+                return Result.Output;
+            }
+            return "";
+        }
+
+        /// <summary>
+        /// Runs external program and writes the output to a logfile.
+        /// </summary>
+        /// <param name="App">Executable to run</param>
+        /// <param name="CommandLine">Commandline to pass on to the executable</param>
+        /// <param name="Logfile">Full path to the logfile, where the application output should be written to.</param>
+        /// <returns>Whether the program executed successfully or not.</returns>
+        public static string RunAndLog(string App, string CommandLine, out int SuccessCode, string Logfile = null, Dictionary<string, string> EnvVars = null)
+        {
+            ProcessResult Result = Run(App, CommandLine, Env: EnvVars);
+            SuccessCode = Result.ExitCode;
+            if (Result.Output.Length > 0 && Logfile != null)
+            {
+                WriteToFile(Logfile, Result.Output);
+            }
+            if (!String.IsNullOrEmpty(Result.Output))
+            {
+                return Result.Output;
+            }
+            return "";
+        }
+
+        /// <summary>
+        /// Runs external program and writes the output to a logfile.
+        /// </summary>
+        /// <param name="Env">Environment to use.</param>
+        /// <param name="App">Executable to run</param>
+        /// <param name="CommandLine">Commandline to pass on to the executable</param>
+        /// <param name="LogName">Name of the logfile ( if null, executable name is used )</param>
+        /// <returns>Whether the program executed successfully or not.</returns>
+        public static string RunAndLog(CommandEnvironment Env, string App, string CommandLine, out int SuccessCode, string LogName = null, Dictionary<string, string> EnvVars = null)
+        {
+            return RunAndLog(App, CommandLine, out SuccessCode, GetRunAndLogLogName(Env, App, LogName), EnvVars);
+        }
+
+		/// <summary>
+		/// Runs UAT recursively
+		/// </summary>
+		/// <param name="Env">Environment to use.</param>
+		/// <param name="CommandLine">Commandline to pass on to the executable</param>
+		public static string RunUAT(CommandEnvironment Env, string CommandLine)
+		{
+			// this doesn't do much, but it does need to make sure log folders are reasonable and don't collide
+			string BaseLogSubdir = "Recur";
+			if (!String.IsNullOrEmpty(CommandLine))
+			{
+				int Space = CommandLine.IndexOf(" ");
+				if (Space > 0)
+				{
+					BaseLogSubdir = BaseLogSubdir + "_" + CommandLine.Substring(0, Space);
+				}
+				else
+				{
+					BaseLogSubdir = BaseLogSubdir + "_" + CommandLine;
+				}
+			}
+			int Index = 0;
+
+			BaseLogSubdir = BaseLogSubdir.Trim();
+			string DirOnlyName = BaseLogSubdir;
+
+			string LogSubdir = CombinePaths(CmdEnv.LogFolder, DirOnlyName, "");
+            while (true)
+			{
+                var ExistingFiles = FindFiles(DirOnlyName + "*", false, CmdEnv.LogFolder);
+                if (ExistingFiles.Length == 0)
+                {
+                    break;
+                }
+				Index++;
+				if (Index == 1000)
+				{
+					throw new AutomationException("Couldn't seem to create a log subdir {0}", LogSubdir);
+				}
+				DirOnlyName = String.Format("{0}_{1}_", BaseLogSubdir, Index);
+				LogSubdir = CombinePaths(CmdEnv.LogFolder, DirOnlyName, "");
+			}
+			string LogFile = CombinePaths(CmdEnv.LogFolder, DirOnlyName + ".log");
+			Log("Recursive UAT Run, in log folder {0}, main log file {1}", LogSubdir, LogFile);
+			CreateDirectory(LogSubdir);
+
+			CommandLine = CommandLine + " -NoCompile";
+
+			string App = CmdEnv.UATExe;
+
+			Log("Running {0} {1}", App, CommandLine);
+			var OSEnv = new Dictionary<string, string>();
+
+			OSEnv.Add(AutomationTool.EnvVarNames.LogFolder, LogSubdir);
+			OSEnv.Add("uebp_UATMutexNoWait", "1");
+			if (!IsBuildMachine)
+			{
+				OSEnv.Add(AutomationTool.EnvVarNames.LocalRoot, ""); // if we don't clear this out, it will think it is a build machine; it will rederive everything
+			}
+
+			ProcessResult Result = Run(App, CommandLine, null, ERunOptions.Default, OSEnv);
+			if (Result.Output.Length > 0)
+			{
+				WriteToFile(LogFile, Result.Output);
+			}
+			else
+			{
+				WriteToFile(LogFile, "[None!, no output produced]");
+			}
+
+			Log("Flattening log folder {0}", LogSubdir);
+
+			var Files = FindFiles("*", true, LogSubdir);
+			string MyLogFolder = CombinePaths(CmdEnv.LogFolder, "");
+			foreach (var ThisFile in Files)
+			{
+				if (!ThisFile.StartsWith(MyLogFolder, StringComparison.InvariantCultureIgnoreCase))
+				{
+					throw new AutomationException("Can't rebase {0} because it doesn't start with {1}", ThisFile, MyLogFolder);
+				}
+				string NewFilename = ThisFile.Substring(MyLogFolder.Length).Replace("/", "_").Replace("\\", "_");
+				NewFilename = CombinePaths(CmdEnv.LogFolder, NewFilename);
+				if (FileExists_NoExceptions(NewFilename))
+				{
+					throw new AutomationException("Destination log file already exists? {0}", NewFilename);
+				}
+				CopyFile(ThisFile, NewFilename);
+				if (!FileExists_NoExceptions(NewFilename))
+				{
+					throw new AutomationException("Destination log file could not be copied {0}", NewFilename);
+				}
+				DeleteFile_NoExceptions(ThisFile);
+			}
+            DeleteDirectory_NoExceptions(LogSubdir);
+
+			if (Result != 0)
+			{
+				throw new AutomationException(String.Format("Recursive UAT Command failed (Result:{3}): {0} {1}. See logfile for details: '{2}' ",
+																				App, CommandLine, Path.GetFileName(LogFile), Result.ExitCode));
+			}
+			return LogFile;
+		}
+
+		protected delegate bool ProcessLog(string LogText);
+		/// <summary>
+		/// Keeps reading a log file as it's being written to by another process until it exits.
+		/// </summary>
+		/// <param name="LogFilename">Name of the log file.</param>
+		/// <param name="LogProcess">Process that writes to the log file.</param>
+		/// <param name="OnLogRead">Callback used to process the recently read log contents.</param>
+		protected static void LogFileReaderProcess(string LogFilename, ProcessResult LogProcess, ProcessLog OnLogRead = null)
+		{
+			while (!FileExists(LogFilename) && !LogProcess.HasExited)
+			{
+				Log("Waiting for logging process to start...");
+				Thread.Sleep(2000);
+			}
+			Thread.Sleep(1000);
+
+			using (FileStream ProcessLog = File.Open(LogFilename, FileMode.Open, FileAccess.Read, FileShare.ReadWrite))
+			{
+				StreamReader LogReader = new StreamReader(ProcessLog);
+				bool bKeepReading = true;
+
+				// Read until the process has exited.
+				while (!LogProcess.HasExited && bKeepReading)
+				{
+					while (!LogReader.EndOfStream && bKeepReading)
+					{
+						string Output = LogReader.ReadToEnd();
+						if (Output != null && OnLogRead != null)
+						{
+							bKeepReading = OnLogRead(Output);
+						}
+					}
+
+					while (LogReader.EndOfStream && !LogProcess.HasExited && bKeepReading)
+					{
+						Thread.Sleep(250);
+						// Tick the callback so that it can respond to external events
+						if (OnLogRead != null)
+						{
+							bKeepReading = OnLogRead(null);
+						}
+					}
+				}
+			}
+		}
+
+	}
+}
