UEB-260 - Break AutomationTool into AutomationUtils that all automation projects depend on, and AutomationTool, which essentially only contains the startup code.

diff --git a/Engine/Source/Programs/AutomationTool/AutomationUtils/ProjectParams.cs b/Engine/Source/Programs/AutomationTool/AutomationUtils/ProjectParams.cs
new file mode 100644
--- /dev/null
+++ b/Engine/Source/Programs/AutomationTool/AutomationUtils/ProjectParams.cs
@@ -0,0 +1,2095 @@
+﻿// Copyright 1998-2015 Epic Games, Inc. All Rights Reserved.
+using System;
+using System.Collections.Generic;
+using System.Linq;
+using System.Text;
+using UnrealBuildTool;
+using System.IO;
+using System.Reflection;
+
+namespace AutomationTool
+{
+	[Help("targetplatform=PlatformName", "target platform for building, cooking and deployment (also -Platform)")]
+	[Help("servertargetplatform=PlatformName", "target platform for building, cooking and deployment of the dedicated server (also -ServerPlatform)")]
+	public class ProjectParams
+	{
+		#region Constructors
+
+		/// <summary>
+		/// Gets a parameter from the command line if it hasn't been specified in the constructor. 
+		/// If the command line is not available, default value will be used.
+		/// </summary>
+		/// <param name="Command">Command to parse the command line for. Can be null.</param>
+		/// <param name="SpecifiedValue">Value specified in the constructor (or not)</param>
+		/// <param name="Default">Default value.</param>
+		/// <param name="ParamNames">Command line parameter names to parse.</param>
+		/// <returns>Parameter value.</returns>
+		bool GetParamValueIfNotSpecified(CommandUtils Command, bool? SpecifiedValue, bool Default, params string[] ParamNames)
+		{
+			if (SpecifiedValue.HasValue)
+			{
+				return SpecifiedValue.Value;
+			}
+			else if (Command != null)
+			{
+				bool Value = false;
+				foreach (var Param in ParamNames)
+				{
+					Value = Value || Command.ParseParam(Param);
+				}
+				return Value;
+			}
+			else
+			{
+				return Default;
+			}
+		}
+
+		/// <summary>
+		/// Gets optional parameter from the command line if it hasn't been specified in the constructor. 
+		/// If the command line is not available or the command has not been specified in the command line, default value will be used.
+		/// </summary>
+		/// <param name="Command">Command to parse the command line for. Can be null.</param>
+		/// <param name="SpecifiedValue">Value specified in the constructor (or not)</param>
+		/// <param name="Default">Default value.</param>
+		/// <param name="TrueParam">Name of a parameter that sets the value to 'true', for example: -clean</param>
+		/// <param name="FalseParam">Name of a parameter that sets the value to 'false', for example: -noclean</param>
+		/// <returns>Parameter value or default value if the paramater has not been specified</returns>
+		bool? GetOptionalParamValueIfNotSpecified(CommandUtils Command, bool? SpecifiedValue, bool? Default, string TrueParam, string FalseParam)
+		{
+			if (SpecifiedValue.HasValue)
+			{
+				return SpecifiedValue.Value;
+			}
+			else if (Command != null)
+			{
+				bool? Value = null;
+				if (!String.IsNullOrEmpty(TrueParam) && Command.ParseParam(TrueParam))
+				{
+					Value = true;
+				}
+				else if (!String.IsNullOrEmpty(FalseParam) && Command.ParseParam(FalseParam))
+				{
+					Value = false;
+				}
+				if (Value.HasValue)
+				{
+					return Value;
+				}
+			}
+			return Default;
+		}
+
+		/// <summary>
+		/// Gets a parameter value from the command line if it hasn't been specified in the constructor. 
+		/// If the command line is not available, default value will be used.
+		/// </summary>
+		/// <param name="Command">Command to parse the command line for. Can be null.</param>
+		/// <param name="SpecifiedValue">Value specified in the constructor (or not)</param>
+		/// <param name="ParamName">Command line parameter name to parse.</param>
+		/// <param name="Default">Default value</param>
+		/// <param name="bTrimQuotes">If set, the leading and trailing quotes will be removed, e.g. instead of "[REDACTED_USER_HOME] Name" it will return [REDACTED_USER_HOME] Name</param>
+		/// <returns>Parameter value.</returns>
+		string ParseParamValueIfNotSpecified(CommandUtils Command, string SpecifiedValue, string ParamName, string Default = "", bool bTrimQuotes = false)
+		{
+			string Result = Default;
+
+			if (SpecifiedValue != null)
+			{
+				Result = SpecifiedValue;
+			}
+			else if (Command != null)
+			{
+				Result = Command.ParseParamValue(ParamName, Default);
+			}
+
+			return bTrimQuotes ? Result.Trim( new char[]{'\"'} ) : Result;
+		}
+
+		/// <summary>
+		/// Sets up platforms
+		/// </summary>
+        /// <param name="DependentPlatformMap">Set with a mapping from source->destination if specified on command line</param>
+		/// <param name="Command">The command line to parse</param>
+		/// <param name="OverrideTargetPlatforms">If passed use this always</param>
+        /// <param name="DefaultTargetPlatforms">Use this if nothing is passed on command line</param>
+		/// <param name="AllowPlatformParams">Allow raw -platform options</param>
+		/// <param name="PlatformParamNames">Possible -parameters to check for</param>
+		/// <returns>List of platforms parsed from the command line</returns>
+		private List<UnrealTargetPlatform> SetupTargetPlatforms(ref Dictionary<UnrealTargetPlatform,UnrealTargetPlatform> DependentPlatformMap, CommandUtils Command, List<UnrealTargetPlatform> OverrideTargetPlatforms, List<UnrealTargetPlatform> DefaultTargetPlatforms, bool AllowPlatformParams, params string[] PlatformParamNames)
+		{
+			List<UnrealTargetPlatform> TargetPlatforms = null;
+			if (CommandUtils.IsNullOrEmpty(OverrideTargetPlatforms))
+			{
+				if (Command != null)
+				{
+					// Parse the command line, we support the following params:
+					// -'PlatformParamNames[n]'=Platform_1+Platform_2+...+Platform_k
+					// or (if AllowPlatformParams is true)
+					// -Platform_1 -Platform_2 ... -Platform_k
+					string CmdLinePlatform = null;
+					foreach (var ParamName in PlatformParamNames)
+					{
+						CmdLinePlatform = Command.ParseParamValue(ParamName);
+						if (!String.IsNullOrEmpty(CmdLinePlatform))
+						{
+							break;
+						}
+					}
+
+					if (!String.IsNullOrEmpty(CmdLinePlatform))
+					{
+						// Get all platforms from the param value: Platform_1+Platform_2+...+Platform_k
+						TargetPlatforms = new List<UnrealTargetPlatform>();
+						var Platforms = new List<string>(CmdLinePlatform.Split('+'));
+						foreach (var PlatformName in Platforms)
+						{
+                            // Look for dependent platforms, Source_1.Dependent_1+Source_2.Dependent_2+Standalone_3
+                            var SubPlatforms = new List<string>(PlatformName.Split('.'));
+
+                            foreach (var SubPlatformName in SubPlatforms)
+                            {
+                                UnrealTargetPlatform NewPlatform = (UnrealTargetPlatform)Enum.Parse(typeof(UnrealTargetPlatform), SubPlatformName, true);
+                                TargetPlatforms.Add(NewPlatform);
+
+                                if (SubPlatformName != SubPlatforms[0])
+                                {
+                                    // We're a dependent platform so add ourselves to the map, pointing to the first element in the list
+                                    DependentPlatformMap.Add(NewPlatform, (UnrealTargetPlatform)Enum.Parse(typeof(UnrealTargetPlatform), SubPlatforms[0], true));
+                                }
+                            }
+						}
+					}
+					else if (AllowPlatformParams)
+					{
+						// Look up platform names in the command line: -Platform_1 -Platform_2 ... -Platform_k
+						TargetPlatforms = new List<UnrealTargetPlatform>();
+						foreach (var Plat in CommandUtils.KnownTargetPlatforms)
+						{
+							if (Command.ParseParam(Plat.ToString()))
+							{
+								TargetPlatforms.Add(Plat);
+							}
+						}
+					}
+				}
+			}
+			else
+			{
+				TargetPlatforms = OverrideTargetPlatforms;
+			}
+			if (CommandUtils.IsNullOrEmpty(TargetPlatforms))
+			{
+				// Revert to single default platform - the current platform we're running
+				TargetPlatforms = DefaultTargetPlatforms;
+			}
+			return TargetPlatforms;
+		}
+
+		public ProjectParams(ProjectParams InParams)
+		{
+			//
+			//// Use this.Name with properties and fields!
+			//
+
+			this.RawProjectPath = InParams.RawProjectPath;
+			this.MapsToCook = InParams.MapsToCook;
+			this.DirectoriesToCook = InParams.DirectoriesToCook;
+            this.InternationalizationPreset = InParams.InternationalizationPreset;
+            this.CulturesToCook = InParams.CulturesToCook;
+            this.BasedOnReleaseVersion = InParams.BasedOnReleaseVersion;
+            this.CreateReleaseVersion = InParams.CreateReleaseVersion;
+            this.GeneratePatch = InParams.GeneratePatch;
+            this.DLCName = InParams.DLCName;
+            this.DLCIncludeEngineContent = InParams.DLCIncludeEngineContent;
+            this.NewCook = InParams.NewCook;
+            this.OldCook = InParams.OldCook;
+            this.AdditionalCookerOptions = InParams.AdditionalCookerOptions;
+			this.ClientCookedTargets = InParams.ClientCookedTargets;
+			this.ServerCookedTargets = InParams.ServerCookedTargets;
+			this.EditorTargets = InParams.EditorTargets;
+			this.ProgramTargets = InParams.ProgramTargets;
+			this.ClientTargetPlatforms = InParams.ClientTargetPlatforms;
+            this.ClientDependentPlatformMap = InParams.ClientDependentPlatformMap;
+			this.ServerTargetPlatforms = InParams.ServerTargetPlatforms;
+            this.ServerDependentPlatformMap = InParams.ServerDependentPlatformMap;
+			this.Build = InParams.Build;
+			this.Run = InParams.Run;
+			this.Cook = InParams.Cook;
+			this.IterativeCooking = InParams.IterativeCooking;
+            this.CookAll = InParams.CookAll;
+            this.CookMapsOnly = InParams.CookMapsOnly;
+			this.CookFlavor = InParams.CookFlavor;
+			this.SkipCook = InParams.SkipCook;
+			this.SkipCookOnTheFly = InParams.SkipCookOnTheFly;
+            this.Prebuilt = InParams.Prebuilt;
+            this.RunTimeoutSeconds = InParams.RunTimeoutSeconds;
+			this.Clean = InParams.Clean;
+			this.Pak = InParams.Pak;
+			this.SignPak = InParams.SignPak;
+			this.SignedPak = InParams.SignedPak;
+			this.SkipPak = InParams.SkipPak;
+			this.NoXGE = InParams.NoXGE;
+			this.CookOnTheFly = InParams.CookOnTheFly;
+            this.CookOnTheFlyStreaming = InParams.CookOnTheFlyStreaming;
+            this.UnversionedCookedContent = InParams.UnversionedCookedContent;
+			this.FileServer = InParams.FileServer;
+			this.DedicatedServer = InParams.DedicatedServer;
+			this.Client = InParams.Client;
+			this.NoClient = InParams.NoClient;
+			this.LogWindow = InParams.LogWindow;
+			this.Stage = InParams.Stage;
+			this.SkipStage = InParams.SkipStage;
+            this.StageDirectoryParam = InParams.StageDirectoryParam;
+            this.StageNonMonolithic = InParams.StageNonMonolithic;
+			this.Manifests = InParams.Manifests;
+            this.CreateChunkInstall = InParams.CreateChunkInstall;
+			this.UE4Exe = InParams.UE4Exe;
+			this.NoDebugInfo = InParams.NoDebugInfo;
+			this.NoCleanStage = InParams.NoCleanStage;
+			this.MapToRun = InParams.MapToRun;
+			this.AdditionalServerMapParams = InParams.AdditionalServerMapParams;
+			this.Foreign = InParams.Foreign;
+			this.ForeignCode = InParams.ForeignCode;
+			this.StageCommandline = InParams.StageCommandline;
+            this.BundleName = InParams.BundleName;
+			this.RunCommandline = InParams.RunCommandline;
+			this.Package = InParams.Package;
+			this.Deploy = InParams.Deploy;
+			this.IterativeDeploy = InParams.IterativeDeploy;
+			this.IgnoreCookErrors = InParams.IgnoreCookErrors;
+			this.FastCook = InParams.FastCook;
+			this.Device = InParams.Device;
+			this.DeviceName = InParams.DeviceName;
+			this.ServerDevice = InParams.ServerDevice;
+            this.NullRHI = InParams.NullRHI;
+            this.FakeClient = InParams.FakeClient;
+            this.EditorTest = InParams.EditorTest;
+            this.RunAutomationTests = InParams.RunAutomationTests;
+            this.RunAutomationTest = InParams.RunAutomationTest;
+            this.CrashIndex = InParams.CrashIndex;
+            this.Port = InParams.Port;
+			this.SkipServer = InParams.SkipServer;
+			this.Rocket = InParams.Rocket;
+			this.Unattended = InParams.Unattended;
+            this.ServerDeviceAddress = InParams.ServerDeviceAddress;
+            this.DeviceUsername = InParams.DeviceUsername;
+            this.DevicePassword = InParams.DevicePassword;
+            this.CrashReporter = InParams.CrashReporter;
+			this.ClientConfigsToBuild = InParams.ClientConfigsToBuild;
+			this.ServerConfigsToBuild = InParams.ServerConfigsToBuild;
+			this.NumClients = InParams.NumClients;
+            this.Compressed = InParams.Compressed;
+            this.UseDebugParamForEditorExe = InParams.UseDebugParamForEditorExe;
+            this.bUsesSteam = InParams.bUsesSteam;
+			this.bUsesCEF3 = InParams.bUsesCEF3;
+			this.bUsesSlate = InParams.bUsesSlate;
+			this.bUsesSlateEditorStyle = InParams.bUsesSlateEditorStyle;
+            this.bDebugBuildsActuallyUseDebugCRT = InParams.bDebugBuildsActuallyUseDebugCRT;
+			this.Archive = InParams.Archive;
+			this.ArchiveDirectoryParam = InParams.ArchiveDirectoryParam;
+			this.ArchiveMetaData = InParams.ArchiveMetaData;
+			this.Distribution = InParams.Distribution;
+			this.Prereqs = InParams.Prereqs;
+			this.NoBootstrapExe = InParams.NoBootstrapExe;
+            this.Prebuilt = InParams.Prebuilt;
+            this.RunTimeoutSeconds = InParams.RunTimeoutSeconds;
+			this.bIsCodeBasedProject = InParams.bIsCodeBasedProject;
+		}
+
+		/// <summary>
+		/// Constructor. Be sure to use this.ParamName to set the actual property name as parameter names and property names
+		/// overlap here.
+		/// If a parameter value is not set, it will be parsed from the command line if the command is null, the default value will be used.
+		/// </summary>
+		public ProjectParams(			
+			string RawProjectPath,
+
+			CommandUtils Command = null,
+			string Device = null,			
+			string MapToRun = null,	
+			string AdditionalServerMapParams = null,
+			ParamList<string> Port = null,
+			string RunCommandline = null,						
+			string StageCommandline = null,
+            string BundleName = null,
+            string StageDirectoryParam = null,
+            bool? StageNonMonolithic = null,
+			string UE4Exe = null,
+			string SignPak = null,
+			List<UnrealTargetConfiguration> ClientConfigsToBuild = null,
+			List<UnrealTargetConfiguration> ServerConfigsToBuild = null,
+			ParamList<string> MapsToCook = null,
+			ParamList<string> DirectoriesToCook = null,
+            string InternationalizationPreset = null,
+            ParamList<string> CulturesToCook = null,
+			ParamList<string> ClientCookedTargets = null,
+			ParamList<string> EditorTargets = null,
+			ParamList<string> ServerCookedTargets = null,
+			List<UnrealTargetPlatform> ClientTargetPlatforms = null,
+            Dictionary<UnrealTargetPlatform, UnrealTargetPlatform> ClientDependentPlatformMap = null,
+			List<UnrealTargetPlatform> ServerTargetPlatforms = null,
+            Dictionary<UnrealTargetPlatform, UnrealTargetPlatform> ServerDependentPlatformMap = null,
+			bool? Build = null,
+			bool? Cook = null,
+			string CookFlavor = null,
+			bool? Run = null,
+			bool? SkipServer = null,
+			bool? Clean = null,
+            bool? Compressed = null,
+            bool? UseDebugParamForEditorExe = null,
+            bool? IterativeCooking = null,
+            bool? CookAll = null,
+            bool? CookMapsOnly = null,
+            bool? CookOnTheFly = null,
+            bool? CookOnTheFlyStreaming = null,
+            bool? UnversionedCookedContent = null,
+            string AdditionalCookerOptions = null,
+            string BasedOnReleaseVersion = null,
+            string CreateReleaseVersion = null,
+            bool? GeneratePatch = null,
+            string DLCName = null,
+            bool? DLCIncludeEngineContent = null,
+            bool? NewCook = null,
+            bool? OldCook = null,
+			bool? CrashReporter = null,
+			bool? DedicatedServer = null,
+			bool? Client = null,
+			bool? Deploy = null,
+			bool? FileServer = null,
+			bool? Foreign = null,
+			bool? ForeignCode = null,
+			bool? LogWindow = null,
+			bool? NoCleanStage = null,
+			bool? NoClient = null,
+			bool? NoDebugInfo = null,
+			bool? NoXGE = null,
+			bool? Package = null,
+			bool? Pak = null,
+			bool? Prereqs = null,
+			bool? NoBootstrapExe = null,
+            bool? SignedPak = null,
+            bool? NullRHI = null,
+            bool? FakeClient = null,
+            bool? EditorTest = null,
+            bool? RunAutomationTests = null,
+            string RunAutomationTest = null,
+            int? CrashIndex = null,
+            bool? Rocket = null,
+			bool? SkipCook = null,
+			bool? SkipCookOnTheFly = null,
+			bool? SkipPak = null,
+			bool? SkipStage = null,
+			bool? Stage = null,
+			bool? Manifests = null,
+            bool? CreateChunkInstall = null,
+			bool? Unattended = null,
+			int? NumClients = null,
+			bool? Archive = null,
+			string ArchiveDirectoryParam = null,
+			bool? ArchiveMetaData = null,
+			ParamList<string> ProgramTargets = null,
+			bool? Distribution = null,
+            bool? Prebuilt = null,
+            int? RunTimeoutSeconds = null,
+			string SpecifiedArchitecture = null,
+            bool? IterativeDeploy = null,
+			bool? FastCook = null,
+			bool? IgnoreCookErrors = null
+			)
+		{
+			//
+			//// Use this.Name with properties and fields!
+			//
+
+			this.RawProjectPath = RawProjectPath;
+			if (DirectoriesToCook != null)
+			{
+				this.DirectoriesToCook = DirectoriesToCook;
+			}
+            this.InternationalizationPreset = ParseParamValueIfNotSpecified(Command, InternationalizationPreset, "i18npreset");
+            if (CulturesToCook != null)
+			{
+                this.CulturesToCook = CulturesToCook;
+			}
+			if (ClientCookedTargets != null)
+			{
+				this.ClientCookedTargets = ClientCookedTargets;
+			}
+			if (ServerCookedTargets != null)
+			{
+				this.ServerCookedTargets = ServerCookedTargets;
+			}
+			if (EditorTargets != null)
+			{
+				this.EditorTargets = EditorTargets;
+			}
+			if (ProgramTargets != null)
+			{
+				this.ProgramTargets = ProgramTargets;
+			}
+
+			// Parse command line params for client platforms "-TargetPlatform=Win64+Mac", "-Platform=Win64+Mac" and also "-Win64", "-Mac" etc.
+            if (ClientDependentPlatformMap != null)
+            {
+                this.ClientDependentPlatformMap = ClientDependentPlatformMap;
+            }
+
+			this.ClientTargetPlatforms = SetupTargetPlatforms(ref this.ClientDependentPlatformMap, Command, ClientTargetPlatforms, new ParamList<UnrealTargetPlatform>() { HostPlatform.Current.HostEditorPlatform }, true, "TargetPlatform", "Platform");
+
+            // Parse command line params for server platforms "-ServerTargetPlatform=Win64+Mac", "-ServerPlatform=Win64+Mac". "-Win64" etc is not allowed here
+            if (ServerDependentPlatformMap != null)
+            {
+                this.ServerDependentPlatformMap = ServerDependentPlatformMap;
+            }
+            this.ServerTargetPlatforms = SetupTargetPlatforms(ref this.ServerDependentPlatformMap, Command, ServerTargetPlatforms, this.ClientTargetPlatforms, false, "ServerTargetPlatform", "ServerPlatform");
+
+			this.Build = GetParamValueIfNotSpecified(Command, Build, this.Build, "build");
+			this.Run = GetParamValueIfNotSpecified(Command, Run, this.Run, "run");
+			this.Cook = GetParamValueIfNotSpecified(Command, Cook, this.Cook, "cook");
+			this.CookFlavor = ParseParamValueIfNotSpecified(Command, CookFlavor, "cookflavor", String.Empty);
+            this.NewCook = GetParamValueIfNotSpecified(Command, NewCook, this.NewCook, "NewCook");
+            this.OldCook = GetParamValueIfNotSpecified(Command, OldCook, this.OldCook, "OldCook");
+            this.CreateReleaseVersion = ParseParamValueIfNotSpecified(Command, CreateReleaseVersion, "createreleaseversion", String.Empty);
+            this.BasedOnReleaseVersion = ParseParamValueIfNotSpecified(Command, BasedOnReleaseVersion, "basedonreleaseversion", String.Empty);
+            this.GeneratePatch = GetParamValueIfNotSpecified(Command, GeneratePatch, this.GeneratePatch, "GeneratePatch");
+            this.AdditionalCookerOptions = ParseParamValueIfNotSpecified(Command, AdditionalCookerOptions, "AdditionalCookerOptions", String.Empty);
+            this.DLCName = ParseParamValueIfNotSpecified(Command, DLCName, "DLCName", String.Empty);
+            this.DLCIncludeEngineContent = GetParamValueIfNotSpecified(Command, DLCIncludeEngineContent, this.DLCIncludeEngineContent, "DLCIncludeEngineContent");
+			this.SkipCook = GetParamValueIfNotSpecified(Command, SkipCook, this.SkipCook, "skipcook");
+			if (this.SkipCook)
+			{
+				this.Cook = true;
+			}
+			this.Clean = GetOptionalParamValueIfNotSpecified(Command, Clean, this.Clean, "clean", null);
+			this.SignPak = ParseParamValueIfNotSpecified(Command, SignPak, "signpak", String.Empty);
+			this.SignedPak = !String.IsNullOrEmpty(this.SignPak) || GetParamValueIfNotSpecified(Command, SignedPak, this.SignedPak, "signedpak");
+			this.Pak = this.SignedPak || GetParamValueIfNotSpecified(Command, Pak, this.Pak, "pak");
+			this.SkipPak = GetParamValueIfNotSpecified(Command, SkipPak, this.SkipPak, "skippak");
+			if (this.SkipPak)
+			{
+				this.Pak = true;
+			}
+			this.NoXGE = GetParamValueIfNotSpecified(Command, NoXGE, this.NoXGE, "noxge");
+			this.CookOnTheFly = GetParamValueIfNotSpecified(Command, CookOnTheFly, this.CookOnTheFly, "cookonthefly");
+            if (this.CookOnTheFly && this.SkipCook)
+            {
+                this.Cook = false;
+            }
+            this.CookOnTheFlyStreaming = GetParamValueIfNotSpecified(Command, CookOnTheFlyStreaming, this.CookOnTheFlyStreaming, "cookontheflystreaming");
+            this.UnversionedCookedContent = GetParamValueIfNotSpecified(Command, UnversionedCookedContent, this.UnversionedCookedContent, "UnversionedCookedContent");
+            this.Compressed = GetParamValueIfNotSpecified(Command, Compressed, this.Compressed, "compressed");
+            this.UseDebugParamForEditorExe = GetParamValueIfNotSpecified(Command, UseDebugParamForEditorExe, this.UseDebugParamForEditorExe, "UseDebugParamForEditorExe");
+            this.IterativeCooking = GetParamValueIfNotSpecified(Command, IterativeCooking, this.IterativeCooking, new string[] { "iterativecooking", "iterate" } );
+			this.SkipCookOnTheFly = GetParamValueIfNotSpecified(Command, SkipCookOnTheFly, this.SkipCookOnTheFly, "skipcookonthefly");
+            this.CookAll = GetParamValueIfNotSpecified(Command, CookAll, this.CookAll, "CookAll");
+            this.CookMapsOnly = GetParamValueIfNotSpecified(Command, CookMapsOnly, this.CookMapsOnly, "CookMapsOnly");
+			this.FileServer = GetParamValueIfNotSpecified(Command, FileServer, this.FileServer, "fileserver");
+			this.DedicatedServer = GetParamValueIfNotSpecified(Command, DedicatedServer, this.DedicatedServer, "dedicatedserver", "server");
+			this.Client = GetParamValueIfNotSpecified(Command, Client, this.Client, "client");
+			if( this.Client )
+			{
+				this.DedicatedServer = true;
+			}
+			this.NoClient = GetParamValueIfNotSpecified(Command, NoClient, this.NoClient, "noclient");
+			this.LogWindow = GetParamValueIfNotSpecified(Command, LogWindow, this.LogWindow, "logwindow");
+			this.Stage = GetParamValueIfNotSpecified(Command, Stage, this.Stage, "stage");
+			this.SkipStage = GetParamValueIfNotSpecified(Command, SkipStage, this.SkipStage, "skipstage");
+			if (this.SkipStage)
+			{
+				this.Stage = true;
+			}
+			this.StageDirectoryParam = ParseParamValueIfNotSpecified(Command, StageDirectoryParam, "stagingdirectory", String.Empty, true);
+            this.StageNonMonolithic = GetParamValueIfNotSpecified(Command, StageNonMonolithic, this.StageNonMonolithic, "StageNonMonolithic");
+			this.Manifests = GetParamValueIfNotSpecified(Command, Manifests, this.Manifests, "manifests");
+            this.CreateChunkInstall = GetParamValueIfNotSpecified(Command, CreateChunkInstall, this.CreateChunkInstall, "createchunkinstall");
+			this.ChunkInstallDirectory = ParseParamValueIfNotSpecified(Command, ChunkInstallDirectory, "chunkinstalldirectory", String.Empty, true);
+			this.ChunkInstallVersionString = ParseParamValueIfNotSpecified(Command, ChunkInstallVersionString, "chunkinstallversion", String.Empty, true);
+			this.Archive = GetParamValueIfNotSpecified(Command, Archive, this.Archive, "archive");
+			this.ArchiveDirectoryParam = ParseParamValueIfNotSpecified(Command, ArchiveDirectoryParam, "archivedirectory", String.Empty, true);
+			this.ArchiveMetaData = GetParamValueIfNotSpecified(Command, ArchiveMetaData, this.ArchiveMetaData, "archivemetadata");
+			this.Distribution = GetParamValueIfNotSpecified(Command, Distribution, this.Distribution, "distribution");
+			this.Prereqs = GetParamValueIfNotSpecified(Command, Prereqs, this.Prereqs, "prereqs");
+			this.NoBootstrapExe = GetParamValueIfNotSpecified(Command, NoBootstrapExe, this.NoBootstrapExe, "nobootstrapexe");
+            this.Prebuilt = GetParamValueIfNotSpecified(Command, Prebuilt, this.Prebuilt, "prebuilt");
+            if (this.Prebuilt)
+            {
+                this.SkipCook = true;
+                /*this.SkipPak = true;
+                this.SkipStage = true;
+                this.Pak = true;
+                this.Stage = true;*/
+                this.Cook = true;
+                this.Archive = true;
+                
+                this.Deploy = true;
+                this.Run = true;
+                //this.StageDirectoryParam = this.PrebuiltDir;
+            }
+            this.NoDebugInfo = GetParamValueIfNotSpecified(Command, NoDebugInfo, this.NoDebugInfo, "nodebuginfo");
+			this.NoCleanStage = GetParamValueIfNotSpecified(Command, NoCleanStage, this.NoCleanStage, "nocleanstage");
+			this.MapToRun = ParseParamValueIfNotSpecified(Command, MapToRun, "map", String.Empty);
+			this.AdditionalServerMapParams = ParseParamValueIfNotSpecified(Command, AdditionalServerMapParams, "AdditionalServerMapParams", String.Empty);
+			this.Foreign = GetParamValueIfNotSpecified(Command, Foreign, this.Foreign, "foreign");
+			this.ForeignCode = GetParamValueIfNotSpecified(Command, ForeignCode, this.ForeignCode, "foreigncode");
+			this.StageCommandline = ParseParamValueIfNotSpecified(Command, StageCommandline, "cmdline");
+			this.BundleName = ParseParamValueIfNotSpecified(Command, BundleName, "bundlename");
+			this.RunCommandline = ParseParamValueIfNotSpecified(Command, RunCommandline, "addcmdline");
+			this.Package = GetParamValueIfNotSpecified(Command, Package, this.Package, "package");
+			this.Deploy = GetParamValueIfNotSpecified(Command, Deploy, this.Deploy, "deploy");
+			this.IterativeDeploy = GetParamValueIfNotSpecified(Command, IterativeDeploy, this.IterativeDeploy, new string[] {"iterativedeploy", "iterate" } );
+			this.FastCook = GetParamValueIfNotSpecified(Command, FastCook, this.FastCook, "FastCook");
+			this.IgnoreCookErrors = GetParamValueIfNotSpecified(Command, IgnoreCookErrors, this.IgnoreCookErrors, "IgnoreCookErrors");
+			this.Device = ParseParamValueIfNotSpecified(Command, Device, "device", String.Empty).Trim(new char[] { '\"' });
+
+			// strip the platform prefix the specified device.
+			if (this.Device.Contains("@"))
+			{
+				this.DeviceName = this.Device.Substring(this.Device.IndexOf("@") + 1);
+			}
+			else
+			{
+				this.DeviceName = this.Device;
+			}
+
+			this.ServerDevice = ParseParamValueIfNotSpecified(Command, ServerDevice, "serverdevice", this.Device);
+			this.NullRHI = GetParamValueIfNotSpecified(Command, NullRHI, this.NullRHI, "nullrhi");
+			this.FakeClient = GetParamValueIfNotSpecified(Command, FakeClient, this.FakeClient, "fakeclient");
+			this.EditorTest = GetParamValueIfNotSpecified(Command, EditorTest, this.EditorTest, "editortest");
+            this.RunAutomationTest = ParseParamValueIfNotSpecified(Command, RunAutomationTest, "RunAutomationTest");
+            this.RunAutomationTests = this.RunAutomationTest != "" || GetParamValueIfNotSpecified(Command, RunAutomationTests, this.RunAutomationTests, "RunAutomationTests");
+            this.SkipServer = GetParamValueIfNotSpecified(Command, SkipServer, this.SkipServer, "skipserver");
+			this.Rocket = GetParamValueIfNotSpecified(Command, Rocket, this.Rocket || GlobalCommandLine.Rocket, "rocket");
+			this.UE4Exe = ParseParamValueIfNotSpecified(Command, UE4Exe, "ue4exe", "UE4Editor-Cmd.exe");
+			this.Unattended = GetParamValueIfNotSpecified(Command, Unattended, this.Unattended, "unattended");
+			this.DeviceUsername = ParseParamValueIfNotSpecified(Command, DeviceUsername, "deviceuser", String.Empty);
+			this.DevicePassword = ParseParamValueIfNotSpecified(Command, DevicePassword, "devicepass", String.Empty);
+			this.CrashReporter = GetParamValueIfNotSpecified(Command, CrashReporter, this.CrashReporter, "crashreporter");
+			this.SpecifiedArchitecture = ParseParamValueIfNotSpecified(Command, SpecifiedArchitecture, "specifiedarchitecture", String.Empty);
+			if (ClientConfigsToBuild == null)
+			{
+				if (Command != null)
+				{
+					var ClientConfig = Command.ParseParamValue("clientconfig");
+					if (ClientConfig != null)
+					{
+						this.ClientConfigsToBuild = new List<UnrealTargetConfiguration>();
+						var Configs = new ParamList<string>(ClientConfig.Split('+'));
+						foreach (var ConfigName in Configs)
+						{
+							this.ClientConfigsToBuild.Add((UnrealTargetConfiguration)Enum.Parse(typeof(UnrealTargetConfiguration), ConfigName, true));
+						}
+					}
+				}
+			}
+			else
+			{
+				this.ClientConfigsToBuild = ClientConfigsToBuild;
+			}
+
+            if (Port == null)
+            {
+                if( Command != null )
+                {
+                    this.Port = new ParamList<string>();
+
+                    var PortString = Command.ParseParamValue("port");
+                    if (String.IsNullOrEmpty(PortString) == false)
+                    {
+                        var Ports = new ParamList<string>(PortString.Split('+'));
+                        foreach (var P in Ports)
+                        {
+                            this.Port.Add(P);
+                        }
+                    }
+                    
+                }
+            }
+            else
+            {
+                this.Port = Port;
+            }
+
+            if (MapsToCook == null)
+            {
+                if (Command != null)
+                {
+                    this.MapsToCook = new ParamList<string>();
+
+                    var MapsString = Command.ParseParamValue("MapsToCook");
+                    if (String.IsNullOrEmpty(MapsString) == false)
+                    {
+                        var MapNames = new ParamList<string>(MapsString.Split('+'));
+                        foreach ( var M in MapNames ) 
+                        {
+                            this.MapsToCook.Add( M );
+                        }
+                    }
+                }
+            }
+            else
+            {
+                this.MapsToCook = MapsToCook;
+            }
+
+            if (String.IsNullOrEmpty(this.MapToRun) == false)
+            {
+                this.MapsToCook.Add(this.MapToRun);
+            }
+            
+			if (ServerConfigsToBuild == null)
+			{
+				if (Command != null)
+				{
+					var ServerConfig = Command.ParseParamValue("serverconfig");
+					if (ServerConfig != null && ServerConfigsToBuild == null)
+					{
+						this.ServerConfigsToBuild = new List<UnrealTargetConfiguration>();
+						var Configs = new ParamList<string>(ServerConfig.Split('+'));
+						foreach (var ConfigName in Configs)
+						{
+							this.ServerConfigsToBuild.Add((UnrealTargetConfiguration)Enum.Parse(typeof(UnrealTargetConfiguration), ConfigName, true));
+						}
+					}
+				}
+			}
+			else
+			{
+				this.ServerConfigsToBuild = ServerConfigsToBuild;
+			}
+			if (NumClients.HasValue)
+			{
+				this.NumClients = NumClients.Value;
+			}
+			else if (Command != null)
+			{
+				this.NumClients = Command.ParseParamInt("numclients");
+			}
+            if (CrashIndex.HasValue)
+            {
+                this.CrashIndex = CrashIndex.Value;
+            }
+            else if (Command != null)
+            {
+                this.CrashIndex = Command.ParseParamInt("CrashIndex");
+            }
+            if (RunTimeoutSeconds.HasValue)
+            {
+                this.RunTimeoutSeconds = RunTimeoutSeconds.Value;
+            }
+            else if (Command != null)
+            {
+                this.RunTimeoutSeconds = Command.ParseParamInt("runtimeoutseconds");
+            }
+
+			AutodetectSettings(false);
+			ValidateAndLog();
+		}
+
+		#endregion
+
+		#region Shared
+
+		/// <summary>
+		/// Shared: Full path where the project exists (For uprojects this should include the uproj filename, otherwise just project folder)
+		/// </summary>
+		public string RawProjectPath { private set; get; }
+
+		/// <summary>
+		/// Shared: The current project is a foreign project, commandline: -foreign
+		/// </summary>
+		[Help("foreign", "Generate a foreign uproject from blankproject and use that")]
+		public bool Foreign { private set; get; }
+
+		/// <summary>
+		/// Shared: The current project is a foreign project, commandline: -foreign
+		/// </summary>
+		[Help("foreigncode", "Generate a foreign code uproject from platformergame and use that")]
+		public bool ForeignCode { private set; get; }
+
+		/// <summary>
+		/// Shared: true if we are running rocket
+		/// </summary>
+		[Help("Rocket", "true if we are running rocket")]
+		public bool Rocket { private set; get; }
+
+		/// <summary>
+		/// Shared: true if we are running rocket
+		/// </summary>
+		[Help("CrashReporter", "true if we should build crash reporter")]
+		public bool CrashReporter { private set; get; }
+
+		/// <summary>
+		/// Shared: Determines if the build is going to use cooked data, commandline: -cook, -cookonthefly
+		/// </summary>	
+		[Help("cook, -cookonthefly", "Determines if the build is going to use cooked data")]
+		public bool Cook { private set; get; }
+
+		/// <summary>
+		/// Shared: Determines if the build is going to use special sub-target platform, commandline: -cookflavor=ATC
+		/// </summary>	
+		[Help( "cookflavor", "Determines if the build is going to use special sub-target platform" )]
+		public string CookFlavor { private set; get; }
+
+		/// <summary>
+		/// Shared: Determines if the build is going to use cooked data, commandline: -cook, -cookonthefly
+		/// </summary>	
+		[Help("skipcook", "use a cooked build, but we assume the cooked data is up to date and where it belongs, implies -cook")]
+		public bool SkipCook { private set; get; }
+
+		/// <summary>
+		/// Shared: In a cookonthefly build, used solely to pass information to the package step. This is necessary because you can't set cookonthefly and cook at the same time, and skipcook sets cook.
+		/// </summary>	
+		[Help("skipcookonthefly", "in a cookonthefly build, used solely to pass information to the package step")]
+		public bool SkipCookOnTheFly { private set; get; }
+
+		/// <summary>
+		/// Shared: Determines if the intermediate folders will be wiped before building, commandline: -clean
+		/// </summary>
+		[Help("clean", "wipe intermediate folders before building")]
+		public bool? Clean { private set; get; }
+
+		/// <summary>
+		/// Shared: Assumes no user is sitting at the console, so for example kills clients automatically, commandline: -Unattended
+		/// </summary>
+		[Help("unattended", "assumes no operator is present, always terminates without waiting for something.")]
+		public bool Unattended { private set; get; }
+
+		/// <summary>
+        /// Shared: Sets platforms to build for non-dedicated servers. commandline: -TargetPlatform
+		/// </summary>
+		public List<UnrealTargetPlatform> ClientTargetPlatforms = new List<UnrealTargetPlatform>();
+
+        /// <summary>
+        /// Shared: Dictionary that maps client dependent platforms to "source" platforms that it should copy data from. commandline: -TargetPlatform=source.dependent
+        /// </summary>
+        public Dictionary<UnrealTargetPlatform, UnrealTargetPlatform> ClientDependentPlatformMap = new Dictionary<UnrealTargetPlatform, UnrealTargetPlatform>();
+
+		/// <summary>
+        /// Shared: Sets platforms to build for dedicated servers. commandline: -ServerTargetPlatform
+		/// </summary>
+		public List<UnrealTargetPlatform> ServerTargetPlatforms = new List<UnrealTargetPlatform>();
+
+        /// <summary>
+        /// Shared: Dictionary that maps server dependent platforms to "source" platforms that it should copy data from: -ServerTargetPlatform=source.dependent
+        /// </summary>
+        public Dictionary<UnrealTargetPlatform, UnrealTargetPlatform> ServerDependentPlatformMap = new Dictionary<UnrealTargetPlatform, UnrealTargetPlatform>();
+
+		/// <summary>
+		/// Shared: True if pak file should be generated.
+		/// </summary>
+		[Help("pak", "generate a pak file")]
+		public bool Pak { private set; get; }
+
+		/// <summary>
+		/// 
+		/// </summary>
+		public bool UsePak(Platform PlatformToCheck)
+		{
+			return Pak || PlatformToCheck.RequiresPak(this) == Platform.PakType.Always;
+		}
+
+		/// <summary>
+		/// Shared: Encryption keys used for signing the pak file.
+		/// </summary>
+		[Help("signpak=keys", "sign the generated pak file with the specified key, i.e. -signpak=C:\\Encryption.keys. Also implies -signedpak.")]
+		public string SignPak { private set; get; }
+
+		/// <summary>
+		/// Shared: the game will use only signed content.
+		/// </summary>
+		[Help("signed", "the game should expect to use a signed pak file.")]
+		public bool SignedPak { private set; get; }
+
+		/// <summary>
+		/// Shared: true if this build is staged, command line: -stage
+		/// </summary>
+		[Help("skippak", "use a pak file, but assume it is already built, implies pak")]
+		public bool SkipPak { private set; get; }
+
+		/// <summary>
+		/// Shared: true if this build is staged, command line: -stage
+		/// </summary>
+		[Help("stage", "put this build in a stage directory")]
+		public bool Stage { private set; get; }
+
+		/// <summary>
+		/// Shared: true if this build is staged, command line: -stage
+		/// </summary>
+		[Help("skipstage", "uses a stage directory, but assumes everything is already there, implies -stage")]
+		public bool SkipStage { private set; get; }
+
+		/// <summary>
+		/// Shared: true if this build is using streaming install manifests, command line: -manifests
+		/// </summary>
+		[Help("manifests", "generate streaming install manifests when cooking data")]
+		public bool Manifests { private set; get; }
+
+        /// <summary>
+        /// Shared: true if this build chunk install streaming install data, command line: -createchunkinstalldata
+        /// </summary>
+        [Help("createchunkinstall", "generate streaming install data from manifest when cooking data, requires -stage & -manifests")]
+        public bool CreateChunkInstall { private set; get; }
+
+		/// <summary>
+		/// Shared: Directory to use for built chunk install data, command line: -chunkinstalldirectory=
+		/// </summary>
+		public string ChunkInstallDirectory { set; get; }
+
+		/// <summary>
+		/// Shared: Version string to use for built chunk install data, command line: -chunkinstallversion=
+		/// </summary>
+		public string ChunkInstallVersionString { set; get; }
+
+		/// <summary>
+		/// Shared: Directory to copy the client to, command line: -stagingdirectory=
+		/// </summary>	
+		public string BaseStageDirectory
+		{
+			get
+			{
+                if( !String.IsNullOrEmpty(StageDirectoryParam ) )
+                {
+                    return Path.GetFullPath( StageDirectoryParam );
+                }
+                if ( HasDLCName )
+                {
+                     return Path.GetFullPath( CommandUtils.CombinePaths(Path.GetDirectoryName(RawProjectPath), "Plugins", DLCName, "Saved", "StagedBuilds" ) );
+                }
+                // default return the project saved\stagedbuilds directory
+                return Path.GetFullPath( CommandUtils.CombinePaths(Path.GetDirectoryName(RawProjectPath), "Saved", "StagedBuilds") );
+			}
+		}
+
+		[Help("stagingdirectory=Path", "Directory to copy the builds to, i.e. -stagingdirectory=C:\\Stage")]
+		public string StageDirectoryParam;
+
+        /// <summary>
+        /// Whether the project should use non monolithic staging
+        /// </summary>
+        public bool StageNonMonolithic;
+
+		[Help("ue4exe=ExecutableName", "Name of the UE4 Editor executable, i.e. -ue4exe=UE4Editor.exe")]
+		public string UE4Exe;
+
+		/// <summary>
+		/// Shared: true if this build is archived, command line: -archive
+		/// </summary>
+		[Help("archive", "put this build in an archive directory")]
+		public bool Archive { private set; get; }
+
+		/// <summary>
+		/// Shared: Directory to archive the client to, command line: -archivedirectory=
+		/// </summary>	
+		public string BaseArchiveDirectory
+		{
+			get
+			{
+                return Path.GetFullPath(String.IsNullOrEmpty(ArchiveDirectoryParam) ? CommandUtils.CombinePaths(Path.GetDirectoryName(RawProjectPath), "ArchivedBuilds") : ArchiveDirectoryParam);
+			}
+		}
+
+		[Help("archivedirectory=Path", "Directory to archive the builds to, i.e. -archivedirectory=C:\\Archive")]
+		public string ArchiveDirectoryParam;
+
+		/// <summary>
+		/// Whether the project should use non monolithic staging
+		/// </summary>
+		[Help("archivemetadata", "Archive extra metadata files in addition to the build (e.g. build.properties)")]
+		public bool ArchiveMetaData;
+
+		#endregion
+
+		#region Build
+
+		/// <summary>
+		/// Build: True if build step should be executed, command: -build
+		/// </summary>
+		[Help("build", "True if build step should be executed")]
+		public bool Build { private set; get; }
+
+		/// <summary>
+		/// Build: True if XGE should NOT be used for building.
+		/// </summary>
+		[Help("noxge", "True if XGE should NOT be used for building")]
+		public bool NoXGE { private set; get; }
+
+		/// <summary>
+		/// Build: List of maps to cook.
+		/// </summary>	
+		private ParamList<string> EditorTargetsList = null;
+		public ParamList<string> EditorTargets
+		{
+			set { EditorTargetsList = value; }
+			get
+			{
+				if (EditorTargetsList == null)
+				{
+					// Lazy auto-initialization
+					AutodetectSettings(false);
+				}
+				return EditorTargetsList;
+			}
+		}
+
+		/// <summary>
+		/// Build: List of maps to cook.
+		/// </summary>	
+		private ParamList<string> ProgramTargetsList = null;
+		public ParamList<string> ProgramTargets
+		{
+			set { ProgramTargetsList = value; }
+			get
+			{
+				if (ProgramTargetsList == null)
+				{
+					// Lazy auto-initialization
+					AutodetectSettings(false);
+				}
+				return ProgramTargetsList;
+			}
+		}
+
+		/// <summary>
+		/// Build: List of client configurations
+		/// </summary>
+		public List<UnrealTargetConfiguration> ClientConfigsToBuild = new List<UnrealTargetConfiguration>() { UnrealTargetConfiguration.Development };
+
+		///<summary>
+		/// Build: List of Server configurations
+		/// </summary>
+		public List<UnrealTargetConfiguration> ServerConfigsToBuild = new List<UnrealTargetConfiguration>() { UnrealTargetConfiguration.Development };
+
+		/// <summary>
+		/// Build: List of client cooked build targets.
+		/// </summary>
+		private ParamList<string> ClientCookedTargetsList = null;
+		public ParamList<string> ClientCookedTargets
+		{
+			set { ClientCookedTargetsList = value; }
+			get
+			{
+				if (ClientCookedTargetsList == null)
+				{
+					// Lazy auto-initialization
+					AutodetectSettings(false);
+				}
+				return ClientCookedTargetsList;
+			}
+		}
+
+		/// <summary>
+		/// Build: List of Server cooked build targets.
+		/// </summary>
+		private ParamList<string> ServerCookedTargetsList = null;
+		public ParamList<string> ServerCookedTargets
+		{
+			set { ServerCookedTargetsList = value; }
+			get
+			{
+				if (ServerCookedTargetsList == null)
+				{
+					// Lazy auto-initialization
+					AutodetectSettings(false);
+				}
+				return ServerCookedTargetsList;
+			}
+		}
+
+		#endregion
+
+		#region Cook
+
+		/// <summary>
+		/// Cook: List of maps to cook.
+		/// </summary>
+		public ParamList<string> MapsToCook = new ParamList<string>();
+
+		/// <summary>
+		/// Cook: List of directories to cook.
+		/// </summary>
+		public ParamList<string> DirectoriesToCook = new ParamList<string>();
+
+        /// <summary>
+        /// Cook: Internationalization preset to cook.
+        /// </summary>
+        public string InternationalizationPreset;
+
+        /// <summary>
+        /// Cook: Create a cooked release version
+        /// </summary>
+        public string CreateReleaseVersion;
+
+        /// <summary>
+        /// Cook: Use new cooker (temporary cooker option until it's enabled permanently)
+        /// </summary>
+        public bool NewCook { private set; get; }
+
+        /// <summary>
+        /// Cook: Use old cooker (temporary cooker options until new cook replaces it)
+        /// </summary>
+        public bool OldCook { private set; get; }
+
+        /// <summary>
+        /// Cook: Base this cook of a already released version of the cooked data
+        /// </summary>
+        public string BasedOnReleaseVersion;
+
+        /// <summary>
+        /// Are we generating a patch, generate a patch from a previously released version of the game (use CreateReleaseVersion to create a release). 
+        /// this requires BasedOnReleaseVersion
+        /// see also CreateReleaseVersion, BasedOnReleaseVersion
+        /// </summary>
+        public bool GeneratePatch;
+
+        /// <summary>
+        /// Name of dlc to cook and package (if this paramter is supplied cooks the dlc and packages it into the dlc directory)
+        /// </summary>
+        public string DLCName;
+
+        /// <summary>
+        /// Enable cooking of engine content when cooking dlc 
+        ///  not included in original release but is referenced by current cook
+        /// </summary>
+        public bool DLCIncludeEngineContent;
+
+        /// <summary>
+        /// Cook: Additional cooker options to include on the cooker commandline
+        /// </summary>
+        public string AdditionalCookerOptions;
+
+        /// <summary>
+        /// Cook: List of cultures to cook.
+        /// </summary>
+        public ParamList<string> CulturesToCook = new ParamList<string>();
+
+        /// <summary>
+        /// Compress packages during cook.
+        /// </summary>
+        public bool Compressed;
+
+        /// <summary>
+        /// put -debug on the editorexe commandline
+        /// </summary>
+        public bool UseDebugParamForEditorExe;
+
+        /// <summary>
+        /// 
+        /// </summary>
+        public bool UnversionedCookedContent;
+
+
+		/// <summary>
+		/// Cook: Uses the iterative cooking, command line: -iterativecooking or -iterate
+		/// </summary>
+		[Help( "iterativecooking", "Uses the iterative cooking, command line: -iterativecooking or -iterate" )]
+		public bool IterativeCooking;
+
+        /// <summary>
+        /// Cook: Only cook maps (and referenced content) instead of cooking everything only affects -cookall flag
+        /// </summary>
+        [Help("CookMapsOnly", "Cook only maps this only affects usage of -cookall the flag")]
+        public bool CookMapsOnly;
+
+        /// <summary>
+        /// Cook: Only cook maps (and referenced content) instead of cooking everything only affects cookall flag
+        /// </summary>
+        [Help("CookAll", "Cook all the things in the content directory for this project")]
+        public bool CookAll;
+
+
+		/// <summary>
+		/// Cook: Uses the iterative deploy, command line: -iterativedeploy or -iterate
+		/// </summary>
+		[Help("iterativecooking", "Uses the iterative cooking, command line: -iterativedeploy or -iterate")]
+		public bool IterativeDeploy;
+
+		[Help("FastCook", "Uses fast cook path if supported by target")]
+		public bool FastCook;
+
+		/// <summary>
+		/// Cook: Ignores cook errors and continues with packaging etc.
+		/// </summary>
+		[Help("IgnoreCookErrors", "Ignores cook errors and continues with packaging etc")]
+		public bool IgnoreCookErrors { private set; get; }
+
+		#endregion
+
+		#region Stage
+
+		/// <summary>
+		/// Stage: Commanndline: -nodebuginfo
+		/// </summary>
+		[Help("nodebuginfo", "do not copy debug files to the stage")]
+		public bool NoDebugInfo { private set; get; }
+
+		/// <summary>
+		/// true if the staging directory is to be cleaned: -cleanstage (also true if -clean is specified)
+		/// </summary>
+		[Help("nocleanstage", "skip cleaning the stage directory")]
+		public bool NoCleanStage { set { bNoCleanStage = value; } get { return SkipStage || bNoCleanStage; } }
+		private bool bNoCleanStage;
+
+		/// <summary>
+		/// Stage: If non-empty, the contents will be put into the stage
+		/// </summary>
+		[Help("cmdline", "command line to put into the stage in UE4CommandLine.txt")]
+		public string StageCommandline;
+
+        /// <summary>
+		/// Stage: If non-empty, the contents will be used for the bundle name
+		/// </summary>
+		[Help("bundlename", "string to use as the bundle name when deploying to mobile device")]
+        public string BundleName;
+
+        /// <summary>
+        /// Whether the project uses Steam (todo: substitute with more generic functionality)
+        /// </summary>
+        public bool bUsesSteam;
+
+        /// <summary>
+        /// Whether the project uses CEF3
+        /// </summary>
+        public bool bUsesCEF3;
+
+		/// <summary>
+		/// Whether the project uses visual Slate UI (as opposed to the low level windowing/messaging which is alway used)
+		/// </summary>
+		public bool bUsesSlate = true;
+
+		/// <summary>
+		/// Hack for legacy game styling isses.  No new project should ever set this to true
+		/// Whether the project uses the Slate editor style in game.  
+		/// </summary>
+		public bool bUsesSlateEditorStyle = false;
+
+        /// <summary>
+        /// By default we use the Release C++ Runtime (CRT), even when compiling Debug builds.  This is because the Debug C++
+        /// Runtime isn't very useful when debugging Unreal Engine projects, and linking against the Debug CRT libraries forces
+        /// our third party library dependencies to also be compiled using the Debug CRT (and often perform more slowly.)  Often
+        /// it can be inconvenient to require a separate copy of the debug versions of third party static libraries simply
+        /// so that you can debug your program's code.
+        /// </summary>
+        public bool bDebugBuildsActuallyUseDebugCRT = false;
+
+        /// <summary>
+        /// On Windows, adds an executable to the root of the staging directory which checks for prerequisites being 
+		/// installed and launches the game with a path to the .uproject file.
+		/// </summary>
+        public bool NoBootstrapExe { get; set; }
+
+		#endregion
+
+		#region Run
+
+		/// <summary>
+		/// Run: True if the Run step should be executed, command: -run
+		/// </summary>
+		[Help("run", "run the game after it is built (including server, if -server)")]
+		public bool Run { private set; get; }
+
+		/// <summary>
+		/// Run: The client runs with cooked data provided by cook on the fly server, command line: -cookonthefly
+		/// </summary>
+		[Help("cookonthefly", "run the client with cooked data provided by cook on the fly server")]
+		public bool CookOnTheFly { private set; get; }
+
+        /// <summary>
+        /// Run: The client should run in streaming mode when connecting to cook on the fly server
+        /// </summary>
+        [Help("Cookontheflystreaming", "run the client in streaming cook on the fly mode (don't cache files locally instead force reget from server each file load)")]
+        public bool CookOnTheFlyStreaming { private set; get; }
+
+		/// <summary>
+		/// Run: The client runs with cooked data provided by UnrealFileServer, command line: -fileserver
+		/// </summary>
+		[Help("fileserver", "run the client with cooked data provided by UnrealFileServer")]
+		public bool FileServer { private set; get; }
+
+		/// <summary>
+		/// Run: The client connects to dedicated server to get data, command line: -dedicatedserver
+		/// </summary>
+		[Help("dedicatedserver", "build, cook and run both a client and a server (also -server)")]
+		public bool DedicatedServer { private set; get; }
+
+		/// <summary>
+		/// Run: Uses a client target configuration, also implies -dedicatedserver, command line: -client
+		/// </summary>
+		[Help( "client", "build, cook and run a client and a server, uses client target configuration" )]
+		public bool Client { private set; get; }
+
+		/// <summary>
+		/// Run: Whether the client should start or not, command line (to disable): -noclient
+		/// </summary>
+		[Help("noclient", "do not run the client, just run the server")]
+		public bool NoClient { private set; get; }
+
+		/// <summary>
+		/// Run: Client should create its own log window, command line: -logwindow
+		/// </summary>
+		[Help("logwindow", "create a log window for the client")]
+		public bool LogWindow { private set; get; }
+
+		/// <summary>
+		/// Run: Map to run the game with.
+		/// </summary>
+		[Help("map", "map to run the game with")]
+		public string MapToRun;
+
+		/// <summary>
+		/// Run: Additional server map params.
+		/// </summary>
+		[Help("AdditionalServerMapParams", "Additional server map params, i.e ?param=value")]
+		public string AdditionalServerMapParams;
+
+		/// <summary>
+		/// Run: The target device to run the game on.  Comes in the form platform@devicename.
+		/// </summary>
+		[Help("device", "Device to run the game on")]
+		public string Device;
+
+		/// <summary>
+		/// Run: The target device to run the game on.  No platform prefix.
+		/// </summary>
+		[Help("device", "Device name without the platform prefix to run the game on")]
+		public string DeviceName;
+
+		/// <summary>
+		/// Run: the target device to run the server on
+		/// </summary>
+		[Help("serverdevice", "Device to run the server on")]
+		public string ServerDevice;
+
+		/// <summary>
+		/// Run: The indicated server has already been started
+		/// </summary>
+		[Help("skipserver", "Skip starting the server")]
+		public bool SkipServer;
+
+		/// <summary>
+		/// Run: The indicated server has already been started
+		/// </summary>
+		[Help("numclients=n", "Start extra clients, n should be 2 or more")]
+		public int NumClients;
+
+		/// <summary>
+		/// Run: Additional command line arguments to pass to the program
+		/// </summary>
+		[Help("addcmdline", "Additional command line arguments for the program")]
+		public string RunCommandline;
+
+        /// <summary>
+        /// Run:adds -nullrhi to the client commandline
+        /// </summary>
+        [Help("nullrhi", "add -nullrhi to the client commandlines")]
+        public bool NullRHI;
+
+        /// <summary>
+        /// Run:adds ?fake to the server URL
+        /// </summary>
+        [Help("fakeclient", "adds ?fake to the server URL")]
+        public bool FakeClient;
+
+        /// <summary>
+        /// Run:adds ?fake to the server URL
+        /// </summary>
+        [Help("editortest", "rather than running a client, run the editor instead")]
+        public bool EditorTest;
+
+        /// <summary>
+        /// Run:when running -editortest or a client, run all automation tests, not compatible with -server
+        /// </summary>
+        [Help("RunAutomationTests", "when running -editortest or a client, run all automation tests, not compatible with -server")]
+        public bool RunAutomationTests;
+
+        /// <summary>
+        /// Run:when running -editortest or a client, run all automation tests, not compatible with -server
+        /// </summary>
+        [Help("RunAutomationTests", "when running -editortest or a client, run a specific automation tests, not compatible with -server")]
+        public string RunAutomationTest;
+
+        /// <summary>
+        /// Run: Adds commands like debug crash, debug rendercrash, etc based on index
+        /// </summary>
+        [Help("Crash=index", "when running -editortest or a client, adds commands like debug crash, debug rendercrash, etc based on index")]
+        public int CrashIndex;
+
+        public ParamList<string> Port;
+
+        /// <summary>
+        /// Run: Linux username for unattended key genereation
+        /// </summary>
+        [Help("deviceuser", "Linux username for unattended key genereation")]
+        public string DeviceUsername;
+
+        /// <summary>
+        /// Run: Linux password for unattended key genereation
+        /// </summary>
+        [Help("devicepass", "Linux password for unattended key genereation")]
+        public string DevicePassword;
+
+        /// <summary>
+        /// Run: Sever device IP address
+        /// </summary>
+        public string ServerDeviceAddress;
+
+        #endregion
+
+		#region Package
+
+		[Help("package", "package the project for the target platform")]
+		public bool Package { get; set; }
+
+		[Help("distribution", "package for distribution the project")]
+		public bool Distribution { get; set; }
+
+		[Help("prereqs", "stage prerequisites along with the project")]
+		public bool Prereqs { get; set; }
+
+        [Help("Prebuilt", "this is a prebuilt cooked and packaged build")]
+        public bool Prebuilt { get; private set; }
+
+        [Help("RunTimeoutSeconds", "timeout to wait after we lunch the game")]
+        public int RunTimeoutSeconds;
+
+		[Help("SpecifiedArchitecture", "Determine a specific Minimum OS")]
+		public string SpecifiedArchitecture;
+
+		#endregion
+
+		#region Deploy
+
+		[Help("deploy", "deploy the project for the target platform")]
+		public bool Deploy { get; set; }
+
+		#endregion
+
+		#region Initialization
+
+		private Dictionary<TargetRules.TargetType, SingleTargetProperties> DetectedTargets;
+		private Dictionary<UnrealTargetPlatform, ConfigCacheIni> LoadedEngineConfigs;
+		private Dictionary<UnrealTargetPlatform, ConfigCacheIni> LoadedGameConfigs;
+
+		private void AutodetectSettings(bool bReset)
+		{
+			if (bReset)
+			{
+				EditorTargetsList = null;
+				ClientCookedTargetsList = null;
+				ServerCookedTargetsList = null;
+				ProgramTargetsList = null;
+				ProjectBinariesPath = null;
+				ProjectGameExePath = null;
+			}
+
+			var Properties = ProjectUtils.GetProjectProperties(RawProjectPath, ClientTargetPlatforms);
+
+			bUsesSteam = Properties.bUsesSteam;
+			bUsesCEF3 = Properties.bUsesCEF3;
+			bUsesSlate = Properties.bUsesSlate;
+			bUsesSlateEditorStyle = Properties.bUsesSlateEditorStyle;
+            bDebugBuildsActuallyUseDebugCRT = Properties.bDebugBuildsActuallyUseDebugCRT;
+
+			bIsCodeBasedProject = Properties.bIsCodeBasedProject;			
+			DetectedTargets = Properties.Targets;
+			LoadedEngineConfigs = Properties.EngineConfigs;
+			LoadedGameConfigs = Properties.GameConfigs;
+
+			var GameTarget = String.Empty;
+			var EditorTarget = String.Empty;
+			var ServerTarget = String.Empty;
+			var ProgramTarget = String.Empty;
+			var ProjectType = TargetRules.TargetType.Game;
+
+			if (GlobalCommandLine.Rocket)
+			{
+				if (!CommandUtils.CmdEnv.HasCapabilityToCompile || !bIsCodeBasedProject)
+				{
+					if (bIsCodeBasedProject)
+					{
+						var ShortName = ProjectUtils.GetShortProjectName(RawProjectPath);
+						GameTarget = ShortName;
+						EditorTarget = ShortName + "Editor";
+						ServerTarget = ShortName + "Server";
+					}
+					else
+					{
+						GameTarget = "UE4Game";
+						EditorTarget = "UE4Editor";
+						//ServerTarget = "RocketServer";
+
+						Build = false;
+					}
+				}
+				else
+				{
+					SingleTargetProperties TargetData;
+					if (DetectedTargets.TryGetValue(TargetRules.TargetType.Editor, out TargetData))
+					{
+						EditorTarget = TargetData.TargetName;
+					}
+
+					if (DetectedTargets.TryGetValue(TargetRules.TargetType.Program, out TargetData))
+					{
+						ProgramTarget = TargetData.TargetName;
+					}
+					//DetectedTargets.TryGetValue(TargetRules.TargetType.Server, out ServerTarget);
+
+					if (string.IsNullOrEmpty(GameTarget))
+					{
+						if (DetectedTargets.TryGetValue(TargetRules.TargetType.Game, out TargetData))
+						{
+							GameTarget = TargetData.TargetName;
+						}
+					}
+				}
+			}
+			else if (!bIsCodeBasedProject)
+			{
+				GameTarget = "UE4Game";
+				EditorTarget = "UE4Editor";
+				ServerTarget = "UE4Server";
+			}
+			else if (!CommandUtils.CmdEnv.HasCapabilityToCompile)
+			{
+				var ShortName = ProjectUtils.GetShortProjectName(RawProjectPath);
+				GameTarget = ShortName;
+				EditorTarget = ShortName + "Editor";
+				ServerTarget = ShortName + "Server";
+			}
+			else if (!CommandUtils.IsNullOrEmpty(Properties.Targets))
+			{
+				SingleTargetProperties TargetData;
+
+				var GameTargetType = TargetRules.TargetType.Game;
+				
+				if( Client )
+				{
+					if( HasClientTargetDetected )
+					{
+						GameTargetType = TargetRules.TargetType.Client;
+					}
+					else
+					{
+						throw new AutomationException( "Client target not found!" );
+					}
+				}
+
+				var ValidGameTargetTypes = new TargetRules.TargetType[]
+				{
+					GameTargetType,
+					TargetRules.TargetType.Program		
+				};
+
+				foreach (var ValidTarget in ValidGameTargetTypes)
+				{
+					if (DetectedTargets.TryGetValue(ValidTarget, out TargetData))
+					{
+						GameTarget = TargetData.TargetName;
+                        bDebugBuildsActuallyUseDebugCRT = TargetData.Rules.bDebugBuildsActuallyUseDebugCRT;
+						bUsesSlate = TargetData.Rules.bUsesSlate;
+						bUsesSlateEditorStyle = TargetData.Rules.bUsesSlateEditorStyle;
+						bUsesSteam = TargetData.Rules.bUsesSteam;
+						bUsesCEF3 = TargetData.Rules.bUsesCEF3;
+						ProjectType = ValidTarget;
+						break;
+					}
+				}
+
+				if (DetectedTargets.TryGetValue(TargetRules.TargetType.Editor, out TargetData))
+				{
+					EditorTarget = TargetData.TargetName;
+				}
+				if (DetectedTargets.TryGetValue(TargetRules.TargetType.Server, out TargetData))
+				{
+					ServerTarget = TargetData.TargetName;
+				}
+				if (DetectedTargets.TryGetValue(TargetRules.TargetType.Program, out TargetData))
+				{
+					ProgramTarget = TargetData.TargetName;
+				}
+			}
+			else if (!CommandUtils.IsNullOrEmpty(Properties.Programs))
+			{
+				SingleTargetProperties TargetData = Properties.Programs[0];
+
+				bDebugBuildsActuallyUseDebugCRT = TargetData.Rules.bDebugBuildsActuallyUseDebugCRT;
+				bUsesSlate = TargetData.Rules.bUsesSlate;
+				bUsesSlateEditorStyle = TargetData.Rules.bUsesSlateEditorStyle;
+				bUsesSteam = TargetData.Rules.bUsesSteam;
+				bUsesCEF3 = TargetData.Rules.bUsesCEF3;
+				ProjectType = TargetRules.TargetType.Program;
+				ProgramTarget = TargetData.TargetName;
+				GameTarget = TargetData.TargetName;
+			}
+			else if (!this.Build)
+			{
+				var ShortName = ProjectUtils.GetShortProjectName(RawProjectPath);
+				GameTarget = ShortName;
+				EditorTarget = ShortName + "Editor";
+				ServerTarget = ShortName + "Server";
+			}
+			else
+			{
+				throw new AutomationException("{0} does not look like uproject file but no targets have been found!", RawProjectPath);
+			}
+
+			IsProgramTarget = ProjectType == TargetRules.TargetType.Program;
+
+			if (String.IsNullOrEmpty(EditorTarget) && ProjectType != TargetRules.TargetType.Program && CommandUtils.IsNullOrEmpty(EditorTargetsList) && !Rocket)
+			{
+				if (Properties.bWasGenerated)
+				{
+					EditorTarget = "UE4Editor";
+				}
+				else
+				{
+					throw new AutomationException("Editor target not found!");
+				}
+			}
+			if (String.IsNullOrEmpty(GameTarget) && Run && !NoClient && (Cook || CookOnTheFly) && CommandUtils.IsNullOrEmpty(ClientCookedTargetsList) && !Rocket)
+			{
+				throw new AutomationException("Game target not found. Game target is required with -cook or -cookonthefly");
+			}
+
+			if (EditorTargetsList == null)
+			{
+				if (!GlobalCommandLine.NoCompile && !GlobalCommandLine.NoCompileEditor && (ProjectType != TargetRules.TargetType.Program) && !String.IsNullOrEmpty(EditorTarget))
+				{
+					EditorTargetsList = new ParamList<string>(EditorTarget);
+				}
+				else
+				{
+					EditorTargetsList = new ParamList<string>();
+				}
+			}
+
+			if (ProgramTargetsList == null)
+			{
+				if (ProjectType == TargetRules.TargetType.Program)
+				{
+					ProgramTargetsList = new ParamList<string>(ProgramTarget);
+				}
+				else
+				{
+					ProgramTargetsList = new ParamList<string>();
+				}
+			}
+
+            if (ClientCookedTargetsList == null && !NoClient && (Cook || CookOnTheFly || Prebuilt))
+			{
+                if (String.IsNullOrEmpty(GameTarget))
+				{
+                    throw new AutomationException("Game target not found. Game target is required with -cook or -cookonthefly");
+                }
+				else
+				{
+                    ClientCookedTargetsList = new ParamList<string>(GameTarget);
+                }
+			}
+            else if (ClientCookedTargetsList == null)
+            {
+                ClientCookedTargetsList = new ParamList<string>();
+            }
+
+            if (ServerCookedTargetsList == null && DedicatedServer && (Cook || CookOnTheFly))
+			{
+				if (String.IsNullOrEmpty(ServerTarget))
+				{
+                    throw new AutomationException("Server target not found. Server target is required with -server and -cook or -cookonthefly");
+				}
+				ServerCookedTargetsList = new ParamList<string>(ServerTarget);
+			}
+			else if (ServerCookedTargetsList == null)
+			{
+				ServerCookedTargetsList = new ParamList<string>();
+			}
+
+			if (String.IsNullOrEmpty(ProjectBinariesPath) || String.IsNullOrEmpty(ProjectGameExePath))
+			{
+				if ( ClientTargetPlatforms.Count > 0 )
+				{
+					var ProjectClientBinariesPath = ProjectUtils.GetClientProjectBinariesRootPath(RawProjectPath, ProjectType, Properties.bIsCodeBasedProject);
+					ProjectBinariesPath = ProjectUtils.GetProjectClientBinariesFolder(ProjectClientBinariesPath, ClientTargetPlatforms[0]);
+					ProjectGameExePath = CommandUtils.CombinePaths(ProjectBinariesPath, GameTarget + Platform.GetExeExtension(ClientTargetPlatforms[0]));
+				}
+			}
+
+			// for the moment, allow the commandline to override the ini if set.  This will keep us from breaking licensee packaging scripts until we have
+			// the full solution for per-platform packaging settings.
+			if (!Manifests)
+			{				
+				ConfigCacheIni GameIni = new ConfigCacheIni("Game", Path.GetDirectoryName(RawProjectPath));
+				String IniPath = "/Script/UnrealEd.ProjectPackagingSettings";
+				bool bSetting = false;
+				if (!GameIni.GetBool(IniPath, "bGenerateChunks", out bSetting))
+				{
+					Manifests = bSetting;
+				}
+			}
+		}
+
+		#endregion
+
+		#region Utilities
+
+		public bool HasEditorTargets
+		{
+			get { return !CommandUtils.IsNullOrEmpty(EditorTargets); }
+		}
+
+		public bool HasCookedTargets
+		{
+			get { return !CommandUtils.IsNullOrEmpty(ClientCookedTargets) || !CommandUtils.IsNullOrEmpty(ServerCookedTargets); }
+		}
+
+		public bool HasServerCookedTargets
+		{
+			get { return !CommandUtils.IsNullOrEmpty(ServerCookedTargets); }
+		}
+
+		public bool HasClientCookedTargets
+		{
+			get { return !CommandUtils.IsNullOrEmpty(ClientCookedTargets); }
+		}
+
+		public bool HasProgramTargets
+		{
+			get { return !CommandUtils.IsNullOrEmpty(ProgramTargets); }
+		}
+
+		public bool HasMapsToCook
+		{
+			get { return !CommandUtils.IsNullOrEmpty(MapsToCook); }
+		}
+
+		public bool HasDirectoriesToCook
+		{
+			get { return !CommandUtils.IsNullOrEmpty(DirectoriesToCook); }
+		}
+
+        public bool HasInternationalizationPreset
+        {
+            get { return !String.IsNullOrEmpty(InternationalizationPreset); }
+        }
+
+        public bool HasBasedOnReleaseVersion
+        {
+            get { return !String.IsNullOrEmpty(BasedOnReleaseVersion); }
+        }
+
+        public bool HasAdditionalCookerOptions
+        {
+            get { return !String.IsNullOrEmpty(AdditionalCookerOptions); }
+        }
+
+        public bool HasDLCName
+        {
+            get { return !String.IsNullOrEmpty(DLCName); }
+        }
+
+        public bool HasCreateReleaseVersion
+        {
+            get { return !String.IsNullOrEmpty(CreateReleaseVersion); }
+        }
+
+        public bool HasCulturesToCook
+        {
+            get { return !CommandUtils.IsNullOrEmpty(CulturesToCook); }
+        }
+
+		public bool HasGameTargetDetected
+		{
+			get { return ProjectTargets.ContainsKey(TargetRules.TargetType.Game); }
+		}
+
+		public bool HasClientTargetDetected
+		{
+			get { return ProjectTargets.ContainsKey( TargetRules.TargetType.Client ); }
+		}
+
+		public bool HasDedicatedServerAndClient
+		{
+			get { return Client && DedicatedServer; }
+		}
+
+		/// <summary>
+		/// Project name (name of the uproject file without extension or directory name where the project is localed)
+		/// </summary>
+		public string ShortProjectName
+		{
+			get { return ProjectUtils.GetShortProjectName(RawProjectPath); }
+		}
+
+  		/// <summary>
+		/// True if this project contains source code.
+		/// </summary>	
+		public bool IsCodeBasedProject
+		{
+			get
+			{
+				return bIsCodeBasedProject;
+			}
+		}
+		private bool bIsCodeBasedProject;
+
+		public string CodeBasedUprojectPath
+		{
+            get { return IsCodeBasedProject ? RawProjectPath : null; }
+		}
+		/// <summary>
+		/// True if this project is a program.
+		/// </summary>
+		public bool IsProgramTarget { get; private set; }
+
+		/// <summary>
+		/// Path where the project's game (or program) binaries are built.
+		/// </summary>
+		public string ProjectBinariesFolder
+		{
+			get
+			{
+				if (String.IsNullOrEmpty(ProjectBinariesPath))
+				{
+					AutodetectSettings(false);
+				}
+				return ProjectBinariesPath;
+			}
+		}
+		private string ProjectBinariesPath;
+
+        /// <summary>
+        /// True if we are generating a patch
+        /// </summary>
+        public bool IsGeneratingPatch
+        {
+            get { return GeneratePatch; }
+        }
+
+		/// <summary>
+		/// Filename of the target game exe (or program exe).
+		/// </summary>
+		public string ProjectGameExeFilename
+		{
+			get
+			{
+				if (String.IsNullOrEmpty(ProjectGameExePath))
+				{
+					AutodetectSettings(false);
+				}
+				return ProjectGameExePath;
+			}
+		}
+		private string ProjectGameExePath;
+
+		public Platform GetTargetPlatformInstance(UnrealTargetPlatform TargetPlatformType)
+		{
+			return Platform.Platforms[TargetPlatformType];
+		}
+
+		public List<Platform> ClientTargetPlatformInstances
+		{
+			get
+			{
+				List<Platform> ClientPlatformInstances = new List<Platform>();
+				foreach ( var ClientPlatform in ClientTargetPlatforms )
+				{
+					ClientPlatformInstances.Add(Platform.Platforms[ClientPlatform]);
+				}
+				return ClientPlatformInstances;
+			}
+		}
+
+        public UnrealTargetPlatform GetCookedDataPlatformForClientTarget(UnrealTargetPlatform TargetPlatformType)
+        {
+            if (ClientDependentPlatformMap.ContainsKey(TargetPlatformType))
+            {
+                return ClientDependentPlatformMap[TargetPlatformType];
+            }
+            return TargetPlatformType;
+        }
+
+		public List<Platform> ServerTargetPlatformInstances
+		{
+			get
+			{
+				List<Platform> ServerPlatformInstances = new List<Platform>();
+				foreach (var ServerPlatform in ServerTargetPlatforms)
+				{
+					ServerPlatformInstances.Add(Platform.Platforms[ServerPlatform]);
+				}
+				return ServerPlatformInstances;
+			}
+		}
+
+        public UnrealTargetPlatform GetCookedDataPlatformForServerTarget(UnrealTargetPlatform TargetPlatformType)
+        {
+            if (ServerDependentPlatformMap.ContainsKey(TargetPlatformType))
+            {
+                return ServerDependentPlatformMap[TargetPlatformType];
+            }
+            return TargetPlatformType;
+        }
+
+		/// <summary>
+		/// All auto-detected targets for this project
+		/// </summary>
+		public Dictionary<TargetRules.TargetType, SingleTargetProperties> ProjectTargets
+		{
+			get
+			{
+				if (DetectedTargets == null)
+				{
+					AutodetectSettings(false);
+				}
+				return DetectedTargets;
+			}
+		}
+
+		/// <summary>
+		/// List of all Engine ini files for this project
+		/// </summary>
+		public Dictionary<UnrealTargetPlatform, ConfigCacheIni> EngineConfigs
+		{
+			get
+			{
+				if (LoadedEngineConfigs == null)
+				{
+					AutodetectSettings(false);
+				}
+				return LoadedEngineConfigs;
+			}
+		}
+
+		/// <summary>
+		/// List of all Game ini files for this project
+		/// </summary>
+		public Dictionary<UnrealTargetPlatform, ConfigCacheIni> GameConfigs
+		{
+			get
+			{
+				if (LoadedGameConfigs == null)
+				{
+					AutodetectSettings(false);
+				}
+				return LoadedGameConfigs;
+			}
+		}
+
+		public void Validate()
+		{
+			if (String.IsNullOrEmpty(RawProjectPath))
+			{
+				throw new AutomationException("RawProjectPath can't be empty.");
+			}
+            if (!RawProjectPath.EndsWith(".uproject", StringComparison.InvariantCultureIgnoreCase))
+            {
+                throw new AutomationException("RawProjectPath {0} must end with .uproject", RawProjectPath);
+            }
+            if (!CommandUtils.FileExists_NoExceptions(RawProjectPath))
+            {
+                throw new AutomationException("RawProjectPath {0} file must exist", RawProjectPath);
+            }
+
+			if (FileServer && !Cook)
+			{
+				throw new AutomationException("Only cooked builds can use a fileserver be staged, use -cook");
+			}
+
+			if (Stage && !Cook && !CookOnTheFly && !IsProgramTarget)
+			{
+				throw new AutomationException("Only cooked builds or programs can be staged, use -cook or -cookonthefly.");
+			}
+
+			if (Manifests && !Cook && !Stage && !Pak)
+			{
+				throw new AutomationException("Only staged pakd and cooked builds can generate streaming install manifests");
+			}
+
+			if (Pak && !Stage)
+			{
+				throw new AutomationException("Only staged builds can be paked, use -stage or -skipstage.");
+			}
+
+            if (Deploy && !Stage)
+            {
+                throw new AutomationException("Only staged builds can be deployed, use -stage or -skipstage.");
+            }
+
+            if ((Pak || Stage || Cook || CookOnTheFly || FileServer || DedicatedServer) && EditorTest)
+            {
+                throw new AutomationException("None of pak, stage, cook, CookOnTheFly or DedicatedServer can be used with EditorTest");
+            }
+
+            if (DedicatedServer && RunAutomationTests)
+            {
+                throw new AutomationException("DedicatedServer cannot be used with RunAutomationTests");
+            }
+
+			if ((CookOnTheFly || FileServer) && DedicatedServer)
+			{
+				throw new AutomationException("Don't use either -cookonthefly or -fileserver with -server.");
+			}
+
+			if (NoClient && !DedicatedServer && !CookOnTheFly)
+			{
+				throw new AutomationException("-noclient can only be used with -server or -cookonthefly.");
+			}
+
+			if (Build && !HasCookedTargets && !HasEditorTargets && !HasProgramTargets)
+			{
+				throw new AutomationException("-build is specified but there are no targets to build.");
+			}
+
+			if (Pak && FileServer)
+			{
+				throw new AutomationException("Can't use -pak and -fileserver at the same time.");
+			}
+
+			if (Cook && CookOnTheFly)
+			{
+				throw new AutomationException("Can't use both -cook and -cookonthefly.");
+			}
+
+            if (!HasDLCName && DLCIncludeEngineContent)
+            {
+                throw new AutomationException("DLCIncludeEngineContent flag is only valid when cooking dlc.");
+            }
+
+            if ((IsGeneratingPatch || HasDLCName) && !HasBasedOnReleaseVersion)
+            {
+                throw new AutomationException("Require based on release version to build patches or dlc");
+            }
+
+            if (HasCreateReleaseVersion && (IsGeneratingPatch || HasDLCName))
+            {
+                throw new AutomationException("Can't create a release version at the same time as creating dlc.");
+            }
+
+            if (HasBasedOnReleaseVersion && (IterativeCooking || IterativeDeploy))
+            {
+                throw new AutomationException("Can't use iterative cooking / deploy on dlc or patching or creating a release");
+            }
+
+            /*if (Compressed && !Pak)
+            {
+                throw new AutomationException("-compressed can only be used with -pak");
+            }*/
+
+            if (CreateChunkInstall && (!Manifests || !Stage))
+            {
+                throw new AutomationException("-createchunkinstall can only be used with -manifests & -stage"); 
+            }
+
+			if (CreateChunkInstall && String.IsNullOrEmpty(ChunkInstallDirectory))
+			{
+				throw new AutomationException("-createchunkinstall must specify the chunk install data directory with -chunkinstalldirectory=");
+			}
+
+			if (CreateChunkInstall && String.IsNullOrEmpty(ChunkInstallVersionString))
+			{
+				throw new AutomationException("-createchunkinstall must specify the chunk install data version string with -chunkinstallversion=");
+			}
+		}
+
+		protected bool bLogged = false;
+		public virtual void ValidateAndLog()
+		{
+			// Avoid spamming, log only once
+			if (!bLogged)
+			{
+				// In alphabetical order.
+				CommandUtils.Log("Project Params **************");
+
+				CommandUtils.Log("AdditionalServerMapParams={0}", AdditionalServerMapParams);
+				CommandUtils.Log("Archive={0}", Archive);
+				CommandUtils.Log("ArchiveMetaData={0}", ArchiveMetaData);
+				CommandUtils.Log("BaseArchiveDirectory={0}", BaseArchiveDirectory);
+				CommandUtils.Log("BaseStageDirectory={0}", BaseStageDirectory);
+				CommandUtils.Log("Build={0}", Build);
+				CommandUtils.Log("Cook={0}", Cook);
+				CommandUtils.Log("Clean={0}", Clean);
+				CommandUtils.Log("Client={0}", Client);
+				CommandUtils.Log("ClientConfigsToBuild={0}", string.Join(",", ClientConfigsToBuild));
+				CommandUtils.Log("ClientCookedTargets={0}", ClientCookedTargets.ToString());
+				CommandUtils.Log("ClientTargetPlatform={0}", string.Join(",", ClientTargetPlatforms));
+                CommandUtils.Log("Compressed={0}", Compressed);
+                CommandUtils.Log("UseDebugParamForEditorExe={0}", UseDebugParamForEditorExe);
+				CommandUtils.Log("CookFlavor={0}", CookFlavor);
+				CommandUtils.Log("CookOnTheFly={0}", CookOnTheFly);
+                CommandUtils.Log("CookOnTheFlyStreaming={0}", CookOnTheFlyStreaming);
+                CommandUtils.Log("UnversionedCookedContent={0}", UnversionedCookedContent);
+                CommandUtils.Log("GeneratePatch={0}", GeneratePatch);
+                CommandUtils.Log("CreateReleaseVersion={0}", CreateReleaseVersion);
+                CommandUtils.Log("BasedOnReleaseVersion={0}", BasedOnReleaseVersion);
+                CommandUtils.Log("DLCName={0}", DLCName);
+                CommandUtils.Log("DLCIncludeEngineContent={0}", DLCIncludeEngineContent);
+                CommandUtils.Log("AdditionalCookerOptions={0}", AdditionalCookerOptions);
+				CommandUtils.Log("DedicatedServer={0}", DedicatedServer);
+				CommandUtils.Log("DirectoriesToCook={0}", DirectoriesToCook.ToString());
+                CommandUtils.Log("CulturesToCook={0}", CulturesToCook.ToString());
+				CommandUtils.Log("EditorTargets={0}", EditorTargets.ToString());
+				CommandUtils.Log("Foreign={0}", Foreign);
+				CommandUtils.Log("IsCodeBasedProject={0}", IsCodeBasedProject.ToString());
+				CommandUtils.Log("IsProgramTarget={0}", IsProgramTarget.ToString());
+				CommandUtils.Log("IterativeCooking={0}", IterativeCooking);
+                CommandUtils.Log("CookAll={0}", CookAll);
+                CommandUtils.Log("CookMapsOnly={0}", CookMapsOnly);
+                CommandUtils.Log("Deploy={0}", Deploy);
+				CommandUtils.Log("IterativeDeploy={0}", IterativeDeploy);
+				CommandUtils.Log("FastCook={0}", FastCook);
+				CommandUtils.Log("LogWindow={0}", LogWindow);
+				CommandUtils.Log("Manifests={0}", Manifests);
+				CommandUtils.Log("MapToRun={0}", MapToRun);
+				CommandUtils.Log("NoClient={0}", NoClient);
+				CommandUtils.Log("NumClients={0}", NumClients);                
+				CommandUtils.Log("NoDebugInfo={0}", NoDebugInfo);
+				CommandUtils.Log("NoCleanStage={0}", NoCleanStage);
+				CommandUtils.Log("NoXGE={0}", NoXGE);
+				CommandUtils.Log("MapsToCook={0}", MapsToCook.ToString());
+				CommandUtils.Log("Pak={0}", Pak);
+				CommandUtils.Log("Package={0}", Package);
+				CommandUtils.Log("NullRHI={0}", NullRHI);
+				CommandUtils.Log("FakeClient={0}", FakeClient);
+                CommandUtils.Log("EditorTest={0}", EditorTest);
+                CommandUtils.Log("RunAutomationTests={0}", RunAutomationTests); 
+                CommandUtils.Log("RunAutomationTest={0}", RunAutomationTest);
+                CommandUtils.Log("RunTimeoutSeconds={0}", RunTimeoutSeconds);
+                CommandUtils.Log("CrashIndex={0}", CrashIndex);
+				CommandUtils.Log("ProgramTargets={0}", ProgramTargets.ToString());
+                CommandUtils.Log("ProjectBinariesFolder={0}", ProjectBinariesFolder);
+				CommandUtils.Log("ProjectBinariesPath={0}", ProjectBinariesPath);
+				CommandUtils.Log("ProjectGameExeFilename={0}", ProjectGameExeFilename);
+				CommandUtils.Log("ProjectGameExePath={0}", ProjectGameExePath);
+				CommandUtils.Log("Distribution={0}", Distribution);
+                CommandUtils.Log("Prebuilt={0}", Prebuilt);
+				CommandUtils.Log("Prereqs={0}", Prereqs);
+				CommandUtils.Log("NoBootstrapExe={0}", NoBootstrapExe);
+				CommandUtils.Log("RawProjectPath={0}", RawProjectPath);
+				CommandUtils.Log("Rocket={0}", Rocket);
+				CommandUtils.Log("Run={0}", Run);
+				CommandUtils.Log("ServerConfigsToBuild={0}", string.Join(",", ServerConfigsToBuild));
+				CommandUtils.Log("ServerCookedTargets={0}", ServerCookedTargets.ToString());
+				CommandUtils.Log("ServerTargetPlatform={0}", string.Join(",", ServerTargetPlatforms));
+				CommandUtils.Log("ShortProjectName={0}", ShortProjectName.ToString());
+				CommandUtils.Log("SignedPak={0}", SignedPak);
+				CommandUtils.Log("SignPak={0}", SignPak);				
+				CommandUtils.Log("SkipCook={0}", SkipCook);
+				CommandUtils.Log("SkipCookOnTheFly={0}", SkipCookOnTheFly);
+				CommandUtils.Log("SkipPak={0}", SkipPak);
+				CommandUtils.Log("SkipStage={0}", SkipStage);
+				CommandUtils.Log("Stage={0}", Stage);
+				CommandUtils.Log("bUsesSteam={0}", bUsesSteam);
+				CommandUtils.Log("bUsesCEF3={0}", bUsesCEF3);
+				CommandUtils.Log("bUsesSlate={0}", bUsesSlate);
+                CommandUtils.Log("bDebugBuildsActuallyUseDebugCRT={0}", bDebugBuildsActuallyUseDebugCRT);
+				CommandUtils.Log("Project Params **************");
+			}
+			bLogged = true;
+
+			Validate();
+		}
+
+		#endregion
+	}
+}
