UEB-260 - Break AutomationTool into AutomationUtils that all automation projects depend on, and AutomationTool, which essentially only contains the startup code.

diff --git a/Engine/Source/Programs/AutomationTool/AutomationUtils/ProjectUtils.cs b/Engine/Source/Programs/AutomationTool/AutomationUtils/ProjectUtils.cs
new file mode 100644
--- /dev/null
+++ b/Engine/Source/Programs/AutomationTool/AutomationUtils/ProjectUtils.cs
@@ -0,0 +1,808 @@
+﻿// Copyright 1998-2015 Epic Games, Inc. All Rights Reserved.
+
+using System;
+using System.Collections.Generic;
+using System.Linq;
+using System.Text;
+using System.IO;
+using UnrealBuildTool;
+using System.Diagnostics;
+using Tools.DotNETCommon.CaselessDictionary;
+using Tools.DotNETCommon.ExecutingAssembly;
+
+namespace AutomationTool
+{
+	public struct SingleTargetProperties
+	{
+		public string TargetName;
+		public TargetRules Rules;
+	}
+
+	/// <summary>
+	/// Autodetected project properties.
+	/// </summary>
+	public class ProjectProperties
+	{
+		/// <summary>
+		/// Full Project path. Must be a .uproject file
+		/// </summary>
+		public string RawProjectPath;
+
+		/// <summary>
+		/// True if the uproject contains source code.
+		/// </summary>
+		public bool bIsCodeBasedProject;
+
+		/// <summary>
+		/// Whether the project uses Steam (todo: substitute with more generic functionality)
+		/// </summary>
+		public bool bUsesSteam;
+
+		/// <summary>
+		/// Whether the project uses CEF3
+		/// </summary>
+		public bool bUsesCEF3;
+
+		/// <summary>
+		/// Whether the project uses visual Slate UI (as opposed to the low level windowing/messaging which is always used)
+		/// </summary>
+		public bool bUsesSlate = true;
+
+		/// <summary>
+		/// Hack for legacy game styling isses.  No new project should ever set this to true
+		/// Whether the project uses the Slate editor style in game.  
+		/// </summary>
+		public bool bUsesSlateEditorStyle = false;
+
+        /// <summary>
+        // By default we use the Release C++ Runtime (CRT), even when compiling Debug builds.  This is because the Debug C++
+        // Runtime isn't very useful when debugging Unreal Engine projects, and linking against the Debug CRT libraries forces
+        // our third party library dependencies to also be compiled using the Debug CRT (and often perform more slowly.)  Often
+        // it can be inconvenient to require a separate copy of the debug versions of third party static libraries simply
+        // so that you can debug your program's code.
+        /// </summary>
+        public bool bDebugBuildsActuallyUseDebugCRT = false;
+
+		/// <summary>
+		/// List of all targets detected for this project.
+		/// </summary>
+		public Dictionary<TargetRules.TargetType, SingleTargetProperties> Targets = new Dictionary<TargetRules.TargetType, SingleTargetProperties>();
+
+		/// <summary>
+		/// List of all Engine ini files for this project
+		/// </summary>
+		public Dictionary<UnrealTargetPlatform, ConfigCacheIni> EngineConfigs = new Dictionary<UnrealTargetPlatform,ConfigCacheIni>();
+
+		/// <summary>
+		/// List of all Game ini files for this project
+		/// </summary>
+		public Dictionary<UnrealTargetPlatform, ConfigCacheIni> GameConfigs = new Dictionary<UnrealTargetPlatform, ConfigCacheIni>();
+
+		/// <summary>
+		/// List of all programs detected for this project.
+		/// </summary>
+		public List<SingleTargetProperties> Programs = new List<SingleTargetProperties>();
+
+		/// <summary>
+		/// Specifies if the target files were generated
+		/// </summary>
+		public bool bWasGenerated = false;
+
+		internal ProjectProperties()
+		{
+		}
+	}
+
+	/// <summary>
+	/// Project related utility functions.
+	/// </summary>
+	public class ProjectUtils
+	{
+		private static CaselessDictionary<ProjectProperties> PropertiesCache = new CaselessDictionary<ProjectProperties>();
+
+		/// <summary>
+		/// Gets a short project name (QAGame, Elemental, etc)
+		/// </summary>
+		/// <param name="RawProjectPath">Full project path.</param>
+		/// <param name="bIsUProjectFile">True if a uproject.</param>
+		/// <returns>Short project name</returns>
+		public static string GetShortProjectName(string RawProjectPath)
+		{
+			return CommandUtils.GetFilenameWithoutAnyExtensions(RawProjectPath);
+		}
+
+		/// <summary>
+		/// Gets project properties.
+		/// </summary>
+		/// <param name="RawProjectPath">Full project path.</param>
+		/// <returns>Properties of the project.</returns>
+		public static ProjectProperties GetProjectProperties(string RawProjectPath, List<UnrealTargetPlatform> ClientTargetPlatforms = null)
+		{
+			string ProjectKey = "UE4";
+			if (!String.IsNullOrEmpty(RawProjectPath))
+			{
+				ProjectKey = CommandUtils.ConvertSeparators(PathSeparator.Slash, Path.GetFullPath(RawProjectPath));
+			}
+			ProjectProperties Properties;
+			if (PropertiesCache.TryGetValue(ProjectKey, out Properties) == false)
+			{
+				Properties = DetectProjectProperties(RawProjectPath, ClientTargetPlatforms);
+				PropertiesCache.Add(ProjectKey, Properties);
+			}
+			return Properties;
+		}
+
+		/// <summary>
+		/// Checks if the project is a UProject file with source code.
+		/// </summary>
+		/// <param name="RawProjectPath">Full project path.</param>
+		/// <returns>True if the project is a UProject file with source code.</returns>
+		public static bool IsCodeBasedUProjectFile(string RawProjectPath)
+		{
+			return GetProjectProperties(RawProjectPath, null).bIsCodeBasedProject;
+		}
+
+		/// <summary>
+		/// Returns a path to the client binaries folder.
+		/// </summary>
+		/// <param name="RawProjectPath">Full project path.</param>
+		/// <param name="Platform">Platform type.</param>
+		/// <returns>Path to the binaries folder.</returns>
+		public static string GetProjectClientBinariesFolder(string ProjectClientBinariesPath, UnrealTargetPlatform Platform = UnrealTargetPlatform.Unknown)
+		{
+			if (Platform != UnrealTargetPlatform.Unknown)
+			{
+				ProjectClientBinariesPath = CommandUtils.CombinePaths(ProjectClientBinariesPath, Platform.ToString());
+			}
+			return ProjectClientBinariesPath;
+		}
+
+		private static bool RequiresTempTarget(string RawProjectPath, List<UnrealTargetPlatform> ClientTargetPlatforms)
+		{
+			// check to see if we already have a Target.cs file
+			if (File.Exists (Path.Combine (Path.GetDirectoryName (RawProjectPath), "Source", Path.GetFileNameWithoutExtension (RawProjectPath) + ".Target.cs")))
+			{
+				return false;
+			}
+			else if (Directory.Exists(Path.Combine(Path.GetDirectoryName(RawProjectPath), "Source")))
+			{
+				// wasn't one in the main Source directory, let's check all sub-directories
+				//@todo: may want to read each target.cs to see if it has a target corresponding to the project name as a final check
+				FileInfo[] Files = (new DirectoryInfo( Path.Combine (Path.GetDirectoryName (RawProjectPath), "Source")).GetFiles ("*.Target.cs", SearchOption.AllDirectories));
+				if (Files.Length > 0)
+				{
+					return false;
+				}
+			}
+
+			// no Target file, now check to see if build settings have changed
+			List<UnrealTargetPlatform> TargetPlatforms = ClientTargetPlatforms;
+			if (ClientTargetPlatforms == null || ClientTargetPlatforms.Count < 1)
+			{
+				// No client target platforms, add all in
+				TargetPlatforms = new List<UnrealTargetPlatform>();
+				foreach (UnrealTargetPlatform TargetPlatformType in Enum.GetValues(typeof(UnrealTargetPlatform)))
+				{
+					if (TargetPlatformType != UnrealTargetPlatform.Unknown)
+					{
+						TargetPlatforms.Add(TargetPlatformType);
+					}
+				}
+			}
+
+			// Change the working directory to be the Engine/Source folder. We are running from Engine/Binaries/DotNET
+			string oldCWD = Directory.GetCurrentDirectory();
+			if (BuildConfiguration.RelativeEnginePath == "../../Engine/")
+			{
+				string EngineSourceDirectory = Path.Combine(ExecutingAssembly.GetDirectory(), "..", "..", "..", "Engine", "Source");
+				if (!Directory.Exists(EngineSourceDirectory)) // only set the directory if it exists, this should only happen if we are launching the editor from an artist sync
+				{
+					EngineSourceDirectory = Path.Combine(ExecutingAssembly.GetDirectory(), "..", "..", "..", "Engine", "Binaries");
+				}
+				Directory.SetCurrentDirectory(EngineSourceDirectory);
+			}
+
+			// Read the project descriptor, and find all the plugins available to this project
+			ProjectDescriptor Project = ProjectDescriptor.FromFile(RawProjectPath);
+			List<PluginInfo> AvailablePlugins = Plugins.ReadAvailablePlugins(RawProjectPath);
+
+			// check the target platforms for any differences in build settings or additional plugins
+			bool RetVal = false;
+			foreach (UnrealTargetPlatform TargetPlatformType in TargetPlatforms)
+			{
+				IUEBuildPlatform BuildPlat = UEBuildPlatform.GetBuildPlatform(TargetPlatformType, true);
+				if (!GlobalCommandLine.Rocket && BuildPlat != null && !(BuildPlat as UEBuildPlatform).HasDefaultBuildConfig(TargetPlatformType, Path.GetDirectoryName(RawProjectPath)))
+				{
+					RetVal = true;
+					break;
+				}
+
+				// find if there are any plugins enabled or disabled which differ from the default
+				foreach(PluginInfo Plugin in AvailablePlugins)
+				{
+					bool bPluginEnabledForProject = UProjectInfo.IsPluginEnabledForProject(Plugin, Project, TargetPlatformType);
+					if ((bPluginEnabledForProject && !Plugin.Descriptor.bEnabledByDefault) || (bPluginEnabledForProject && Plugin.Descriptor.bInstalled))
+					{
+						if(Plugin.Descriptor.Modules.Any(Module => Module.IsCompiledInConfiguration(TargetPlatformType, TargetRules.TargetType.Game, bBuildDeveloperTools: false, bBuildEditor: false)))
+						{
+							RetVal = true;
+							break;
+						}
+					}
+				}
+			}
+
+			// Change back to the original directory
+			Directory.SetCurrentDirectory(oldCWD);
+			return RetVal;
+		}
+
+		private static void GenerateTempTarget(string RawProjectPath)
+		{
+			// read in the template target cs file
+			var TempCSFile = CommandUtils.CombinePaths(CommandUtils.CmdEnv.LocalRoot, "Engine", "Build", "Target.cs.template");
+			string TargetCSFile = File.ReadAllText(TempCSFile);
+
+			// replace {GAME_NAME} with the game name
+			TargetCSFile = TargetCSFile.Replace("{GAME_NAME}", Path.GetFileNameWithoutExtension(RawProjectPath));
+
+			// write out the file in a new Source directory
+			string FileName = CommandUtils.CombinePaths(Path.GetDirectoryName(RawProjectPath), "Intermediate", "Source", Path.GetFileNameWithoutExtension(RawProjectPath) + ".Target.cs");
+			if (!Directory.Exists(Path.GetDirectoryName(FileName)))
+			{
+				Directory.CreateDirectory(Path.GetDirectoryName(FileName));
+			}
+
+			File.WriteAllText(FileName, TargetCSFile);
+		}
+
+		/// <summary>
+		/// Attempts to autodetect project properties.
+		/// </summary>
+		/// <param name="RawProjectPath">Full project path.</param>
+		/// <returns>Project properties.</returns>
+		private static ProjectProperties DetectProjectProperties(string RawProjectPath, List<UnrealTargetPlatform> ClientTargetPlatforms)
+		{
+			var Properties = new ProjectProperties();
+			Properties.RawProjectPath = RawProjectPath;
+
+			// detect if the project is content only, but has non-default build settings
+			List<string> ExtraSearchPaths = null;
+			if (!string.IsNullOrEmpty(RawProjectPath))
+			{
+				if (RequiresTempTarget(RawProjectPath, ClientTargetPlatforms))
+				{
+					GenerateTempTarget(RawProjectPath);
+					Properties.bWasGenerated = true;
+					ExtraSearchPaths = new List<string>();
+					ExtraSearchPaths.Add(CommandUtils.CombinePaths(Path.GetDirectoryName(RawProjectPath), "Intermediate", "Source"));
+				}
+				else if (File.Exists(Path.Combine(Path.GetDirectoryName(RawProjectPath), "Intermediate", "Source", Path.GetFileNameWithoutExtension(RawProjectPath) + ".Target.cs")))
+				{
+					File.Delete(Path.Combine(Path.GetDirectoryName(RawProjectPath), "Intermediate", "Source", Path.GetFileNameWithoutExtension(RawProjectPath) + ".Target.cs"));
+				}
+			}
+
+			if (CommandUtils.CmdEnv.HasCapabilityToCompile)
+			{
+				DetectTargetsForProject(Properties, ExtraSearchPaths);
+				Properties.bIsCodeBasedProject = !CommandUtils.IsNullOrEmpty(Properties.Targets) || !CommandUtils.IsNullOrEmpty(Properties.Programs);
+			}
+			else
+			{
+				// should never ask for engine targets if we can't compile
+				if (String.IsNullOrEmpty(RawProjectPath))
+				{
+					throw new AutomationException("Cannot determine engine targets if we can't compile.");
+				}
+
+				Properties.bIsCodeBasedProject = Properties.bWasGenerated;
+				// if there's a Source directory with source code in it, then mark us as having source code
+				string SourceDir = CommandUtils.CombinePaths(Path.GetDirectoryName(RawProjectPath), "Source");
+				if (Directory.Exists(SourceDir))
+				{
+					string[] CppFiles = Directory.GetFiles(SourceDir, "*.cpp", SearchOption.AllDirectories);
+					string[] HFiles = Directory.GetFiles(SourceDir, "*.h", SearchOption.AllDirectories);
+					Properties.bIsCodeBasedProject |= (CppFiles.Length > 0 || HFiles.Length > 0);
+				}
+			}
+
+			// check to see if the uproject loads modules, only if we haven't already determined it is a code based project
+			if (!Properties.bIsCodeBasedProject)
+			{
+				string uprojectStr = File.ReadAllText(RawProjectPath);
+				Properties.bIsCodeBasedProject = uprojectStr.Contains("\"Modules\"");
+			}
+
+			// Get all ini files
+			if (!String.IsNullOrWhiteSpace(RawProjectPath))
+			{
+				CommandUtils.LogVerbose("Loading ini files for {0}", RawProjectPath);
+
+				var EngineDirectory = CommandUtils.CombinePaths(CommandUtils.CmdEnv.LocalRoot, "Engine");
+				foreach (UnrealTargetPlatform TargetPlatformType in Enum.GetValues(typeof(UnrealTargetPlatform)))
+				{
+					if (TargetPlatformType != UnrealTargetPlatform.Unknown)
+					{
+						var Config = new ConfigCacheIni(TargetPlatformType, "Engine", Path.GetDirectoryName(RawProjectPath), EngineDirectory);
+						Properties.EngineConfigs.Add(TargetPlatformType, Config);
+					}
+				}
+
+				foreach (UnrealTargetPlatform TargetPlatformType in Enum.GetValues(typeof(UnrealTargetPlatform)))
+				{
+					if (TargetPlatformType != UnrealTargetPlatform.Unknown)
+					{
+						var Config = new ConfigCacheIni(TargetPlatformType, "Game", Path.GetDirectoryName(RawProjectPath));
+						Properties.GameConfigs.Add(TargetPlatformType, Config);
+					}
+				}
+			}
+
+			return Properties;
+		}
+
+
+		/// <summary>
+		/// Gets the project's root binaries folder.
+		/// </summary>
+		/// <param name="RawProjectPath">Full project path.</param>
+		/// <param name="TargetType">Target type.</param>
+		/// <param name="bIsUProjectFile">True if uproject file.</param>
+		/// <returns>Binaries path.</returns>
+		public static string GetClientProjectBinariesRootPath(string RawProjectPath, TargetRules.TargetType TargetType, bool bIsCodeBasedProject)
+		{
+			var BinPath = String.Empty;
+			switch (TargetType)
+			{
+				case TargetRules.TargetType.Program:
+					BinPath = CommandUtils.CombinePaths(CommandUtils.CmdEnv.LocalRoot, "Engine", "Binaries");
+					break;
+				case TargetRules.TargetType.Client:
+				case TargetRules.TargetType.Game:
+					if (!bIsCodeBasedProject)
+					{
+						BinPath = CommandUtils.CombinePaths(CommandUtils.CmdEnv.LocalRoot, "Engine", "Binaries");
+					}
+					else
+					{
+						BinPath = CommandUtils.CombinePaths(CommandUtils.GetDirectoryName(RawProjectPath), "Binaries");
+					}
+					break;
+			}
+			return BinPath;
+		}
+
+		/// <summary>
+		/// Gets the location where all rules assemblies should go
+		/// </summary>
+		private static string GetRulesAssemblyFolder()
+		{
+			string RulesFolder;
+			if (GlobalCommandLine.Installed.IsSet)
+			{
+				RulesFolder = CommandUtils.CombinePaths(Path.GetTempPath(), "UAT", CommandUtils.EscapePath(CommandUtils.CmdEnv.LocalRoot), "Rules"); 
+			}
+			else
+			{
+				RulesFolder = CommandUtils.CombinePaths(CommandUtils.CmdEnv.EngineSavedFolder, "Rules");
+			}
+			return RulesFolder;
+		}
+
+		/// <summary>
+		/// Finds all targets for the project.
+		/// </summary>
+		/// <param name="Properties">Project properties.</param>
+		/// <param name="ExtraSearchPaths">Additional search paths.</param>
+		private static void DetectTargetsForProject(ProjectProperties Properties, List<string> ExtraSearchPaths = null)
+		{
+			Properties.Targets = new Dictionary<TargetRules.TargetType, SingleTargetProperties>();
+			string TargetsDllFilename;
+			string FullProjectPath = null;
+
+			var GameFolders = new List<string>();
+			var RulesFolder = GetRulesAssemblyFolder();
+			if (!String.IsNullOrEmpty(Properties.RawProjectPath))
+			{
+				CommandUtils.LogVerbose("Looking for targets for project {0}", Properties.RawProjectPath);
+
+				TargetsDllFilename = CommandUtils.CombinePaths(RulesFolder, String.Format("UATRules{0}.dll", Properties.RawProjectPath.GetHashCode()));
+
+				FullProjectPath = CommandUtils.GetDirectoryName(Properties.RawProjectPath);
+				GameFolders.Add(FullProjectPath);
+				CommandUtils.LogVerbose("Searching for target rule files in {0}", FullProjectPath);
+			}
+			else
+			{
+				TargetsDllFilename = CommandUtils.CombinePaths(RulesFolder, String.Format("UATRules{0}.dll", "_BaseEngine_"));
+			}
+			RulesCompiler.SetAssemblyNameAndGameFolders(TargetsDllFilename, GameFolders);
+
+			// the UBT code assumes a certain CWD, but artists don't have this CWD.
+			var SourceDir = CommandUtils.CombinePaths(CommandUtils.CmdEnv.LocalRoot, "Engine", "Source");
+			bool DirPushed = false;
+			if (CommandUtils.DirectoryExists_NoExceptions(SourceDir))
+			{
+				CommandUtils.PushDir(SourceDir);
+				DirPushed = true;
+			}
+			var TargetScripts = RulesCompiler.FindAllRulesSourceFiles(RulesCompiler.RulesFileType.Target, ExtraSearchPaths);
+			if (DirPushed)
+			{
+				CommandUtils.PopDir();
+			}
+
+			if (!CommandUtils.IsNullOrEmpty(TargetScripts))
+			{
+				// We only care about project target script so filter out any scripts not in the project folder, or take them all if we are just doing engine stuff
+				var ProjectTargetScripts = new List<string>();
+				foreach (var Filename in TargetScripts)
+				{
+					var FullScriptPath = CommandUtils.CombinePaths(Path.GetFullPath(Filename));
+					if (FullProjectPath == null || FullScriptPath.StartsWith(FullProjectPath, StringComparison.InvariantCultureIgnoreCase))
+					{
+						ProjectTargetScripts.Add(FullScriptPath);
+					}
+				}
+				TargetScripts = ProjectTargetScripts;
+			}
+
+			if (!CommandUtils.IsNullOrEmpty(TargetScripts))
+			{
+				CommandUtils.LogVerbose("Found {0} target rule files:", TargetScripts.Count);
+				foreach (var Filename in TargetScripts)
+				{
+					CommandUtils.LogVerbose("  {0}", Filename);
+				}
+
+				// Check if the scripts require compilation
+				bool DoNotCompile = false;
+				if (!CommandUtils.IsBuildMachine && !CheckIfScriptAssemblyIsOutOfDate(TargetsDllFilename, TargetScripts))
+				{
+					Log.TraceInformation("Targets DLL {0} is up to date.", TargetsDllFilename);
+					DoNotCompile = true;
+				}
+				if (!DoNotCompile && CommandUtils.FileExists_NoExceptions(TargetsDllFilename))
+				{
+					if (!CommandUtils.DeleteFile_NoExceptions(TargetsDllFilename, true))
+					{
+						DoNotCompile = true;
+						CommandUtils.Log("Could not delete {0} assuming it is up to date and reusable for a recursive UAT call.", TargetsDllFilename);
+					}
+				}
+
+				CompileAndLoadTargetsAssembly(Properties, TargetsDllFilename, DoNotCompile, TargetScripts);
+			}
+		}
+
+		/// <summary>
+		/// Optionally compiles and loads target rules assembly.
+		/// </summary>
+		/// <param name="Properties"></param>
+		/// <param name="TargetsDllFilename"></param>
+		/// <param name="DoNotCompile"></param>
+		/// <param name="TargetScripts"></param>
+		private static void CompileAndLoadTargetsAssembly(ProjectProperties Properties, string TargetsDllFilename, bool DoNotCompile, List<string> TargetScripts)
+		{
+			CommandUtils.Log("Compiling targets DLL: {0}", TargetsDllFilename);
+
+			var ReferencedAssemblies = new List<string>() 
+					{ 
+						"System.dll", 
+						"System.Core.dll", 
+						"System.Xml.dll", 
+						typeof(UnrealBuildTool.UnrealBuildTool).Assembly.Location
+					};
+			var TargetsDLL = DynamicCompilation.CompileAndLoadAssembly(TargetsDllFilename, TargetScripts, ReferencedAssemblies, null, DoNotCompile);
+			var DummyTargetInfo = new TargetInfo(BuildHostPlatform.Current.Platform, UnrealTargetConfiguration.Development);
+			var AllCompiledTypes = TargetsDLL.GetTypes();
+			foreach (Type TargetType in AllCompiledTypes)
+			{
+				// Find TargetRules but skip all "UE4Editor", "UE4Game" targets.
+				if (typeof(TargetRules).IsAssignableFrom(TargetType))
+				{
+					// Create an instance of this type
+					CommandUtils.LogVerbose("Creating target rules object: {0}", TargetType.Name);
+					var Rules = Activator.CreateInstance(TargetType, DummyTargetInfo) as TargetRules;
+					CommandUtils.LogVerbose("Adding target: {0} ({1})", TargetType.Name, Rules.Type);
+
+					SingleTargetProperties TargetData;
+					TargetData.TargetName = GetTargetName(TargetType);
+                    Rules.TargetName = TargetData.TargetName;
+					TargetData.Rules = Rules;
+					if (Rules.Type == TargetRules.TargetType.Program)
+					{
+						Properties.Programs.Add(TargetData);
+					}
+					else
+					{
+						Properties.Targets.Add(Rules.Type, TargetData);
+					}
+
+					Properties.bUsesSteam |= Rules.bUsesSteam;
+					Properties.bUsesCEF3 |= Rules.bUsesCEF3;
+					Properties.bUsesSlate |= Rules.bUsesSlate;
+                    Properties.bDebugBuildsActuallyUseDebugCRT |= Rules.bDebugBuildsActuallyUseDebugCRT;
+					Properties.bUsesSlateEditorStyle |= Rules.bUsesSlateEditorStyle;
+				}
+			}
+		}
+
+		/// <summary>
+		/// Checks if any of the script files in newer than the generated assembly.
+		/// </summary>
+		/// <param name="TargetsDllFilename"></param>
+		/// <param name="TargetScripts"></param>
+		/// <returns>True if the generated assembly is out of date.</returns>
+		private static bool CheckIfScriptAssemblyIsOutOfDate(string TargetsDllFilename, List<string> TargetScripts)
+		{
+			var bOutOfDate = false;
+			var AssemblyInfo = new FileInfo(TargetsDllFilename);
+			if (AssemblyInfo.Exists)
+			{
+				foreach (var ScriptFilename in TargetScripts)
+				{
+					var ScriptInfo = new FileInfo(ScriptFilename);
+					if (ScriptInfo.Exists && ScriptInfo.LastWriteTimeUtc > AssemblyInfo.LastWriteTimeUtc)
+					{
+						bOutOfDate = true;
+						break;
+					}
+				}
+			}
+			else
+			{
+				bOutOfDate = true;
+			}
+			return bOutOfDate;
+		}
+
+		/// <summary>
+		/// Converts class type name (usually ends with Target) to a target name (without the postfix).
+		/// </summary>
+		/// <param name="TargetRulesType">Tagert class.</param>
+		/// <returns>Target name</returns>
+		private static string GetTargetName(Type TargetRulesType)
+		{
+			const string TargetPostfix = "Target";
+			var Name = TargetRulesType.Name;
+			if (Name.EndsWith(TargetPostfix, StringComparison.InvariantCultureIgnoreCase))
+			{
+				Name = Name.Substring(0, Name.Length - TargetPostfix.Length);
+			}
+			return Name;
+		}
+
+		/// <summary>
+		/// Performs initial cleanup of target rules folder
+		/// </summary>
+		public static void CleanupFolders()
+		{
+			CommandUtils.Log("Cleaning up project rules folder");
+			var RulesFolder = GetRulesAssemblyFolder();
+			if (CommandUtils.DirectoryExists(RulesFolder))
+			{
+				CommandUtils.DeleteDirectoryContents(RulesFolder);
+			}
+		}
+	}
+
+    public class BranchInfo
+    {
+
+        public static TargetRules.TargetType[] MonolithicKinds = new TargetRules.TargetType[]
+        {
+            TargetRules.TargetType.Game,
+            TargetRules.TargetType.Client,
+            TargetRules.TargetType.Server,
+        };
+
+		[DebuggerDisplay("{GameName}")]
+        public class BranchUProject
+        {
+            public string GameName;
+            public string FilePath;
+            public ProjectProperties Properties;
+
+            public BranchUProject(UnrealBuildTool.UProjectInfo InfoEntry)
+            {
+                GameName = InfoEntry.GameName;
+
+                //not sure what the heck this path is relative to
+                FilePath = Path.GetFullPath(CommandUtils.CombinePaths(CommandUtils.CmdEnv.LocalRoot, "Engine", "Binaries", InfoEntry.FilePath));
+
+                if (!CommandUtils.FileExists_NoExceptions(FilePath))
+                {
+                    throw new AutomationException("Could not resolve relative path corrctly {0} -> {1} which doesn't exist.", InfoEntry.FilePath, FilePath);
+                }
+
+                Properties = ProjectUtils.GetProjectProperties(Path.GetFullPath(FilePath));
+
+
+
+            }
+            public BranchUProject()
+            {
+                GameName = "UE4";
+                Properties = ProjectUtils.GetProjectProperties("");
+                if (!Properties.Targets.ContainsKey(TargetRules.TargetType.Editor))
+                {
+                    throw new AutomationException("Base UE4 project did not contain an editor target.");
+                }
+            }
+            public TargetRules.GUBPProjectOptions Options(UnrealTargetPlatform HostPlatform)
+            {
+                var Options = new TargetRules.GUBPProjectOptions();
+                if (Properties.Targets.ContainsKey(TargetRules.TargetType.Editor))
+                {
+                    var Target = Properties.Targets[TargetRules.TargetType.Editor];
+                    Options = Target.Rules.GUBP_IncludeProjectInPromotedBuild_EditorTypeOnly(HostPlatform);
+                }
+				else if(Properties.Targets.ContainsKey(TargetRules.TargetType.Game))
+				{
+					var Target = Properties.Targets[TargetRules.TargetType.Game];
+					Options = Target.Rules.GUBP_IncludeProjectInPromotedBuild_EditorTypeOnly(HostPlatform);
+				}
+                return Options;
+            }
+            public void Dump(List<UnrealTargetPlatform> InHostPlatforms)
+            {
+                CommandUtils.LogVerbose("    ShortName:    " + GameName);
+				CommandUtils.LogVerbose("      FilePath          : " + FilePath);
+				CommandUtils.LogVerbose("      bIsCodeBasedProject  : " + (Properties.bIsCodeBasedProject ? "YES" : "NO"));
+				CommandUtils.LogVerbose("      bUsesSteam  : " + (Properties.bUsesSteam ? "YES" : "NO"));
+				CommandUtils.LogVerbose("      bUsesCEF3   : " + (Properties.bUsesCEF3 ? "YES" : "NO"));
+				CommandUtils.LogVerbose("      bUsesSlate  : " + (Properties.bUsesSlate ? "YES" : "NO"));
+                foreach (var HostPlatform in InHostPlatforms)
+                {
+					CommandUtils.LogVerbose("      For Host : " + HostPlatform.ToString());
+                    if (Properties.bIsCodeBasedProject)
+                    {
+						CommandUtils.LogVerbose("          Promoted  : " + (Options(HostPlatform).bIsPromotable ? "YES" : "NO"));
+						CommandUtils.LogVerbose("          SeparatePromotable  : " + (Options(HostPlatform).bSeparateGamePromotion ? "YES" : "NO"));
+						CommandUtils.LogVerbose("          Test With Shared  : " + (Options(HostPlatform).bTestWithShared ? "YES" : "NO"));
+                    }
+					CommandUtils.LogVerbose("          Targets {0}:", Properties.Targets.Count);
+                    foreach (var ThisTarget in Properties.Targets)
+                    {
+						CommandUtils.LogVerbose("            TargetName          : " + ThisTarget.Value.TargetName);
+						CommandUtils.LogVerbose("              Type          : " + ThisTarget.Key.ToString());
+						CommandUtils.LogVerbose("              bUsesSteam  : " + (ThisTarget.Value.Rules.bUsesSteam ? "YES" : "NO"));
+						CommandUtils.LogVerbose("              bUsesCEF3   : " + (ThisTarget.Value.Rules.bUsesCEF3 ? "YES" : "NO"));
+						CommandUtils.LogVerbose("              bUsesSlate  : " + (ThisTarget.Value.Rules.bUsesSlate ? "YES" : "NO"));
+                        if (Array.IndexOf(MonolithicKinds, ThisTarget.Key) >= 0)
+                        {
+                            var Platforms = ThisTarget.Value.Rules.GUBP_GetPlatforms_MonolithicOnly(HostPlatform);
+                            var AdditionalPlatforms = ThisTarget.Value.Rules.GUBP_GetBuildOnlyPlatforms_MonolithicOnly(HostPlatform);
+                            var AllPlatforms = Platforms.Union(AdditionalPlatforms);
+                            foreach (var Plat in AllPlatforms)
+                            {
+                                var Configs = ThisTarget.Value.Rules.GUBP_GetConfigs_MonolithicOnly(HostPlatform, Plat);
+                                foreach (var Config in Configs)
+                                {
+									CommandUtils.LogVerbose("                Config  : " + Plat.ToString() + " " + Config.ToString());
+                                }
+                            }
+
+                        }
+                    }
+					CommandUtils.LogVerbose("      Programs {0}:", Properties.Programs.Count);
+                    foreach (var ThisTarget in Properties.Programs)
+                    {
+                        bool bInternalToolOnly;
+                        bool SeparateNode;
+						bool CrossCompile;
+                        bool Tool = ThisTarget.Rules.GUBP_AlwaysBuildWithTools(HostPlatform, out bInternalToolOnly, out SeparateNode, out CrossCompile);
+
+						CommandUtils.LogVerbose("            TargetName                    : " + ThisTarget.TargetName);
+						CommandUtils.LogVerbose("              Build With Editor           : " + (ThisTarget.Rules.GUBP_AlwaysBuildWithBaseEditor() ? "YES" : "NO"));
+						CommandUtils.LogVerbose("              Build With Tools            : " + (Tool && !bInternalToolOnly ? "YES" : "NO"));
+						CommandUtils.LogVerbose("              Build With Internal Tools   : " + (Tool && bInternalToolOnly ? "YES" : "NO"));
+						CommandUtils.LogVerbose("              Separate Node               : " + (Tool && SeparateNode ? "YES" : "NO"));
+						CommandUtils.LogVerbose("              Cross Compile               : " + (Tool && CrossCompile ? "YES" : "NO"));
+                    }
+                }
+            }
+        };
+        public BranchUProject BaseEngineProject = null;
+        public List<BranchUProject> CodeProjects = new List<BranchUProject>();
+        public List<BranchUProject> NonCodeProjects = new List<BranchUProject>();
+
+		public IEnumerable<BranchUProject> AllProjects
+		{
+			get { return CodeProjects.Union(NonCodeProjects); }
+		}
+
+        public BranchInfo(List<UnrealTargetPlatform> InHostPlatforms)
+        {
+            BaseEngineProject = new BranchUProject();
+
+            var AllProjects = UnrealBuildTool.UProjectInfo.FilterGameProjects(false, null);
+            var StartTime = DateTime.Now.ToString();
+            foreach (var InfoEntry in AllProjects)
+            {
+                var UProject = new BranchUProject(InfoEntry);
+                if (UProject.Properties.bIsCodeBasedProject)
+                {
+                    CodeProjects.Add(UProject);
+                }
+                else
+                {
+                    NonCodeProjects.Add(UProject);
+                    // the base project uses BlankProject if it really needs a .uproject file
+                    if (String.IsNullOrEmpty(BaseEngineProject.FilePath) && UProject.GameName == "BlankProject")
+                    {
+                        BaseEngineProject.FilePath = UProject.FilePath;
+                    }
+                }
+            }
+            var FinishTime = DateTime.Now.ToString();
+            CommandUtils.PrintCSVFile(String.Format("UAT,SortProjects,{0},{1}", StartTime, FinishTime));            
+ /*           if (String.IsNullOrEmpty(BaseEngineProject.FilePath))
+            {
+                throw new AutomationException("All branches must have the blank project /Samples/Sandbox/BlankProject");
+            }*/
+
+            CommandUtils.Log("  Base Engine:");
+            var StartBranchDump = DateTime.Now.ToString();
+            BaseEngineProject.Dump(InHostPlatforms);
+
+            CommandUtils.Log("  {0} Code projects:", CodeProjects.Count);
+            foreach (var Proj in CodeProjects)
+            {
+                Proj.Dump(InHostPlatforms);
+            }
+            CommandUtils.Log("  {0} Non-Code projects:", CodeProjects.Count);
+            foreach (var Proj in NonCodeProjects)
+            {
+                Proj.Dump(InHostPlatforms);
+            }
+            var FinishBranchDump = DateTime.Now.ToString();
+            CommandUtils.PrintCSVFile(String.Format("UAT,Project Dump,{0},{1}", StartBranchDump, FinishBranchDump));
+        }
+
+        public BranchUProject FindGame(string GameName)
+        {
+            foreach (var Proj in CodeProjects)
+            {
+                if (Proj.GameName.Equals(GameName, StringComparison.InvariantCultureIgnoreCase))
+                {
+                    return Proj;
+                }
+            }
+            foreach (var Proj in NonCodeProjects)
+            {
+                if (Proj.GameName.Equals(GameName, StringComparison.InvariantCultureIgnoreCase))
+                {
+                    return Proj;
+                }
+            }
+            return null;
+        }
+		public BranchUProject FindGameChecked(string GameName)
+		{
+			BranchUProject Project = FindGame(GameName);
+			if(Project == null)
+			{
+				throw new AutomationException("Cannot find project '{0}' in branch", GameName);
+			}
+			return Project;
+		}
+        public SingleTargetProperties FindProgram(string ProgramName)
+        {
+            foreach (var Proj in BaseEngineProject.Properties.Programs)
+            {
+                if (Proj.TargetName.Equals(ProgramName, StringComparison.InvariantCultureIgnoreCase))
+                {
+                    return Proj;
+                }
+            }
+            SingleTargetProperties Result;
+            Result.TargetName = ProgramName;
+            Result.Rules = null;
+            return Result;
+        }
+    };
+
+}
diff --git a/Engine/Source/Programs/AutomationTool/AutomationUtils/ScriptCompiler.cs b/Engine/Source/Programs/AutomationTool/AutomationUtils/ScriptCompiler.cs
new file mode 100644
--- /dev/null
+++ b/Engine/Source/Programs/AutomationTool/AutomationUtils/ScriptCompiler.cs
@@ -0,0 +1,334 @@
+// Copyright 1998-2015 Epic Games, Inc. All Rights Reserved.
+using System;
+using System.Collections.Generic;
+using System.Linq;
+using System.Text;
+using System.IO;
+using System.CodeDom.Compiler;
+using Microsoft.CSharp;
+using Microsoft.Win32;
+using System.Reflection;
+using System.Diagnostics;
+using UnrealBuildTool;
+using Tools.DotNETCommon.CaselessDictionary;
+
+namespace AutomationTool
+{
+	/// <summary>
+	/// Exception thrown by PreprocessScriptFile.
+	/// </summary>
+	class CompilationException : AutomationException
+	{
+		public CompilationException(string Filename, int StartLine, int StartColumn, int EndLine, int EndColumn, string Message, params string[] Args)
+			: base(String.Format("Compilation Failed.\n>{0}({1},{2},{3},{4}): error: {5}", Path.GetFullPath(Filename), StartLine + 1, StartColumn + 1, EndLine + 1, EndColumn + 1,
+			String.Format(Message, Args)))
+		{
+		}
+	}
+
+	/// <summary>
+	/// Compiles and loads script assemblies.
+	/// </summary>
+	class ScriptCompiler
+	{
+		#region Fields
+				
+		private CaselessDictionary<Type> ScriptCommands;
+#if DEBUG
+		const string BuildConfig = "Debug";
+#else
+		const string BuildConfig = "Development";
+#endif
+		const string DefaultScriptsDLLName = "AutomationScripts.Automation.dll";
+
+		#endregion
+
+		#region Compilation
+
+		public ScriptCompiler()
+		{
+		}
+
+
+		/// <summary>
+		/// Finds and/or compiles all script files and assemblies.
+		/// </summary>
+		/// <param name="AdditionalScriptsFolders">Additional script fodlers to look for source files in.</param>
+		public void FindAndCompileAllScripts(List<string> AdditionalScriptsFolders)
+		{
+			bool DoCompile = true;
+			if (GlobalCommandLine.NoCompile)
+			{
+				DoCompile = false;
+			}
+
+			// Change to Engine\Source (if exists) to properly discover all UBT classes
+			var OldCWD = Environment.CurrentDirectory;
+			var UnrealBuildToolCWD = CommandUtils.CombinePaths(CommandUtils.CmdEnv.LocalRoot, "Engine", "Source");
+			if (Directory.Exists(UnrealBuildToolCWD))
+			{
+				Environment.CurrentDirectory = UnrealBuildToolCWD;
+			}
+			// Register all the classes inside UBT
+			Log.TraceVerbose("Registering UBT Classes.");
+			UnrealBuildTool.UnrealBuildTool.RegisterAllUBTClasses();
+			Environment.CurrentDirectory = OldCWD;
+
+			// Compile only if not disallowed.
+			if (DoCompile && !String.IsNullOrEmpty(CommandUtils.CmdEnv.MsBuildExe))
+			{
+				CleanupScriptsAssemblies();
+				FindAndCompileScriptModules(AdditionalScriptsFolders);
+			}
+
+			var ScriptAssemblies = new List<Assembly>();
+			LoadPreCompiledScriptAssemblies(ScriptAssemblies);
+
+			// Setup platforms
+			Platform.InitializePlatforms(ScriptAssemblies.ToArray());
+
+			// Instantiate all the automation classes for interrogation
+			Log.TraceVerbose("Creating commands.");
+			ScriptCommands = new CaselessDictionary<Type>();
+			foreach (var CompiledScripts in ScriptAssemblies)
+			{
+				foreach (var ClassType in CompiledScripts.GetTypes())
+				{
+					if (ClassType.IsSubclassOf(typeof(BuildCommand)) && ClassType.IsAbstract == false)
+					{
+						if (ScriptCommands.ContainsKey(ClassType.Name) == false)
+						{
+							ScriptCommands.Add(ClassType.Name, ClassType);
+						}
+						else
+						{
+							Log.TraceWarning("Unable to add command {0} twice. Previous: {1}, Current: {2}", ClassType.Name,
+								ClassType.AssemblyQualifiedName, ScriptCommands[ClassType.Name].AssemblyQualifiedName);
+						}
+					}
+				}
+			}
+		}
+
+		private static void FindAndCompileScriptModules(List<string> AdditionalScriptsFolders)
+		{
+			var OldCWD = Environment.CurrentDirectory;
+			var UnrealBuildToolCWD = CommandUtils.CombinePaths(CommandUtils.CmdEnv.LocalRoot, "Engine", "Source");
+
+			// Convert script folders to be relative to UnrealBuildTool's expected CWD
+			var RemappedAdditionalScriptFolders = new List<string>();
+			foreach (var CurFolder in AdditionalScriptsFolders)
+			{
+				RemappedAdditionalScriptFolders.Add(UnrealBuildTool.Utils.MakePathRelativeTo(CurFolder, UnrealBuildToolCWD));
+			}
+
+			Environment.CurrentDirectory = UnrealBuildToolCWD;
+
+			// Configure the rules compiler
+			// Get all game folders and convert them to build subfolders.
+			var AllGameFolders = UnrealBuildTool.UEBuildTarget.DiscoverAllGameFolders();
+			var BuildFolders = new List<string>(AllGameFolders.Count);
+			foreach (var Folder in AllGameFolders)
+			{
+				var GameBuildFolder = CommandUtils.CombinePaths(Folder, "Build");
+				if (Directory.Exists(GameBuildFolder))
+				{
+					BuildFolders.Add(GameBuildFolder);
+				}
+			}
+			RemappedAdditionalScriptFolders.AddRange(BuildFolders);
+
+			Log.TraceVerbose("Discovering game folders.");
+			UnrealBuildTool.RulesCompiler.SetAssemblyNameAndGameFolders("UnrealAutomationToolRules", AllGameFolders);
+
+			var DiscoveredModules = UnrealBuildTool.RulesCompiler.FindAllRulesSourceFiles(UnrealBuildTool.RulesCompiler.RulesFileType.AutomationModule, RemappedAdditionalScriptFolders);
+			var ModulesToCompile = new List<string>(DiscoveredModules.Count);
+			foreach (var ModuleFilename in DiscoveredModules)
+			{
+				if (HostPlatform.Current.IsScriptModuleSupported(CommandUtils.GetFilenameWithoutAnyExtensions(ModuleFilename)))
+				{
+					ModulesToCompile.Add(ModuleFilename);
+				}
+				else
+				{
+					CommandUtils.LogVerbose("Script module {0} filtered by the Host Platform and will not be compiled.", ModuleFilename);
+				}
+			}
+
+			if ((UnrealBuildTool.BuildHostPlatform.Current.Platform == UnrealBuildTool.UnrealTargetPlatform.Win64) ||
+				(UnrealBuildTool.BuildHostPlatform.Current.Platform == UnrealBuildTool.UnrealTargetPlatform.Win32))
+			{
+				string Modules = string.Join(";", ModulesToCompile.ToArray());
+				var UATProj = CommandUtils.CombinePaths(CommandUtils.CmdEnv.LocalRoot, @"Engine\Source\Programs\AutomationTool\Scripts\UAT.proj");
+				var CmdLine = String.Format("\"{0}\" /p:Modules=\"{1}\" /p:Configuration={2}", UATProj, Modules, BuildConfig);
+				var Result = CommandUtils.Run(CommandUtils.CmdEnv.MsBuildExe, CmdLine);
+				if (Result.ExitCode != 0)
+				{
+					throw new AutomationException(String.Format("Failed to build \"{0}\":{1}{2}", UATProj, Environment.NewLine, Result.Output));
+				}
+			}
+			else
+			{
+				CompileModules(ModulesToCompile);
+			}
+			
+
+			Environment.CurrentDirectory = OldCWD;
+		}
+
+		/// <summary>
+		/// Compiles all script modules.
+		/// </summary>
+		/// <param name="Modules">Module project filenames.</param>
+		/// <param name="CompiledModuleFilenames">The resulting compiled module assembly filenames.</param>
+		private static void CompileModules(List<string> Modules)
+		{
+			Log.TraceInformation("Building script modules");
+			// Make sure DefaultScriptsDLLName is compiled first
+			var DefaultScriptsProjName = Path.ChangeExtension(DefaultScriptsDLLName, "csproj");
+			foreach (var ModuleName in Modules)
+			{
+				if (ModuleName.IndexOf(DefaultScriptsProjName, StringComparison.InvariantCultureIgnoreCase) >= 0)
+				{
+					Log.TraceInformation("Building script module: {0}", ModuleName);
+					try
+					{
+						CompileScriptModule(ModuleName);
+					}
+					catch (Exception Ex)
+					{
+						CommandUtils.Log(TraceEventType.Error, Ex);
+						throw new AutomationException("Failed to compile module {0}", ModuleName);
+					}
+					break;
+				}
+			}
+
+			// Second pass, compile everything else
+			foreach (var ModuleName in Modules)
+			{
+				if (ModuleName.IndexOf(DefaultScriptsProjName, StringComparison.InvariantCultureIgnoreCase) < 0)
+				{
+					Log.TraceInformation("Building script module: {0}", ModuleName);
+					try
+					{
+						CompileScriptModule(ModuleName);
+					}
+					catch (Exception Ex)
+					{
+						CommandUtils.Log(TraceEventType.Error, Ex);
+						throw new AutomationException("Failed to compile module {0}", ModuleName);
+					}
+				}
+			}
+		}
+
+		/// <summary>
+		/// Builds a script module (csproj file)
+		/// </summary>
+		/// <param name="ProjectFile"></param>
+		/// <returns></returns>
+		private static bool CompileScriptModule(string ProjectFile)
+		{
+			if (!ProjectFile.EndsWith(".csproj", StringComparison.InvariantCultureIgnoreCase))
+			{
+				throw new AutomationException(String.Format("Unable to build Project {0}. Not a valid .csproj file.", ProjectFile));
+			}
+			if (!CommandUtils.FileExists(ProjectFile))
+			{
+				throw new AutomationException(String.Format("Unable to build Project {0}. Project file not found.", ProjectFile));
+			}
+
+			var CmdLine = String.Format("\"{0}\" /verbosity:quiet /nologo /target:Build /property:Configuration={1} /property:Platform=AnyCPU /p:TreatWarningsAsErrors=true /p:BuildProjectReferences=true",
+				ProjectFile, BuildConfig);
+
+			// Compile the project
+			var Result = CommandUtils.Run(CommandUtils.CmdEnv.MsBuildExe, CmdLine);
+			if (Result.ExitCode != 0)
+			{
+				throw new AutomationException(String.Format("Failed to build \"{0}\":{1}{2}", ProjectFile, Environment.NewLine, Result.Output));
+			}
+			else
+			{
+				// Remove .Automation.csproj and copy to target dir
+				Log.TraceVerbose("Successfully compiled {0}", ProjectFile);
+			}
+			return Result.ExitCode == 0;
+		}
+
+		/// <summary>
+		/// Loads all precompiled assemblies (DLLs that end with *Scripts.dll).
+		/// </summary>
+		/// <param name="OutScriptAssemblies">List to store all loaded assemblies.</param>
+		private static void LoadPreCompiledScriptAssemblies(List<Assembly> OutScriptAssemblies)
+		{
+			CommandUtils.Log("Loading precompiled script DLLs");
+
+			bool DefaultScriptsDLLFound = false;
+			var ScriptsLocation = GetScriptAssemblyFolder();
+			if (CommandUtils.DirectoryExists(ScriptsLocation))
+			{
+				var ScriptDLLFiles = Directory.GetFiles(ScriptsLocation, "*.Automation.dll", SearchOption.AllDirectories);
+
+				CommandUtils.Log("Found {0} script DLL(s).", ScriptDLLFiles.Length);
+				foreach (var ScriptsDLLFilename in ScriptDLLFiles)
+				{
+
+					if (!HostPlatform.Current.IsScriptModuleSupported(CommandUtils.GetFilenameWithoutAnyExtensions(ScriptsDLLFilename)))
+					{
+						CommandUtils.LogVerbose("Script module {0} filtered by the Host Platform and will not be loaded.", ScriptsDLLFilename);
+						continue;
+					}
+					// Load the assembly into our app domain
+					CommandUtils.LogVerbose("Loading script DLL: {0}", ScriptsDLLFilename);
+					try
+					{
+						var Dll = AppDomain.CurrentDomain.Load(AssemblyName.GetAssemblyName(ScriptsDLLFilename));
+						OutScriptAssemblies.Add(Dll);
+						// Check if this is the default scripts DLL.
+						if (!DefaultScriptsDLLFound && String.Compare(Path.GetFileName(ScriptsDLLFilename), DefaultScriptsDLLName, true) == 0)
+						{
+							DefaultScriptsDLLFound = true;
+						}
+					}
+					catch (Exception Ex)
+					{
+						throw new AutomationException("Failed to load script DLL: {0}: {1}", ScriptsDLLFilename, Ex.Message);
+					}
+				}
+			}
+			else
+			{
+				CommandUtils.LogError("Scripts folder {0} does not exist!", ScriptsLocation);
+			}
+
+			// The default scripts DLL is required!
+			if (!DefaultScriptsDLLFound)
+			{
+				throw new AutomationException("{0} was not found or could not be loaded, can't run scripts.", DefaultScriptsDLLName);
+			}
+		}
+
+		private void CleanupScriptsAssemblies()
+		{
+			Log.TraceInformation("Cleaning up script DLL folder");
+			CommandUtils.DeleteDirectory(GetScriptAssemblyFolder());
+		}
+
+		private static string GetScriptAssemblyFolder()
+		{
+			return CommandUtils.CombinePaths(CommandUtils.CmdEnv.LocalRoot, "Engine", "Binaries", "DotNET", "AutomationScripts");
+		}
+
+		#endregion
+
+		#region Properties
+
+		public CaselessDictionary<Type> Commands
+		{
+			get { return ScriptCommands; }
+		}
+
+		#endregion
+	}
+}
diff --git a/Engine/Source/Programs/AutomationTool/AutomationUtils/UBTUtils.cs b/Engine/Source/Programs/AutomationTool/AutomationUtils/UBTUtils.cs
new file mode 100644
--- /dev/null
+++ b/Engine/Source/Programs/AutomationTool/AutomationUtils/UBTUtils.cs
@@ -0,0 +1,93 @@
+// Copyright 1998-2015 Epic Games, Inc. All Rights Reserved.
+using System;
+using System.Collections.Generic;
+using System.Text;
+
+namespace AutomationTool
+{
+
+	public partial class CommandUtils
+	{
+		/// <summary>
+		/// True if UBT has deleted junk files
+		/// </summary>
+		private static bool bJunkDeleted = false;
+
+		/// <summary>
+		/// Runs UBT with the specified commandline. Automatically creates a logfile. When 
+		/// no LogName is specified, the executable name is used as logfile base name.
+		/// </summary>
+		/// <param name="Env">Environment to use.</param>
+		/// <param name="CommandLine">Commandline to pass on to UBT.</param>
+		/// <param name="LogName">Optional logfile name.</param>
+        public static void RunUBT(CommandEnvironment Env, string UBTExecutable, string CommandLine, string LogName = null, Dictionary<string, string> EnvVars = null)
+		{
+			if (!FileExists(UBTExecutable))
+			{
+				throw new AutomationException("Unable to find UBT executable: " + UBTExecutable);
+			}
+
+			if (GlobalCommandLine.Rocket)
+			{
+				CommandLine += " -rocket";
+			}
+			if (!IsBuildMachine && UnrealBuildTool.BuildHostPlatform.Current.Platform == UnrealBuildTool.UnrealTargetPlatform.Mac)
+			{
+				CommandLine += " -nocreatestub";
+			}
+			CommandLine += " -NoHotReloadFromIDE";
+			if (bJunkDeleted || GlobalCommandLine.IgnoreJunk)
+			{
+				// UBT has already deleted junk files, make sure it doesn't do it again
+				CommandLine += " -ignorejunk";
+			}
+			else
+			{
+				// UBT will delete junk on first run
+				bJunkDeleted = true;
+			}
+
+			CommandUtils.RunAndLog(Env, UBTExecutable, CommandLine, LogName, EnvVars: EnvVars);
+		}
+
+		/// <summary>
+		/// Builds a UBT Commandline.
+		/// </summary>
+		/// <param name="Project">Unreal project to build (optional)</param>
+		/// <param name="Target">Target to build.</param>
+		/// <param name="Platform">Platform to build for.</param>
+		/// <param name="Config">Configuration to build.</param>
+		/// <param name="AdditionalArgs">Additional arguments to pass on to UBT.</param>
+		public static string UBTCommandline(string Project, string Target, string Platform, string Config, string AdditionalArgs = "")
+		{
+			string CmdLine;
+			if (String.IsNullOrEmpty(Project))
+			{
+				CmdLine = String.Format("{0} {1} {2} {3}", Target, Platform, Config, AdditionalArgs);
+			}
+			else
+			{
+				CmdLine = String.Format("{0} {1} {2} -Project={3} {4}", Target, Platform, Config, CommandUtils.MakePathSafeToUseWithCommandLine(Project), AdditionalArgs);
+			}
+			return CmdLine;
+		}
+
+		/// <summary>
+		/// Builds a target using UBT.  Automatically creates a logfile. When 
+		/// no LogName is specified, the executable name is used as logfile base name.
+		/// </summary>
+		/// <param name="Env">BuildEnvironment to use.</param>
+		/// <param name="Project">Unreal project to build (optional)</param>
+		/// <param name="Target">Target to build.</param>
+		/// <param name="Platform">Platform to build for.</param>
+		/// <param name="Config">Configuration to build.</param>
+		/// <param name="AdditionalArgs">Additional arguments to pass on to UBT.</param>
+		/// <param name="LogName">Optional logifle name.</param>
+		public static void RunUBT(CommandEnvironment Env, string UBTExecutable, string Project, string Target, string Platform, string Config, string AdditionalArgs = "", string LogName = null, Dictionary<string, string> EnvVars = null)
+		{
+			RunUBT(Env, UBTExecutable, UBTCommandline(Project, Target, Platform, Config, AdditionalArgs), LogName, EnvVars);
+		}
+
+	}
+
+}
\ No newline at end of file
