UEB-260 - Break AutomationTool into AutomationUtils that all automation projects depend on, and AutomationTool, which essentially only contains the startup code.

diff --git a/Engine/Source/Programs/AutomationTool/AutomationUtils/UE4Build.cs b/Engine/Source/Programs/AutomationTool/AutomationUtils/UE4Build.cs
new file mode 100644
--- /dev/null
+++ b/Engine/Source/Programs/AutomationTool/AutomationUtils/UE4Build.cs
@@ -0,0 +1,1638 @@
+// Copyright 1998-2015 Epic Games, Inc. All Rights Reserved.
+using System;
+using System.Collections.Generic;
+using System.Text;
+using System.IO;
+using System.Xml;
+using System.Diagnostics;
+
+namespace AutomationTool
+{
+	[Help("ForceMonolithic", "Toggle to combined the result into one executable")]
+	[Help("ForceDebugInfo", "Forces debug info even in development builds")]
+	[Help("NoXGE", "Toggle to disable the distributed build process")]
+	[Help("ForceNonUnity", "Toggle to disable the unity build system")]
+	[Help("ForceUnity", "Toggle to force enable the unity build system")]
+	[Help("Licensee", "If set, this build is being compiled by a licensee")]
+	public class UE4Build : CommandUtils
+	{
+		private BuildCommand OwnerCommand;
+
+		private String XGEConsoleExePath = "";
+		private String ParallelExecutorPath = "";
+
+		public void PrepareBuildProduct(string File, bool DeleteProduct = true)
+		{
+			if (!BuildProductFiles.Contains(File))
+			{
+				if (DeleteProduct && DeleteBuildProducts)
+				{
+					DeleteFile(File);
+				}
+			}
+		}
+
+		public bool HasBuildProduct(string InFile)
+		{
+			string File = CombinePaths(InFile);
+			foreach (var ExistingFile in BuildProductFiles)
+			{
+				if (ExistingFile.Equals(File, StringComparison.InvariantCultureIgnoreCase))
+				{
+					return true;
+				}
+			}
+			return false;
+		}
+
+		public void AddBuildProduct(string InFile)
+		{
+			string File = CombinePaths(InFile);
+			if (!FileExists(File))
+			{
+				throw new AutomationException("BUILD FAILED specified file to AddBuildProduct {0} does not exist.", File);
+			}
+			if (!HasBuildProduct(InFile))
+			{
+				BuildProductFiles.Add(File);
+			}
+		}
+
+		void PrepareManifest(string ManifestName, bool bAddReceipt)
+		{
+			if (FileExists(ManifestName) == false)
+			{
+				throw new AutomationException("BUILD FAILED UBT Manifest {0} does not exist.", ManifestName);
+			}
+			int FileNum = 0;
+			string OutFile;
+			while (true)
+			{
+				OutFile = CombinePaths(CmdEnv.LogFolder, String.Format("UBTManifest.{0}", FileNum) + Path.GetExtension(ManifestName));
+				FileInfo ItemInfo = new FileInfo(OutFile);
+				if (!ItemInfo.Exists)
+				{
+					break;
+				}
+				FileNum++;
+			}
+			CopyFile(ManifestName, OutFile);
+			Log("Copied UBT manifest to {0}", OutFile);
+
+
+			UnrealBuildTool.BuildManifest Manifest = ReadManifest(ManifestName);
+			foreach (string Item in Manifest.BuildProducts)
+			{
+				if(bAddReceipt && IsBuildReceipt(Item))
+				{
+					AddBuildProduct(Item);
+				}
+				else
+				{
+					PrepareBuildProduct(Item);
+				}
+			}
+		}
+
+		static bool IsBuildReceipt(string FileName)
+		{
+			return Path.GetDirectoryName(FileName).Replace('\\', '/').EndsWith("/Receipts", StringComparison.InvariantCultureIgnoreCase);
+		}
+
+		void AddBuildProductsFromManifest(string ManifestName)
+		{
+			if (FileExists(ManifestName) == false)
+			{
+				throw new AutomationException("BUILD FAILED UBT Manifest {0} does not exist.", ManifestName);
+			}
+			UnrealBuildTool.BuildManifest Manifest = ReadManifest(ManifestName);
+			foreach (string Item in Manifest.BuildProducts)
+			{
+				if (!FileExists_NoExceptions(Item))
+				{
+					throw new AutomationException("BUILD FAILED {0} was in manifest but was not produced.", Item);
+				}
+				AddBuildProduct(Item);
+			}
+		}
+
+	
+		/// True if UBT is compiled and ready to build!
+		private bool bIsUBTReady = false;
+		
+		private void PrepareUBT()
+		{			
+			// Don't build UBT if we're running with pre-compiled binaries and if there's a debugger attached to this process.
+			// With the debugger attached, even though deleting the exe will work, the pdb files are still locked and the build will fail.
+			// Also, if we're running from VS then since UAT references UBT, we already have the most up-to-date version of UBT.exe
+			if (!bIsUBTReady && !GlobalCommandLine.NoCompile && !System.Diagnostics.Debugger.IsAttached)
+			{
+				DeleteFile(UBTExecutable);
+
+				MsBuild(CmdEnv,
+						CmdEnv.LocalRoot + @"/Engine/Source/Programs/UnrealBuildTool/" + HostPlatform.Current.UBTProjectName + @".csproj",
+						"/verbosity:normal /target:Rebuild /property:Configuration=Development /property:Platform=AnyCPU",
+						"BuildUBT");
+
+				bIsUBTReady = true;
+			}
+
+			if (FileExists(UBTExecutable) == false)
+			{
+				throw new AutomationException("UBT does not exist in {0}.", UBTExecutable);
+			}
+		}
+
+		public class XGEItem
+		{
+			public UnrealBuildTool.BuildManifest Manifest;
+			public string CommandLine;
+			public UnrealBuildTool.UnrealTargetPlatform Platform;
+			public string Config;
+			public string ProjectName;
+			public string UProjectPath;
+			public List<string> XgeXmlFiles;
+			public string OutputCaption;
+		}
+
+		XGEItem XGEPrepareBuildWithUBT(string ProjectName, string TargetName, UnrealBuildTool.UnrealTargetPlatform Platform, string Config, string UprojectPath, bool ForceMonolithic = false, bool ForceNonUnity = false, bool ForceDebugInfo = false, string InAddArgs = "", bool ForceUnity = false, Dictionary<string, string> EnvVars = null)
+		{
+			string AddArgs = "";
+			if (string.IsNullOrEmpty(UprojectPath) == false)
+			{
+				AddArgs += " " + CommandUtils.MakePathSafeToUseWithCommandLine(UprojectPath);
+			}
+			AddArgs += " " + InAddArgs;
+			if (ForceMonolithic)
+			{
+				AddArgs += " -monolithic";
+			}
+			if (ForceNonUnity)
+			{
+				AddArgs += " -disableunity";
+			}
+			if (ForceUnity)
+			{
+				AddArgs += " -forceunity";
+			}
+			if (ForceDebugInfo)
+			{
+				AddArgs += " -forcedebuginfo";
+			}
+
+			PrepareUBT();
+
+            string UBTManifest = GetUBTManifest(UprojectPath, AddArgs);
+
+			DeleteFile(UBTManifest);
+			XGEItem Result = new XGEItem();
+
+			ClearExportedXGEXML();
+
+			RunUBT(CmdEnv, UBTExecutable: UBTExecutable, Project: ProjectName, Target: TargetName, Platform: Platform.ToString(), Config: Config, AdditionalArgs: "-generatemanifest -nobuilduht -xgeexport" + AddArgs, EnvVars: EnvVars);
+
+			PrepareManifest(UBTManifest, true);
+
+			Result.Platform = Platform;
+			Result.Config = Config;
+			Result.ProjectName = string.IsNullOrEmpty(ProjectName) ? TargetName : ProjectName; // use the target as the project if no project is set
+			Result.UProjectPath = UprojectPath;
+			Result.Manifest = ReadManifest(UBTManifest);
+			Result.OutputCaption = String.Format("{0}-{1}-{2}", Path.GetFileNameWithoutExtension(Result.ProjectName), Platform.ToString(), Config.ToString());
+			DeleteFile(UBTManifest);
+
+			Result.CommandLine = UBTExecutable + " " + UBTCommandline(Project: ProjectName, Target: TargetName, Platform: Platform.ToString(), Config: Config, AdditionalArgs: "-noxge -nobuilduht" + AddArgs);
+
+			Result.XgeXmlFiles = new List<string>();
+			foreach (var XGEFile in FindXGEFiles())
+			{
+				if (!FileExists_NoExceptions(XGEFile))
+				{
+					throw new AutomationException("BUILD FAILED: Couldn't find file: {0}", XGEFile);
+				}
+				int FileNum = 0;
+				string OutFile;
+				while (true)
+				{
+					OutFile = CombinePaths(CmdEnv.LogFolder, String.Format("UBTExport.{0}.xge.xml", FileNum));
+					FileInfo ItemInfo = new FileInfo(OutFile);
+					if (!ItemInfo.Exists)
+					{
+						break;
+					}
+					FileNum++;
+				}
+				CopyFile(XGEFile, OutFile);
+				Result.XgeXmlFiles.Add(OutFile);
+			}
+			ClearExportedXGEXML();
+			return Result;
+		}
+
+		void XGEFinishBuildWithUBT(XGEItem Item)
+		{
+			// allow the platform to perform any actions after building a target (seems almost like this should be done in UBT)
+			Platform.Platforms[Item.Platform].PostBuildTarget(this, Item.ProjectName, Item.UProjectPath, Item.Config);
+
+			foreach (string ManifestItem in Item.Manifest.BuildProducts)
+			{
+				if (!FileExists_NoExceptions(ManifestItem))
+				{
+					throw new AutomationException("BUILD FAILED {0} was in manifest but was not produced.", ManifestItem);
+				}
+				AddBuildProduct(ManifestItem);
+			}
+		}
+
+		void XGEDeleteBuildProducts(UnrealBuildTool.BuildManifest Manifest)
+		{
+			foreach (string Item in Manifest.BuildProducts)
+			{
+				DeleteFile(Item);
+			}
+		}
+
+		void CleanWithUBT(string ProjectName, string TargetName, UnrealBuildTool.UnrealTargetPlatform Platform, string Config, string UprojectPath, bool ForceMonolithic = false, bool ForceNonUnity = false, bool ForceDebugInfo = false, string InAddArgs = "", bool ForceUnity = false, Dictionary<string, string> EnvVars = null)
+		{
+			string AddArgs = "";
+			if (string.IsNullOrEmpty(UprojectPath) == false)
+			{
+				AddArgs += " " + CommandUtils.MakePathSafeToUseWithCommandLine(UprojectPath);
+			}
+			AddArgs += " " + InAddArgs;
+			if (ForceMonolithic)
+			{
+				AddArgs += " -monolithic";
+			}
+			if (ForceNonUnity)
+			{
+				AddArgs += " -disableunity";
+			}
+			if (ForceUnity)
+			{
+				AddArgs += " -forceunity";
+			}
+			if (ForceDebugInfo)
+			{
+				AddArgs += " -forcedebuginfo";
+			}
+			if (!TargetName.Equals("UnrealHeaderTool", StringComparison.InvariantCultureIgnoreCase))
+			{
+				AddArgs += " -nobuilduht";
+			}
+
+			PrepareUBT();
+
+			RunUBT(CmdEnv, UBTExecutable: UBTExecutable, Project: ProjectName, Target: TargetName, Platform: Platform.ToString(), Config: Config, AdditionalArgs: "-clean" + AddArgs, EnvVars: EnvVars);
+		}
+
+		void BuildWithUBT(string ProjectName, string TargetName, UnrealBuildTool.UnrealTargetPlatform TargetPlatform, string Config, string UprojectPath, bool ForceMonolithic = false, bool ForceNonUnity = false, bool ForceDebugInfo = false, bool ForceFlushMac = false, bool DisableXGE = false, string InAddArgs = "", bool ForceUnity = false, Dictionary<string, string> EnvVars = null)
+		{
+			string AddArgs = "";
+			if (string.IsNullOrEmpty(UprojectPath) == false)
+			{
+				AddArgs += " " + CommandUtils.MakePathSafeToUseWithCommandLine(UprojectPath);
+			}
+			AddArgs += " " + InAddArgs;
+			if (ForceMonolithic)
+			{
+				AddArgs += " -monolithic";
+			}
+			if (ForceNonUnity)
+			{
+				AddArgs += " -disableunity";
+			}
+			if (ForceUnity)
+			{
+				AddArgs += " -forceunity";
+			}
+			if (ForceDebugInfo)
+			{
+				AddArgs += " -forcedebuginfo";
+			}
+			if (ForceFlushMac)
+			{
+				AddArgs += " -flushmac";
+			}
+			if (DisableXGE)
+			{
+				AddArgs += " -noxge";
+			}
+
+			PrepareUBT();
+
+			// let the platform determine when to use the manifest
+            bool UseManifest = Platform.Platforms[TargetPlatform].ShouldUseManifestForUBTBuilds(AddArgs);
+
+			if (UseManifest)
+			{
+                string UBTManifest = GetUBTManifest(UprojectPath, AddArgs);
+
+				DeleteFile(UBTManifest);
+
+				RunUBT(CmdEnv, UBTExecutable: UBTExecutable, Project: ProjectName, Target: TargetName, Platform: TargetPlatform.ToString(), Config: Config, AdditionalArgs: AddArgs + " -generatemanifest" , EnvVars: EnvVars);
+
+				PrepareManifest(UBTManifest, false);
+			}
+
+			RunUBT(CmdEnv, UBTExecutable: UBTExecutable, Project: ProjectName, Target: TargetName, Platform: TargetPlatform.ToString(), Config: Config, AdditionalArgs: AddArgs, EnvVars: EnvVars);
+
+			// allow the platform to perform any actions after building a target (seems almost like this should be done in UBT)
+			Platform.Platforms[TargetPlatform].PostBuildTarget(this, string.IsNullOrEmpty(ProjectName) ? TargetName : ProjectName, UprojectPath, Config);
+
+			if (UseManifest)
+			{
+                string UBTManifest = GetUBTManifest(UprojectPath, AddArgs);
+
+				AddBuildProductsFromManifest(UBTManifest);
+
+				DeleteFile(UBTManifest);
+			}
+		}
+
+		string[] DotNetProductExtenstions()
+		{
+			return new string[] 
+			{
+				".dll",
+				".pdb",
+				".exe.config",
+				".exe",
+				"exe.mdb"
+			};
+		}
+
+		string[] SwarmBuildProducts()
+		{
+			return new string[]
+            {
+                "AgentInterface",
+                "SwarmAgent",
+                "SwarmCoordinator",
+                "SwarmCoordinatorInterface",
+                "SwarmInterface",
+				"SwarmCommonUtils"
+            };
+		}
+
+		void PrepareBuildProductsForCSharpProj(string CsProj)
+		{
+			string BaseOutput = CmdEnv.LocalRoot + @"/Engine/Binaries/DotNET/" + Path.GetFileNameWithoutExtension(CsProj);
+			foreach (var Ext in DotNetProductExtenstions())
+			{
+				PrepareBuildProduct(BaseOutput + Ext, false); // we don't delete these because we are not sure they are actually a build product
+			}
+		}
+
+		void AddBuildProductsForCSharpProj(string CsProj)
+		{
+			string BaseOutput = CmdEnv.LocalRoot + @"/Engine/Binaries/DotNET/" + Path.GetFileNameWithoutExtension(CsProj);
+			foreach (var Ext in DotNetProductExtenstions())
+			{
+				if (FileExists(BaseOutput + Ext))
+				{
+					AddBuildProduct(BaseOutput + Ext);
+				}
+			}
+		}
+
+		void AddIOSBuildProductsForCSharpProj(string CsProj)
+		{
+			string BaseOutput = CmdEnv.LocalRoot + @"/Engine/Binaries/DotNET/IOS/" + Path.GetFileNameWithoutExtension(CsProj);
+			foreach (var Ext in DotNetProductExtenstions())
+			{
+				if (FileExists(BaseOutput + Ext))
+				{
+					AddBuildProduct(BaseOutput + Ext);
+				}
+			}
+		}
+
+		void AddSwarmBuildProducts()
+		{
+			foreach (var SwarmProduct in SwarmBuildProducts())
+			{
+				string DotNETOutput = CmdEnv.LocalRoot + @"/Engine/Binaries/DotNET/" + SwarmProduct;
+				string Win64Output = CmdEnv.LocalRoot + @"/Engine/Binaries/Win64/" + SwarmProduct;
+				foreach (var Ext in DotNetProductExtenstions())
+				{
+					if (FileExists(DotNETOutput + Ext))
+					{
+						AddBuildProduct(DotNETOutput + Ext);
+					}
+				}
+				foreach (var Ext in DotNetProductExtenstions())
+				{
+					if (FileExists(Win64Output + Ext))
+					{
+						AddBuildProduct(Win64Output + Ext);
+					}
+				}
+			}
+		}
+
+		/// <summary>
+		/// Updates the engine version files
+		/// </summary>
+		public List<string> UpdateVersionFiles(bool ActuallyUpdateVersionFiles = true, int? ChangelistNumberOverride = null)
+		{
+			bool bIsLicenseeVersion = OwnerCommand.ParseParam("Licensee");
+			bool bDoUpdateVersionFiles = CommandUtils.P4Enabled && ActuallyUpdateVersionFiles;		
+			int ChangelistNumber = 0;
+			string ChangelistString = String.Empty;
+			if (bDoUpdateVersionFiles)
+			{
+				ChangelistNumber = ChangelistNumberOverride.GetValueOrDefault(P4Env.Changelist);
+				ChangelistString = ChangelistNumber.ToString();
+			}
+
+			var Result = new List<String>();
+			string Branch = P4Enabled ? P4Env.BuildRootEscaped : "";
+			{
+				string VerFile = CombinePaths(CmdEnv.LocalRoot, "Engine", "Build", "build.properties");
+				if (bDoUpdateVersionFiles)
+				{
+					Log("Updating {0} with:", VerFile);
+					Log("  TimestampForBVT={0}", CmdEnv.TimestampAsString);
+					Log("  EngineVersion={0}", ChangelistNumber.ToString());
+					Log("  BranchName={0}", Branch);					
+
+					PrepareBuildProduct(VerFile, false);
+					VersionFileUpdater BuildProperties = new VersionFileUpdater(VerFile);
+					BuildProperties.ReplaceLine("TimestampForBVT=", CmdEnv.TimestampAsString);
+					BuildProperties.ReplaceLine("EngineVersion=", ChangelistNumber.ToString());
+					BuildProperties.ReplaceLine("BranchName=", Branch);
+                    BuildProperties.Commit();
+				}
+				else
+				{
+					Log("{0} will not be updated because P4 is not enabled.", VerFile);
+				}
+				if (IsBuildMachine)
+				{
+					// Only BuildMachines can check this file in
+					AddBuildProduct(VerFile);
+				}
+				Result.Add(VerFile);
+			}
+
+			{
+				string VerFile = CombinePaths(CmdEnv.LocalRoot, "Engine", "Source", "Runtime", "Core", "Private", "UObject", "ObjectVersion.cpp");
+				if (bDoUpdateVersionFiles)
+				{
+					Log("Updating {0} with:", VerFile);
+					Log(" #define	ENGINE_VERSION  {0}", ChangelistNumber.ToString());
+
+					VersionFileUpdater ObjectVersionCpp = new VersionFileUpdater(VerFile);
+					ObjectVersionCpp.ReplaceLine("#define	ENGINE_VERSION	", ChangelistNumber.ToString());
+                    ObjectVersionCpp.Commit();
+                }
+				else
+				{
+					Log("{0} will not be updated because P4 is not enabled.", VerFile);
+				}
+                Result.Add(VerFile);
+            }
+
+            string EngineVersionFile = CombinePaths(CmdEnv.LocalRoot, "Engine", "Source", "Runtime", "Launch", "Resources", "Version.h");
+            {
+                string VerFile = EngineVersionFile;
+				if (bDoUpdateVersionFiles)
+				{
+					Log("Updating {0} with:", VerFile);
+					Log(" #define	ENGINE_VERSION  {0}", ChangelistNumber.ToString());
+					Log(" #define	BRANCH_NAME  {0}", Branch);
+					Log(" #define	BUILT_FROM_CHANGELIST  {0}", ChangelistString);
+					Log(" #define   ENGINE_IS_LICENSEE_VERSION  {0}", bIsLicenseeVersion ? "1" : "0");
+
+					VersionFileUpdater VersionH = new VersionFileUpdater(VerFile);
+					VersionH.ReplaceLine("#define ENGINE_VERSION ", ChangelistNumber.ToString());
+					VersionH.ReplaceLine("#define BRANCH_NAME ", "\"" + Branch + "\"");
+					VersionH.ReplaceLine("#define BUILT_FROM_CHANGELIST ", ChangelistString);
+					VersionH.ReplaceOrAddLine("#define ENGINE_IS_LICENSEE_VERSION ", bIsLicenseeVersion ? "1" : "0");
+
+                    VersionH.Commit();
+                }
+				else
+				{
+					Log("{0} will not be updated because P4 is not enabled.", VerFile);
+				}
+                Result.Add(VerFile);
+            }
+
+            {
+                // Use Version.h data to update MetaData.cs so the assemblies match the engine version.
+                string VerFile = CombinePaths(CmdEnv.LocalRoot, "Engine", "Source", "Programs", "DotNETCommon", "MetaData.cs");
+
+				if (bDoUpdateVersionFiles)
+                {
+                    // Get the MAJOR/MINOR/PATCH from the Engine Version file, as it is authoritative. The rest we get from the P4Env.
+                    string NewInformationalVersion = FEngineVersionSupport.FromVersionFile(EngineVersionFile, ChangelistNumber).ToString();
+
+                    Log("Updating {0} with AssemblyInformationalVersion: {1}", VerFile, NewInformationalVersion);
+
+                    VersionFileUpdater VersionH = new VersionFileUpdater(VerFile);
+                    VersionH.SetAssemblyInformationalVersion(NewInformationalVersion);
+                    VersionH.Commit();
+                }
+                else
+                {
+                    Log("{0} will not be updated because P4 is not enabled.", VerFile);
+                }
+                Result.Add(VerFile);
+            }
+
+			return Result;
+		}
+
+		/// <summary>
+		/// Updates the public key file (decryption)
+		/// </summary>
+		public void UpdatePublicKey(string KeyFilename)
+		{
+			var Lines = ReadEncryptionKeys(KeyFilename);
+			if (Lines != null && Lines.Length >= 3)
+			{
+				// Line0: Private key exponent, Line1: Modulus, Line2: Public key exponent.
+				string Modulus = String.Format("\"{0}\"", Lines[1]);
+				string PublicKeyExponent = String.Format("\"{0}\"", Lines[2]);
+				string VerFile = CmdEnv.LocalRoot + @"/Engine/Source/Runtime/PakFile/Private/PublicKey.inl";
+
+				Log("Updating {0} with:", VerFile);
+				Log(" #define DECRYPTION_KEY_EXPONENT {0}", PublicKeyExponent);
+				Log(" #define DECYRPTION_KEY_MODULUS {0}", Modulus);
+
+				VersionFileUpdater PublicKeyInl = new VersionFileUpdater(VerFile);
+				PublicKeyInl.ReplaceLine("#define DECRYPTION_KEY_EXPONENT ", PublicKeyExponent);
+				PublicKeyInl.ReplaceLine("#define DECYRPTION_KEY_MODULUS ", Modulus);
+                PublicKeyInl.Commit();
+			}
+			else
+			{
+				LogError("{0} doesn't look like a valid encryption key file or value.");
+			}
+		}
+
+		/// <summary>
+		/// Parses the encryption keys from the command line or loads them from a file.
+		/// </summary>
+		/// <param name="KeyFilenameOrValues">Key values or filename</param>
+		/// <returns>List of three encruption key values: private exponent, modulus and public exponent.</returns>
+		private string[] ReadEncryptionKeys(string KeyFilenameOrValues)
+		{
+			string[] Values;
+			if (KeyFilenameOrValues.StartsWith("0x", StringComparison.InvariantCultureIgnoreCase))
+			{
+				Values = KeyFilenameOrValues.Split(new char[] { '+' }, StringSplitOptions.RemoveEmptyEntries);
+			}
+			else
+			{
+				Values = ReadAllLines(KeyFilenameOrValues);
+			}
+			return Values;
+		}
+
+		[DebuggerDisplay("{TargetName} {Platform} {Config}")]
+		public class BuildTarget
+		{
+			/// Unreal project name, if needed
+			public string ProjectName;
+
+			/// Name of the target
+			public string TargetName;
+
+			/// For code-based projects with a .uproject, the TargetName isn't enough for UBT to find the target, this will point UBT to the target
+			public string UprojectPath;
+
+			/// Platform to build
+			public UnrealBuildTool.UnrealTargetPlatform Platform;
+
+			/// Configuration to build
+			public UnrealBuildTool.UnrealTargetConfiguration Config;
+
+			/// If true, force monolithic
+			public bool ForceMonolithic;
+
+			/// If true, force non-unity
+			public bool ForceNonUnity;
+
+			/// If true, force debug info, even in development builds (useful for distributed builds where we need symbols)
+			public bool ForceDebugInfo;
+
+			/// Extra UBT args
+			public string UBTArgs;
+		}
+
+
+		public class BuildAgenda
+		{
+			/// Full .NET solution files that we will compile and include in the build.  Currently we assume the output
+			/// binary file names match the solution file base name, but with various different binary file extensions.
+			public List<string> DotNetSolutions = new List<string>();
+
+			/// .NET .csproj files that will be compiled and included in the build.  Currently we assume the output
+			/// binary file names match the solution file base name, but with various different binary file extensions.
+			public List<string> DotNetProjects = new List<string>();
+
+			/// These are .NET project binary base file names that we want to include and label with the build, but
+			/// we won't be compiling these projects directly ourselves (however, they may be incidentally build or
+			/// published by a different project that we are building.)  We'll look for various .NET binary files with
+			/// this base file name but with different extensions.
+			public List<string> ExtraDotNetFiles = new List<string>();
+
+			/// These are .NET project binary files that are specific to IOS and found in the IOS subdirectory.  We define
+			/// these buildproducts as existing in the DotNET\IOS directory and assume that the output binary file names match
+			/// the solution file base name, but with various binary file extensions
+			public List<string> IOSDotNetProjects = new List<string>();
+
+			/// These are .NET project binary files that are specific to HTML5.  We define these build products as existing in the 
+			/// DotNET directory and assume that the output binary file names match
+			/// the solution file base name, but with various binary file extensions
+			public List<string> HTML5DotNetProjects = new List<string>();
+
+			public string SwarmProject = "";
+
+
+			/// List of targets to build.  These can be various Unreal projects, programs or libraries in various configurations
+			public List<BuildTarget> Targets = new List<BuildTarget>();
+
+			/// Used for special jobs to test things. At the moment, debugging some megaXGE thing
+			public bool SpecialTestFlag = false;
+
+			public bool DoRetries = IsBuildMachine;
+
+			/// <summary>
+			/// Adds a target with the specified configuration.
+			/// </summary>
+			/// <param name="TargetName">Name of the target</param>
+			/// <param name="InPlatform">Platform</param>
+			/// <param name="InConfiguration">Configuration</param>
+			/// <param name="InUprojectPath">Path to optional uproject file</param>
+			/// <param name="bForceMonolithic">Force monolithic build.</param>
+			/// <param name="bForceNonUnity">Force non-unity</param>
+			/// <param name="bForceDebugInfo">Force debug info even in development builds</param>
+			/// <param name="InAddArgs">Specifies additional arguments for UBT</param>
+			public void AddTarget(string TargetName, UnrealBuildTool.UnrealTargetPlatform InPlatform, UnrealBuildTool.UnrealTargetConfiguration InConfiguration, string InUprojectPath = null, bool bForceMonolithic = false, bool bForceNonUnity = false, bool bForceDebugInfo = false, string InAddArgs = "")
+			{
+				// Is this platform a compilable target?
+				if (!Platform.Platforms[InPlatform].CanBeCompiled())
+				{
+					return;
+				}
+
+				Targets.Add(new BuildTarget()
+				{
+					TargetName = TargetName,
+					Platform = InPlatform,
+					Config = InConfiguration,
+					UprojectPath = InUprojectPath,
+					ForceMonolithic = bForceMonolithic,
+					ForceNonUnity = bForceNonUnity,
+					ForceDebugInfo = bForceDebugInfo,
+					UBTArgs = InAddArgs,
+				});
+			}
+
+			/// <summary>
+			/// Adds multiple targets with the specified configuration.
+			/// </summary>
+			/// <param name="TargetNames">List of targets.</param>
+			/// <param name="InPlatform">Platform</param>
+			/// <param name="InConfiguration">Configuration</param>
+			/// <param name="InUprojectPath">Path to optional uproject file</param>
+			/// <param name="bForceMonolithic">Force monolithic build.</param>
+			/// <param name="bForceNonUnity">Force non-unity</param>
+			/// <param name="bForceDebugInfo">Force debug info even in development builds</param>
+			public void AddTargets(string[] TargetNames, UnrealBuildTool.UnrealTargetPlatform InPlatform, UnrealBuildTool.UnrealTargetConfiguration InConfiguration, string InUprojectPath = null, bool bForceMonolithic = false, bool bForceNonUnity = false, bool bForceDebugInfo = false, string InAddArgs = "")
+			{
+				// Is this platform a compilable target?
+				if (!Platform.Platforms[InPlatform].CanBeCompiled())
+				{
+					return;
+				}
+
+				foreach (var Target in TargetNames)
+				{
+					Targets.Add(new BuildTarget()
+					{
+						TargetName = Target,
+						Platform = InPlatform,
+						Config = InConfiguration,
+						UprojectPath = InUprojectPath,
+						ForceMonolithic = bForceMonolithic,
+						ForceNonUnity = bForceNonUnity,
+						ForceDebugInfo = bForceDebugInfo,
+						UBTArgs = InAddArgs,
+					});
+				}
+			}
+		}
+
+
+		public UE4Build(BuildCommand Command)
+		{
+			OwnerCommand = Command;
+			BuildProductFiles.Clear();
+		}
+
+		public List<string> FindXGEFiles()
+		{
+			var Result = new List<string>();
+			var Root = CombinePaths(CmdEnv.LocalRoot, @"\Engine\Intermediate\Build");			
+			Result.AddRange(FindFiles_NoExceptions("*.xge.xml", false, Root));
+			Result.Sort();
+			return Result;
+		}
+
+		public string XGEConsoleExe()
+		{
+			if (string.IsNullOrEmpty(XGEConsoleExePath))
+			{
+				string[] PathDirs = Environment.GetEnvironmentVariable("PATH").Split(Path.PathSeparator);
+				foreach (string PathDir in PathDirs)
+				{
+					string FullPath = Path.Combine(PathDir, "xgConsole.exe");
+					if (FileExists(FullPath))
+					{
+						XGEConsoleExePath = FullPath;
+						break;
+					}
+				}
+			}
+			return XGEConsoleExePath;
+		}
+
+		public string ParallelExecutorExec()
+		{
+			if (string.IsNullOrEmpty(ParallelExecutorPath))
+			{
+				string RootPath = Path.Combine(CmdEnv.LocalRoot, @"Engine\Binaries\Win64");
+				string FullPath = Path.Combine(RootPath, "ParallelExecutor.exe");
+				if (FileExists(FullPath))
+				{
+					ParallelExecutorPath = FullPath;
+				}
+			}
+			return ParallelExecutorPath;
+		}
+
+		public bool ProcessXGEItems(List<XGEItem> Actions, string XGETool, string Args, string TaskFilePath, bool DoRetries, bool SpecialTestFlag, bool ShowProgress)
+		{
+			if (string.IsNullOrEmpty(XGETool))
+			{
+				throw new AutomationException("Unable to find xge tool: " + XGETool);
+			}
+
+			XmlDocument XGETaskDocument;	
+			if (!CombineXGEItemFiles(Actions, TaskFilePath, out XGETaskDocument))
+			{
+				if (DeleteBuildProducts)
+				{
+					throw new AutomationException("Unable to find xge xmls: " + CombinePaths(CmdEnv.LogFolder, "*.xge.xml"));
+				}
+				else
+				{
+					Log("Incremental build, apparently everything was up to date, no XGE jobs produced.");
+				}
+			}
+			else
+			{
+				XmlElement EnvironmentsElement = XGETaskDocument.CreateElement("Environments");// Make sure all the tasks have a unique prefix
+				if(ShowProgress)
+				{
+					List<XmlElement> AllToolElements = new List<XmlElement>();
+					foreach(XmlElement EnvironmentElement in EnvironmentsElement.GetElementsByTagName("Environment"))
+					{
+						foreach(XmlElement ToolsElement in EnvironmentElement.GetElementsByTagName("Tools"))
+						{
+							foreach(XmlElement ToolElement in ToolsElement.GetElementsByTagName("Tool"))
+							{
+								AllToolElements.Add(ToolElement);
+							}
+						}
+					}
+					for(int Idx = 0; Idx < AllToolElements.Count; Idx++)
+					{
+						XmlElement ToolElement = AllToolElements[Idx];
+						if (ToolElement.HasAttribute("OutputPrefix"))
+						{
+							ToolElement.SetAttribute("OutputPrefix", ToolElement.Attributes["OutputPrefix"].Value + String.Format(" [@progress increment 1/{0}]", AllToolElements.Count));
+						}
+						else
+						{
+							ToolElement.SetAttribute("OutputPrefix", String.Format(" [@progress increment 1/{0} skipline]", AllToolElements.Count));
+						}
+					}
+				}
+
+				// Write the XGE task XML to a temporary file.
+				using (FileStream OutputFileStream = new FileStream(TaskFilePath, FileMode.Create, FileAccess.Write))
+				{
+					XGETaskDocument.Save(OutputFileStream);
+				}
+				if (!FileExists(TaskFilePath))
+				{
+					throw new AutomationException("Unable to find xge xml: " + TaskFilePath);
+				}
+
+				int Retries = DoRetries ? 2 : 1;
+                int ConnectionRetries = 4;
+                for (int i = 0; i < Retries; i++)
+				{
+					try
+					{
+                        while (true)
+                        {
+                            Log("Running {0} *******", XGETool);
+                            PushDir(CombinePaths(CmdEnv.LocalRoot, @"\Engine\Source"));
+                            int SuccesCode;
+                            string LogFile = GetRunAndLogLogName(CmdEnv, XGETool);
+                            string Output = RunAndLog(XGETool, Args, out SuccesCode, LogFile);
+                            PopDir();
+                            if (ConnectionRetries > 0 && (SuccesCode == 4 || SuccesCode == 2) && !Output.Contains("------Project:"))
+                            {
+                                Log(System.Diagnostics.TraceEventType.Warning, String.Format("{0} failure on the local connection timeout", XGETool));
+                                if (ConnectionRetries < 2)
+                                {
+                                    System.Threading.Thread.Sleep(60000);
+                                }
+                                ConnectionRetries--;
+                                continue;
+                            }
+                            else if (SuccesCode != 0)
+                            {
+                                Log("{0} did not succeed *******", XGETool);
+                                throw new AutomationException("Command failed (Result:{3}): {0} {1}. See logfile for details: '{2}' ",
+                                                                XGETool, Args, Path.GetFileName(LogFile), SuccesCode);
+                            }
+                            Log("{0} {1} Done *******", XGETool, Args);
+                            break;
+                        }
+                        break;
+					}
+					catch (Exception Ex)
+					{
+						Log("{0} failed on try {1}: {2}", XGETool, i + 1, Ex.ToString());
+						if (i + 1 >= Retries)
+						{
+							return false;
+						}
+						Log("Deleting build products to force a rebuild.");
+						foreach (var Item in Actions)
+						{
+							XGEDeleteBuildProducts(Item.Manifest);
+						}
+					}
+				}
+
+			}
+			foreach (var Item in Actions)
+			{
+				XGEFinishBuildWithUBT(Item);
+			}
+			return true;
+			}
+
+        private static bool CombineXGEItemFiles(List<XGEItem> Actions, string TaskFilePath, out XmlDocument XGETaskDocument)
+        {
+            XGETaskDocument = new XmlDocument();
+
+			// <BuildSet FormatVersion="1">...</BuildSet>
+			XmlElement BuildSetElement = XGETaskDocument.CreateElement("BuildSet");
+			XGETaskDocument.AppendChild(BuildSetElement);
+			BuildSetElement.SetAttribute("FormatVersion", "1");
+
+			// <Environments>...</Environments>
+			XmlElement EnvironmentsElement = XGETaskDocument.CreateElement("Environments");
+			BuildSetElement.AppendChild(EnvironmentsElement);
+
+			int Job = 0;
+			int Env = 0;
+			Dictionary<string, XmlElement> EnvStringToEnv = new Dictionary<string, XmlElement>();
+			Dictionary<string, XmlElement> EnvStringToProject = new Dictionary<string, XmlElement>();
+			Dictionary<string, string> ParamsToTool = new Dictionary<string, string>();
+			Dictionary<string, XmlElement> ParamsToToolElement = new Dictionary<string, XmlElement>();
+			Dictionary<string, string> ToolToAction = new Dictionary<string, string>();
+			foreach (var Item in Actions)
+			{
+				var CurrentDependencies = new List<string>();
+				foreach (var XGEFile in Item.XgeXmlFiles)
+				{
+					if (!FileExists_NoExceptions(XGEFile))
+					{
+						throw new AutomationException("BUILD FAILED: Couldn't find file: {0}", XGEFile);
+					}
+					var TargetFile = TaskFilePath + "." + Path.GetFileName(XGEFile);
+					CopyFile(XGEFile, TargetFile);
+					CopyFile_NoExceptions(XGEFile, TaskFilePath);
+					XmlDocument UBTTask = new XmlDocument();
+					UBTTask.XmlResolver = null;
+					UBTTask.Load(XGEFile);
+					DeleteFile(XGEFile);
+
+					var All = new List<string>();
+					{
+						var Elements = UBTTask.GetElementsByTagName("Variable");
+						foreach (XmlElement Element in Elements)
+						{
+							string Pair = Element.Attributes["Name"].Value + "=" + Element.Attributes["Value"].Value;
+							All.Add(Pair);
+						}
+					}
+					All.Sort();
+					string AllString = "";
+					foreach (string Element in All)
+					{
+						AllString += Element + "\n";
+					}
+					XmlElement ToolsElement;
+					XmlElement ProjectElement;
+
+					if (EnvStringToEnv.ContainsKey(AllString))
+					{
+						ToolsElement = EnvStringToEnv[AllString];
+						ProjectElement = EnvStringToProject[AllString];
+					}
+					else
+					{
+						string EnvName = string.Format("Env_{0}", Env);
+						Env++;
+						// <Environment Name="Win32">...</Environment>
+						XmlElement EnvironmentElement = XGETaskDocument.CreateElement("Environment");
+						EnvironmentsElement.AppendChild(EnvironmentElement);
+						EnvironmentElement.SetAttribute("Name", EnvName);
+
+						// <Tools>...</Tools>
+						ToolsElement = XGETaskDocument.CreateElement("Tools");
+						EnvironmentElement.AppendChild(ToolsElement);
+
+						{
+							// <Variables>...</Variables>
+							XmlElement VariablesElement = XGETaskDocument.CreateElement("Variables");
+							EnvironmentElement.AppendChild(VariablesElement);
+
+							var Elements = UBTTask.GetElementsByTagName("Variable");
+							foreach (XmlElement Element in Elements)
+							{
+								// <Variable>...</Variable>
+								XmlElement VariableElement = XGETaskDocument.CreateElement("Variable");
+								VariablesElement.AppendChild(VariableElement);
+								VariableElement.SetAttribute("Name", Element.Attributes["Name"].Value);
+								VariableElement.SetAttribute("Value", Element.Attributes["Value"].Value);
+							}
+						}
+
+						// <Project Name="Default" Env="Default">...</Project>
+						ProjectElement = XGETaskDocument.CreateElement("Project");
+						BuildSetElement.AppendChild(ProjectElement);
+						ProjectElement.SetAttribute("Name", EnvName);
+						ProjectElement.SetAttribute("Env", EnvName);
+
+						EnvStringToEnv.Add(AllString, ToolsElement);
+						EnvStringToProject.Add(AllString, ProjectElement);
+
+					}
+
+					Dictionary<string, string> ToolToTool = new Dictionary<string, string>();
+					Dictionary<string, string> ActionToAction = new Dictionary<string, string>();
+
+					{
+						var Elements = UBTTask.GetElementsByTagName("Tool");
+						foreach (XmlElement Element in Elements)
+						{
+							string Key = Element.Attributes["Path"].Value;
+							Key += " ";
+							Key += Element.Attributes["Params"].Value;
+
+							//hack hack hack
+							string ElementParams = Element.Attributes["Params"].Value;
+							if (!String.IsNullOrEmpty(ElementParams))
+							{
+								int YcIndex = ElementParams.IndexOf(" /Yc\"");
+								if (YcIndex >= 0)
+								{
+									// /Fp&quot;D:\BuildFarm\buildmachine_++depot+UE4\Engine\Intermediate\BuildData\Win64\UE4Editor\Development\SharedPCHs\CoreUObject.h.pch&quot
+									string Fp = " /Fp\"";
+									int FpIndex = ElementParams.IndexOf(Fp, YcIndex);
+									if (FpIndex >= 0)
+									{
+										int EndIndex = ElementParams.IndexOf("\"", FpIndex + Fp.Length);
+										if (EndIndex >= 0)
+										{
+											string PCHFileName = ElementParams.Substring(FpIndex + Fp.Length, EndIndex - FpIndex - Fp.Length);
+											if (PCHFileName.Contains(@"\SharedPCHs\") && PCHFileName.Contains(@"\UE4Editor\"))
+											{
+												Key = "SharedEditorPCH$ " + PCHFileName;
+												Log("Hack: detected Shared PCH, which will use a different key {0}", Key);
+
+											}
+										}
+									}
+								}
+
+							}
+
+							string ToolName = string.Format("{0}_{1}", Element.Attributes["Name"].Value, Job);
+							string OriginalToolName = ToolName;
+
+							if (ParamsToTool.ContainsKey(Key))
+							{
+								ToolName = ParamsToTool[Key];
+								ToolToTool.Add(OriginalToolName, ToolName);
+
+								XmlElement ToolElement = ParamsToToolElement[Key];
+								ToolElement.SetAttribute("GroupPrefix", ToolElement.Attributes["GroupPrefix"].Value + " + " + Item.OutputCaption);
+							}
+							else
+							{
+								// <Tool ... />
+								XmlElement ToolElement = XGETaskDocument.CreateElement("Tool");
+								ToolsElement.AppendChild(ToolElement);
+
+								ParamsToTool.Add(Key, ToolName);
+								ParamsToToolElement.Add(Key, ToolElement);
+
+								ToolElement.SetAttribute("Name", ToolName);
+								ToolElement.SetAttribute("AllowRemote", Element.Attributes["AllowRemote"].Value);
+								if (Element.HasAttribute("OutputPrefix"))
+								{
+									ToolElement.SetAttribute("OutputPrefix", Element.Attributes["OutputPrefix"].Value);
+								}
+								ToolElement.SetAttribute("GroupPrefix", "** For " + Item.OutputCaption);
+
+								ToolElement.SetAttribute("Params", Element.Attributes["Params"].Value);
+								ToolElement.SetAttribute("Path", Element.Attributes["Path"].Value);
+								if(Element.HasAttribute("VCCompiler"))
+								{
+									ToolElement.SetAttribute("VCCompiler", Element.Attributes["VCCompiler"].Value);
+								}
+								ToolElement.SetAttribute("SkipIfProjectFailed", Element.Attributes["SkipIfProjectFailed"].Value);
+								if (Element.HasAttribute("AutoReserveMemory"))
+								{
+									ToolElement.SetAttribute("AutoReserveMemory", Element.Attributes["AutoReserveMemory"].Value);
+								}
+								ToolElement.SetAttribute("OutputFileMasks", Element.Attributes["OutputFileMasks"].Value);
+								//ToolElement.SetAttribute("AllowRestartOnLocal", "false");  //vs2012 can't really restart, so we will go with this for now
+								if (Element.Attributes["OutputFileMasks"].Value == "PCLaunch.rc.res")
+								{
+									// total hack, when doing clean compiles, this output directory does not exist, we need to create it now
+									string Parms = Element.Attributes["Params"].Value;
+									string Start = "/fo \"";
+									int StartIndex = Parms.IndexOf(Start);
+									if (StartIndex >= 0)
+									{
+										Parms = Parms.Substring(StartIndex + Start.Length);
+										int EndIndex = Parms.IndexOf("\"");
+										if (EndIndex > 0)
+										{
+											string ResLocation = CombinePaths(Parms.Substring(0, EndIndex));
+											if (!DirectoryExists_NoExceptions(GetDirectoryName(ResLocation)))
+											{
+												CreateDirectory(GetDirectoryName(ResLocation));
+											}
+										}
+									}
+								}
+							}
+						}
+					}
+					{
+						var NextDependencies = new List<string>();
+
+						var Elements = UBTTask.GetElementsByTagName("Task");
+						foreach (XmlElement Element in Elements)
+						{
+							string ToolName = string.Format("{0}_{1}", Element.Attributes["Tool"].Value, Job);
+							string ActionName = string.Format("{0}_{1}", Element.Attributes["Name"].Value, Job);
+							string OriginalActionName = ActionName;
+
+							if (ToolToTool.ContainsKey(ToolName))
+							{
+								ToolName = ToolToTool[ToolName];
+								ActionName = ToolToAction[ToolName];
+								ActionToAction.Add(OriginalActionName, ActionName);
+							}
+							else
+							{
+								ActionToAction.Add(OriginalActionName, ActionName);
+								ToolToAction.Add(ToolName, ActionName);
+								// <Task ... />
+								XmlElement TaskElement = XGETaskDocument.CreateElement("Task");
+								ProjectElement.AppendChild(TaskElement);
+
+								TaskElement.SetAttribute("SourceFile", Element.Attributes["SourceFile"].Value);
+								if (Element.HasAttribute("Caption"))
+								{
+									TaskElement.SetAttribute("Caption", Element.Attributes["Caption"].Value);
+								}
+								TaskElement.SetAttribute("Name", ActionName);
+								NextDependencies.Add(ActionName);
+								TaskElement.SetAttribute("Tool", ToolName);
+								TaskElement.SetAttribute("WorkingDir", Element.Attributes["WorkingDir"].Value);
+								TaskElement.SetAttribute("SkipIfProjectFailed", Element.Attributes["SkipIfProjectFailed"].Value);
+
+								string NewDepends = "";
+								if (Element.HasAttribute("DependsOn"))
+								{
+									string Depends = Element.Attributes["DependsOn"].Value;
+									while (Depends.Length > 0)
+									{
+										string ThisAction = Depends;
+										int Semi = Depends.IndexOf(";");
+										if (Semi >= 0)
+										{
+											ThisAction = Depends.Substring(0, Semi);
+											Depends = Depends.Substring(Semi + 1);
+										}
+										else
+										{
+											Depends = "";
+										}
+										if (ThisAction.Length > 0)
+										{
+											if (NewDepends.Length > 0)
+											{
+												NewDepends += ";";
+											}
+											string ActionJob = ThisAction + string.Format("_{0}", Job);
+											if (!ActionToAction.ContainsKey(ActionJob))
+											{
+												throw new AutomationException("Action not found '{0}' in {1}->{2}", ActionJob, XGEFile, TargetFile);
+												// the XGE schedule is not topologically sorted. Hmmm. Well, we make a scary assumption here that this 
+											}
+											NewDepends += ActionToAction[ActionJob];
+										}
+									}
+								}
+								foreach (var Dep in CurrentDependencies)
+								{
+									if (NewDepends.Length > 0)
+									{
+										NewDepends += ";";
+									}
+									NewDepends += Dep;
+								}
+								if (NewDepends != "")
+								{
+									TaskElement.SetAttribute("DependsOn", NewDepends);
+								}
+							}
+
+						}
+						CurrentDependencies.AddRange(NextDependencies);
+					}
+					Job++;
+				}
+			}
+			return (Job > 0);
+		}
+
+		public void ClearExportedXGEXML()
+		{
+			foreach (var XGEFile in FindXGEFiles())
+			{
+				DeleteFile(XGEFile);
+			}
+		}
+
+		public bool CanUseXGE(UnrealBuildTool.UnrealTargetPlatform Platform)
+		{
+			if (!UnrealBuildTool.UEBuildPlatform.BuildPlatformDictionary.ContainsKey(Platform))
+			{
+				return false;
+			}
+
+			return UnrealBuildTool.UEBuildPlatform.GetBuildPlatform(Platform).CanUseXGE();
+		}
+
+		public string GenerateExternalFileList(BuildAgenda Agenda)
+		{
+			string AdditionalArguments = "";
+			foreach (var Target in Agenda.Targets)
+			{
+				RunUBT(CmdEnv, UBTExecutable, Target.UprojectPath, Target.TargetName, Target.Platform.ToString(), Target.Config.ToString(), "-generateexternalfilelist" + AdditionalArguments + " " + Target.UBTArgs);
+				AdditionalArguments = " -mergeexternalfilelist";
+			}
+			return CommandUtils.CombinePaths(CmdEnv.LocalRoot, @"/Engine/Intermediate/Build/ExternalFiles.xml");
+		}
+
+		public string[] GetExternalFileList(BuildAgenda Agenda)
+		{
+			string FileListPath = GenerateExternalFileList(Agenda);
+			return UnrealBuildTool.Utils.ReadClass<UnrealBuildTool.ExternalFileList>(FileListPath).FileNames.ToArray();
+		}
+
+		/// <summary>
+		/// Executes a build.
+		/// </summary>
+		/// <param name="Agenda">Build agenda.</param>
+		/// <param name="InDeleteBuildProducts">if specified, determines if the build products will be deleted before building. If not specified -clean parameter will be used,</param>
+		/// <param name="InUpdateVersionFiles">True if the version files are to be updated </param>
+		/// <param name="InForceNoXGE">If true will force XGE off</param>
+		/// <param name="InUseParallelExecutor">If true AND XGE not present or not being used then use ParallelExecutor</param>
+		public void Build(BuildAgenda Agenda, bool? InDeleteBuildProducts = null, bool InUpdateVersionFiles = true, bool InForceNoXGE = false, bool InUseParallelExecutor = false, bool InForceNonUnity = false, bool InForceUnity = false, bool InShowProgress = false, Dictionary<UnrealBuildTool.UnrealTargetPlatform, Dictionary<string, string>> PlatformEnvVars = null)
+		{
+			if (!CmdEnv.HasCapabilityToCompile)
+			{
+				throw new AutomationException("You are attempting to compile on a machine that does not have a supported compiler!");
+			}
+			DeleteBuildProducts = InDeleteBuildProducts.HasValue ? InDeleteBuildProducts.Value : OwnerCommand.ParseParam("Clean");
+			if (InUpdateVersionFiles)
+			{
+				UpdateVersionFiles();
+			}
+
+			{
+				var EncryptionKeyFilename = OwnerCommand.ParseParamValue("SignPak");
+				if (String.IsNullOrEmpty(EncryptionKeyFilename) == false)
+				{
+					UpdatePublicKey(EncryptionKeyFilename);
+				}
+			}
+
+			//////////////////////////////////////
+
+			// make a set of unique platforms involved
+			var UniquePlatforms = new List<UnrealBuildTool.UnrealTargetPlatform>();
+			foreach (var Target in Agenda.Targets)
+			{
+				if (!UniquePlatforms.Contains(Target.Platform))
+				{
+					UniquePlatforms.Add(Target.Platform);
+				}
+			}
+
+			// allow all involved platforms to hook into the agenda
+			foreach (var TargetPlatform in UniquePlatforms)
+			{
+				Platform.Platforms[TargetPlatform].PreBuildAgenda(this, Agenda);
+			}
+
+
+			foreach (var File in Agenda.ExtraDotNetFiles)
+			{
+				PrepareBuildProductsForCSharpProj(Path.Combine(CmdEnv.LocalRoot, File));
+			}
+
+			if (Agenda.SwarmProject != "")
+			{
+				string SwarmSolution = Path.Combine(CmdEnv.LocalRoot, Agenda.SwarmProject);
+				PrepareBuildProductsForCSharpProj(SwarmSolution);
+
+				BuildSolution(CmdEnv, SwarmSolution);
+				AddSwarmBuildProducts();
+			}
+
+			foreach (var DotNetSolution in Agenda.DotNetSolutions)
+			{
+				string Solution = Path.Combine(CmdEnv.LocalRoot, DotNetSolution);
+
+				PrepareBuildProductsForCSharpProj(Solution);
+
+				BuildSolution(CmdEnv, Solution);
+
+				AddBuildProductsForCSharpProj(Solution);
+			}
+
+			foreach (var DotNetProject in Agenda.DotNetProjects)
+			{
+				string CsProj = Path.Combine(CmdEnv.LocalRoot, DotNetProject);
+
+				PrepareBuildProductsForCSharpProj(CsProj);
+
+				BuildCSharpProject(CmdEnv, CsProj);
+
+				AddBuildProductsForCSharpProj(CsProj);
+			}
+
+			foreach (var IOSDotNetProject in Agenda.IOSDotNetProjects)
+			{
+				string IOSCsProj = Path.Combine(CmdEnv.LocalRoot, IOSDotNetProject);
+
+				PrepareBuildProductsForCSharpProj(IOSCsProj);
+
+				BuildCSharpProject(CmdEnv, IOSCsProj);
+
+				AddIOSBuildProductsForCSharpProj(IOSCsProj);
+			}
+
+			foreach (var HTML5DotNetProject in Agenda.HTML5DotNetProjects)
+			{
+				string CsProj = Path.Combine(CmdEnv.LocalRoot, HTML5DotNetProject);
+
+				PrepareBuildProductsForCSharpProj(CsProj);
+
+				BuildCSharpProject(CmdEnv, CsProj);
+
+				AddBuildProductsForCSharpProj(CsProj);
+			}
+
+			foreach (var File in Agenda.ExtraDotNetFiles)
+			{
+				AddBuildProductsForCSharpProj(Path.Combine(CmdEnv.LocalRoot, File));
+			}
+
+			bool bForceMonolithic = OwnerCommand.ParseParam("ForceMonolithic");
+			bool bForceNonUnity = OwnerCommand.ParseParam("ForceNonUnity") || InForceNonUnity;
+			bool bForceUnity = OwnerCommand.ParseParam("ForceUnity") || InForceUnity;
+			bool bForceDebugInfo = OwnerCommand.ParseParam("ForceDebugInfo");
+			bool bUsedXGE = false;
+			bool bCanUseXGE = !OwnerCommand.ParseParam("NoXGE") && !InForceNoXGE;
+			string XGEConsole = XGEConsoleExe();
+			if (bCanUseXGE && string.IsNullOrEmpty(XGEConsole))
+			{
+				Log("XGE was requested, but is unavailable, so we won't use it.");
+				bCanUseXGE = false;
+			}
+
+			// only run ParallelExecutor if not running XGE (and we've requested ParallelExecutor and it exists)
+			bool bCanUseParallelExecutor = !bCanUseXGE && InUseParallelExecutor && !string.IsNullOrEmpty(ParallelExecutorExec());
+			Log("************************* UE4Build:");
+			Log("************************* ForceMonolithic: {0}", bForceMonolithic);
+			Log("************************* ForceNonUnity:{0} ", bForceNonUnity);
+			Log("************************* ForceDebugInfo: {0}", bForceDebugInfo);
+			Log("************************* UseXGE: {0}", bCanUseXGE);
+            Log("************************* UseParallelExecutor: {0}", bCanUseParallelExecutor);
+
+			// process XGE files
+			if (bCanUseXGE || bCanUseParallelExecutor)
+			{
+				string TaskFilePath = CombinePaths(CmdEnv.LogFolder, @"UAT_XGE.xml");
+
+				if (DeleteBuildProducts)
+				{
+					foreach (var Target in Agenda.Targets)
+					{
+						CleanWithUBT(Target.ProjectName, Target.TargetName, Target.Platform, Target.Config.ToString(), Target.UprojectPath, bForceMonolithic, bForceNonUnity, bForceDebugInfo, Target.UBTArgs, bForceUnity);
+					}
+				}
+				foreach (var Target in Agenda.Targets)
+				{
+					if (Target.TargetName == "UnrealHeaderTool" || !CanUseXGE(Target.Platform))
+					{
+						// When building a target for Mac or iOS, use UBT's -flushmac option to clean up the remote builder
+						bool bForceFlushMac = DeleteBuildProducts && (Target.Platform == UnrealBuildTool.UnrealTargetPlatform.Mac || Target.Platform == UnrealBuildTool.UnrealTargetPlatform.IOS);
+						LogSetProgress(InShowProgress, "Building header tool...");
+						BuildWithUBT(Target.ProjectName, Target.TargetName, Target.Platform, Target.Config.ToString(), Target.UprojectPath, bForceMonolithic, bForceNonUnity, bForceDebugInfo, bForceFlushMac, true, Target.UBTArgs, bForceUnity);
+					}
+				}
+
+				List<XGEItem> XGEItems = new List<XGEItem>();
+				LogSetProgress(InShowProgress, "Generating headers...");
+				foreach (var Target in Agenda.Targets)
+				{
+					if (Target.TargetName != "UnrealHeaderTool" && CanUseXGE(Target.Platform))
+					{
+						XGEItem Item = XGEPrepareBuildWithUBT(Target.ProjectName, Target.TargetName, Target.Platform, Target.Config.ToString(), Target.UprojectPath, bForceMonolithic, bForceNonUnity, bForceDebugInfo, Target.UBTArgs, bForceUnity);
+						XGEItems.Add(Item);
+					}
+				}
+				if (XGEItems.Count > 0)
+				{
+					string XGETool = "";
+					string Args = "";
+					if (bCanUseXGE) 
+					{
+						XGETool = XGEConsoleExePath;
+						Args = "\"" + TaskFilePath + "\" /Rebuild /MaxCPUS=200";
+					}
+					else if (bCanUseParallelExecutor)
+					{
+						XGETool = ParallelExecutorExec();
+						Args = "\"" + TaskFilePath + "\"";
+					}
+					if (XGETool.Equals(""))
+					{
+						throw new AutomationException("BUILD FAILED: no tool present to process XGE files");
+					}
+
+					LogSetProgress(InShowProgress, "Building...");
+					if (!ProcessXGEItems(XGEItems, XGETool, Args, TaskFilePath, Agenda.DoRetries, Agenda.SpecialTestFlag, InShowProgress))
+					{
+						throw new AutomationException("BUILD FAILED: {0} failed, retries not enabled:", XGETool);
+					}
+					else
+					{
+						bUsedXGE = true;
+					}
+				}
+				else
+				{
+					bUsedXGE = true; // if there was nothing to build, we still consider this a success
+				}
+			}
+
+			if (!bUsedXGE)
+			{
+				if (DeleteBuildProducts)
+				{
+					foreach (var Target in Agenda.Targets)
+					{
+						CleanWithUBT(Target.ProjectName, Target.TargetName, Target.Platform, Target.Config.ToString(), Target.UprojectPath, bForceMonolithic, bForceNonUnity, bForceDebugInfo, Target.UBTArgs, bForceUnity);
+					}
+				}
+				foreach (var Target in Agenda.Targets)
+				{
+					// When building a target for Mac or iOS, use UBT's -flushmac option to clean up the remote builder
+					bool bForceFlushMac = DeleteBuildProducts && (Target.Platform == UnrealBuildTool.UnrealTargetPlatform.Mac || Target.Platform == UnrealBuildTool.UnrealTargetPlatform.IOS);
+					BuildWithUBT(Target.ProjectName, Target.TargetName, Target.Platform, Target.Config.ToString(), Target.UprojectPath, bForceMonolithic, bForceNonUnity, bForceDebugInfo, bForceFlushMac, true, Target.UBTArgs, bForceUnity);
+				}
+			}
+
+			// NOTE: OK, we're done building projects.  You can build additional targets after this.  Call FinalizebuildProducts() when done.
+		}
+
+
+		/// <summary>
+		/// Checks to make sure there was at least one build product, and that all files exist.  Also, logs them all out.
+		/// </summary>
+		/// <param name="BuildProductFiles">List of files</param>
+		public static void CheckBuildProducts(List<string> BuildProductFiles)
+		{
+			// Check build products
+			{
+				Log("Build products *******");
+				if (BuildProductFiles.Count < 1)
+				{
+					Log("No build products were made");
+				}
+				else
+				{
+					foreach (var Product in BuildProductFiles)
+					{
+						if (!FileExists(Product))
+						{
+							throw new AutomationException("BUILD FAILED {0} was a build product but no longer exists", Product);
+						}
+						Log(Product);
+					}
+				}
+				Log("End Build products *******");
+			}
+		}
+
+
+		/// <summary>
+		/// Adds or edits existing files at head revision, expecting an exclusive lock, resolving by clobbering any existing version
+		/// </summary>
+		/// <param name="Files">List of files to check out</param>
+		public static void AddBuildProductsToChangelist(int WorkingCL, List<string> Files)
+		{
+			Log("Adding {0} build products to changelist {1}...", Files.Count, WorkingCL);
+			foreach (var File in Files)
+			{
+				P4.Sync("-f -k " + CommandUtils.MakePathSafeToUseWithCommandLine(File) + "#head"); // sync the file without overwriting local one
+				if (!FileExists(File))
+				{
+					throw new AutomationException("BUILD FAILED {0} was a build product but no longer exists", File);
+				}
+
+				P4.ReconcileNoDeletes(WorkingCL, CommandUtils.MakePathSafeToUseWithCommandLine(File));
+
+				// Change file type on binary files to be always writeable.
+				var FileStats = P4.FStat(File);
+
+                if (IsProbablyAMacOrIOSExe(File))
+                {
+                    if (FileStats.Type == P4FileType.Binary && (FileStats.Attributes & (P4FileAttributes.Executable | P4FileAttributes.Writeable)) != (P4FileAttributes.Executable | P4FileAttributes.Writeable))
+                    {
+                        P4.ChangeFileType(File, (P4FileAttributes.Executable | P4FileAttributes.Writeable));
+                    }
+                }
+                else
+                {
+					if (IsBuildProduct(File, FileStats) && (FileStats.Attributes & P4FileAttributes.Writeable) != P4FileAttributes.Writeable)
+                    {
+                        P4.ChangeFileType(File, P4FileAttributes.Writeable);
+                    }
+
+                }                    
+			}
+		}
+
+		/// <summary>
+		/// Determines if this file is a build product.
+		/// </summary>
+		/// <param name="File">File path</param>
+		/// <param name="FileStats">P4 file stats.</param>
+		/// <returns>True if this is a Windows build product. False otherwise.</returns>
+		private static bool IsBuildProduct(string File, P4FileStat FileStats)
+		{
+			if(FileStats.Type == P4FileType.Binary || IsBuildReceipt(File))
+			{
+				return true;
+			}
+
+			return FileStats.Type == P4FileType.Text && File.EndsWith(".exe.config", StringComparison.InvariantCultureIgnoreCase);
+		}
+
+		/// <summary>
+		/// Add UBT files to build products
+		/// </summary>
+		public void AddUBTFilesToBuildProducts()
+		{
+			if (GlobalCommandLine.NoCompile)
+			{
+				Log("We are being asked to copy the UBT build products, but we are running precompiled, so this does not make much sense.");
+			}
+
+			var UBTLocation = Path.GetDirectoryName(GetUBTExecutable());
+			var UBTFiles = new List<string>(new string[] 
+					{
+						"UnrealBuildTool.exe",
+						"UnrealBuildTool.exe.config",
+						"EnvVarsToXML.exe",
+						"EnvVarsToXML.exe.config"
+					});
+
+			foreach (var UBTFile in UBTFiles)
+			{
+				var UBTProduct = CombinePaths(UBTLocation, UBTFile);
+				if (!FileExists_NoExceptions(UBTProduct))
+				{
+					throw new AutomationException("Cannot add UBT to the build products because {0} does not exist.", UBTProduct);
+				}
+				AddBuildProduct(UBTProduct);
+			}
+		}
+
+		/// <summary>
+		/// Copy the UAT files to their precompiled location, and add them as build products
+		/// </summary>
+		public void AddUATFilesToBuildProducts()
+		{
+			// Find all DLLs (scripts and their dependencies)
+			const string ScriptsPostfix = ".dll";
+            var DotNetOutputLocation = CombinePaths(CmdEnv.LocalRoot, "Engine", "Binaries", "DotNET");
+			var UATScriptsLocation = CombinePaths(DotNetOutputLocation, "AutomationScripts");
+			var UATScripts = Directory.GetFiles(UATScriptsLocation, "*" + ScriptsPostfix, SearchOption.AllDirectories);
+			if (IsNullOrEmpty(UATScripts))
+			{
+				throw new AutomationException("No automation scripts found in {0}. Cannot add UAT files to the build products.", UATScriptsLocation);
+			}
+
+			var UATFiles = new List<string>(new string[] 
+					{
+						"AutomationTool.exe",
+						"AutomationTool.exe.config",
+						"UnrealBuildTool.exe",
+						"UnrealBuildTool.exe.config",
+						"DotNETUtilities.dll",
+					});
+
+			foreach (var UATFile in UATFiles)
+			{
+				var OutputFile = CombinePaths(DotNetOutputLocation, UATFile);
+				if (!FileExists_NoExceptions(OutputFile))
+				{
+					throw new AutomationException("Cannot add UAT to the build products because {0} does not exist.", OutputFile);
+				}
+				AddBuildProduct(OutputFile);
+			}
+
+			// All scripts are expected to exist in DotNET/AutomationScripts subfolder.
+			foreach (var UATScriptFilePath in UATScripts)
+			{
+				if (!FileExists_NoExceptions(UATScriptFilePath))
+				{
+					throw new AutomationException("Cannot add UAT to the build products because {0} does not exist.", UATScriptFilePath);
+				}
+
+				AddBuildProduct(UATScriptFilePath);
+			}
+		}
+
+		string GetUBTManifest(string UProjectPath, string InAddArgs)
+		{
+			// Can't write to Engine directory on 
+            bool bForceRocket = InAddArgs.ToLowerInvariant().Contains(" -rocket"); //awful
+
+            if ((GlobalCommandLine.Rocket || bForceRocket) && !String.IsNullOrEmpty(UProjectPath))
+			{
+				return Path.Combine(Path.GetDirectoryName(UProjectPath), "Intermediate/Build/Manifest.xml");				
+			}
+			else
+			{
+				return CommandUtils.CombinePaths(CmdEnv.LocalRoot, @"/Engine/Intermediate/Build/Manifest.xml");
+			}
+		}
+
+		public static string GetUBTExecutable()
+		{
+			return CommandUtils.CombinePaths(CmdEnv.LocalRoot, @"Engine/Binaries/DotNET/UnrealBuildTool.exe");
+		}
+
+		public string UBTExecutable
+		{
+			get
+			{
+				return GetUBTExecutable();							
+			}
+		}
+
+		// List of everything we built so far
+		public readonly List<string> BuildProductFiles = new List<string>();
+
+		private bool DeleteBuildProducts = true;
+	}
+}
diff --git a/Engine/Source/Programs/AutomationTool/AutomationUtils/Utils.cs b/Engine/Source/Programs/AutomationTool/AutomationUtils/Utils.cs
new file mode 100644
--- /dev/null
+++ b/Engine/Source/Programs/AutomationTool/AutomationUtils/Utils.cs
@@ -0,0 +1,1074 @@
+﻿// Copyright 1998-2015 Epic Games, Inc. All Rights Reserved.
+using System;
+using System.Collections.Generic;
+using System.Linq;
+using System.Text;
+using System.IO;
+using System.Threading;
+using System.Diagnostics;
+using System.Reflection;
+using UnrealBuildTool;
+using System.Text.RegularExpressions;
+
+namespace AutomationTool
+{
+	#region UAT Internal Utils
+
+	/// <summary>
+	/// AutomationTool internal Utilities.
+	/// </summary>
+	public static class InternalUtils
+	{
+		/// <summary>
+		/// Gets environment variable value.
+		/// </summary>
+		/// <param name="VarName">Variable name.</param>
+		/// <param name="Default">Default value to be returned if the variable does not exist.</param>
+		/// <returns>Variable value or the default value if the variable did not exist.</returns>
+		public static string GetEnvironmentVariable(string VarName, string Default, bool bQuiet = false)
+		{
+			var Value = Environment.GetEnvironmentVariable(VarName);
+			if (Value == null)
+			{
+				Value = Default;
+			}
+			if (!bQuiet)
+			{
+				Log.WriteLine(TraceEventType.Verbose, "GetEnvironmentVariable {0}={1}", VarName, Value);
+			}
+			return Value;
+		}
+
+		/// <summary>
+		/// Creates a directory.
+		/// </summary>
+		/// <param name="Path">Directory name.</param>
+		/// <returns>True if the directory was created, false otherwise.</returns>
+		public static bool SafeCreateDirectory(string Path, bool bQuiet = false)
+		{
+			if( !bQuiet)
+			{ 
+				Log.WriteLine(TraceEventType.Verbose, "SafeCreateDirectory {0}", Path);
+			}
+
+			bool Result = true;
+			try
+			{
+				Result = Directory.CreateDirectory(Path).Exists;
+			}
+			catch (Exception)
+			{
+				if (Directory.Exists(Path) == false)
+				{
+					Result = false;
+				}
+			}
+			return Result;
+		}
+
+		/// <summary>
+		/// Deletes a file (will remove read-only flag if necessary).
+		/// </summary>
+		/// <param name="Path">Filename</param>
+		/// <param name="bQuiet">if true, then do not print errors, also in quiet mode do not retry</param>
+		/// <returns>True if the file does not exist, false otherwise.</returns>
+		public static bool SafeDeleteFile(string Path, bool bQuiet = false)
+		{
+			if (!bQuiet)
+			{
+				Log.WriteLine(TraceEventType.Verbose, "SafeDeleteFile {0}", Path);
+			}
+			int MaxAttempts = bQuiet ? 1 : 10;
+			int Attempts = 0;
+			bool Result = true;
+			Exception LastException = null;
+
+			do
+			{
+				Result = true;
+				try
+				{
+					if (File.Exists(Path))
+					{
+						FileAttributes Attributes = File.GetAttributes(Path);
+						if ((Attributes & FileAttributes.ReadOnly) != 0)
+						{
+							File.SetAttributes(Path, Attributes & ~FileAttributes.ReadOnly);
+						}
+						File.Delete(Path);
+					}
+				}
+				catch (Exception Ex)
+				{
+					if (File.Exists(Path))
+					{
+						Result = false;
+					}
+					LastException = Ex;
+				}
+				if (Result == false && Attempts + 1 < MaxAttempts)
+				{
+					System.Threading.Thread.Sleep(1000);
+				}
+			} while (Result == false && ++Attempts < MaxAttempts);
+
+			if (Result == false && LastException != null)
+			{
+				if (bQuiet)
+				{
+					Log.WriteLine(TraceEventType.Verbose, "Failed to delete file {0} in {1} attempts.", Path, MaxAttempts);
+				}
+				else
+				{
+					Log.WriteLine(TraceEventType.Warning, "Failed to delete file {0} in {1} attempts.", Path, MaxAttempts);
+					Log.WriteLine(TraceEventType.Warning, LogUtils.FormatException(LastException));
+				}
+			}
+
+			return Result;
+		}
+
+		/// <summary>
+		/// Recursively deletes a directory and all its files and subdirectories.
+		/// </summary>
+		/// <param name="Path">Path to delete.</param>
+		/// <returns>Whether the deletion was succesfull.</returns>
+		private static bool RecursivelyDeleteDirectory(string Path, bool bQuiet = false)
+		{
+			if (!bQuiet)
+			{
+				Log.WriteLine(TraceEventType.Verbose, "RecursivelyDeleteDirectory {0}", Path);
+			}
+			// Delete all files. This will also delete read-only files.
+			var FilesInDirectory = Directory.EnumerateFiles(Path);
+			foreach (string Filename in FilesInDirectory)
+			{
+				if (SafeDeleteFile(Filename, bQuiet) == false)
+				{
+					return false;
+				}
+			}
+
+			// Recursively delete all files from sub-directories.
+			var FoldersInDirectory = Directory.EnumerateDirectories(Path);
+			foreach (string Folder in FoldersInDirectory)
+			{
+				if (RecursivelyDeleteDirectory(Folder, bQuiet) == false)
+				{
+					return false;
+				}
+			}
+
+			// At this point there should be no read-only files in any of the directories and
+			// this directory should be empty too.
+			return SafeDeleteEmptyDirectory(Path, bQuiet);
+		}
+
+		/// <summary>
+		/// Deletes an empty directory.
+		/// </summary>
+		/// <param name="Path">Path to the Directory.</param>
+		/// <returns>True if deletion was successful, otherwise false.</returns>
+		public static bool SafeDeleteEmptyDirectory(string Path, bool bQuiet = false)
+		{
+			if (!bQuiet)
+			{
+				Log.WriteLine(TraceEventType.Verbose, "SafeDeleteEmptyDirectory {0}", Path);
+			}
+			const int MaxAttempts = 10;
+			int Attempts = 0;
+			bool Result = true;
+			Exception LastException = null;
+			do
+			{
+				Result = !Directory.Exists(Path);
+				if (!Result)
+				{
+					try
+					{
+						Directory.Delete(Path, true);
+					}
+					catch (Exception Ex)
+					{
+						if (Directory.Exists(Path))
+						{
+							Thread.Sleep(3000);
+						}
+						Result = !Directory.Exists(Path);
+						LastException = Ex;
+					}
+				}
+			} while (Result == false && ++Attempts < MaxAttempts);
+
+			if (Result == false && LastException != null)
+			{
+				Log.WriteLine(TraceEventType.Warning, "Failed to delete directory {0} in {1} attempts.", Path, MaxAttempts);
+				Log.WriteLine(TraceEventType.Warning, LogUtils.FormatException(LastException));
+			}
+
+			return Result;
+		}
+
+		/// <summary>
+		/// Deletes a directory and all its contents. Will delete read-only files.
+		/// </summary>
+		/// <param name="Path">Directory name.</param>
+		/// <returns>True if the directory no longer exists, false otherwise.</returns>
+		public static bool SafeDeleteDirectory(string Path, bool bQuiet = false)
+		{
+			if (!bQuiet)
+			{
+				Log.WriteLine(TraceEventType.Verbose, "SafeDeleteDirectory {0}", Path);
+			}
+			if (Directory.Exists(Path))
+			{
+				return RecursivelyDeleteDirectory(Path, bQuiet);
+			}
+			else
+			{
+				return true;
+			}
+		}
+
+		/// <summary>
+		/// Renames/moves a file.
+		/// </summary>
+		/// <param name="OldName">Old name</param>
+		/// <param name="NewName">New name</param>
+		/// <returns>True if the operation was successful, false otherwise.</returns>
+		public static bool SafeRenameFile(string OldName, string NewName, bool bQuiet = false)
+		{
+			if( !bQuiet )
+			{ 
+				Log.WriteLine(TraceEventType.Verbose, "SafeRenameFile {0} {1}", OldName, NewName);
+			}
+			const int MaxAttempts = 10;
+			int Attempts = 0;
+
+			bool Result = true;
+			do
+			{
+				Result = true;
+				try
+				{
+					if (File.Exists(OldName))
+					{
+						FileAttributes Attributes = File.GetAttributes(OldName);
+						if ((Attributes & FileAttributes.ReadOnly) != 0)
+						{
+							File.SetAttributes(OldName, Attributes & ~FileAttributes.ReadOnly);
+						}
+					}
+					File.Move(OldName, NewName);
+				}
+				catch (Exception Ex)
+				{
+					if (File.Exists(OldName) == true || File.Exists(NewName) == false)
+					{
+						Log.WriteLine(TraceEventType.Warning, "Failed to rename {0} to {1}", OldName, NewName);
+						Log.WriteLine(TraceEventType.Warning, LogUtils.FormatException(Ex));
+						Result = false;
+					}
+				}
+			}
+			while (Result == false && ++Attempts < MaxAttempts);
+
+			return Result;
+		}
+
+		// @todo: This could be passed in from elsewhere, and this should be somehow done per ini section
+		// but this will get it so that games won't ship passwords
+		private static string[] LinesToFilter = new string[]
+		{
+			"KeyStorePassword",
+			"KeyPassword",
+		};
+
+		private static void FilterIniFile(string SourceName, string TargetName)
+		{
+			string[] Lines = File.ReadAllLines(SourceName);
+			StringBuilder NewLines = new StringBuilder("");
+
+			foreach (string Line in Lines)
+			{
+				// look for each filter on each line
+				bool bFiltered = false;
+				foreach (string Filter in LinesToFilter)
+				{
+					if (Line.StartsWith(Filter + "="))
+					{
+						bFiltered = true;
+						break;
+					}
+				}
+
+				// write out if it's not filtered out
+				if (!bFiltered)
+				{
+					NewLines.AppendLine(Line);
+				}
+			}
+
+			// now write out the final .ini file
+			if (File.Exists(TargetName))
+			{
+				File.Delete(TargetName);
+			}
+			File.WriteAllText(TargetName, NewLines.ToString());
+
+			// other code assumes same timestamp for source and dest
+			File.SetLastWriteTimeUtc(TargetName, File.GetLastWriteTimeUtc(SourceName));
+		}
+
+		/// <summary>
+		/// Copies a file.
+		/// </summary>
+		/// <param name="SourceName">Source name</param>
+		/// <param name="TargetName">Target name</param>
+		/// <returns>True if the operation was successful, false otherwise.</returns>
+		public static bool SafeCopyFile(string SourceName, string TargetName, bool bQuiet = false, bool bFilterSpecialLinesFromIniFiles = false)
+		{
+			if (!bQuiet)
+			{
+				Log.WriteLine(TraceEventType.Verbose, "SafeCopyFile {0} {1}", SourceName, TargetName);
+			}
+			const int MaxAttempts = 10;
+			int Attempts = 0;
+
+			bool Result = true;
+			do
+			{
+				Result = true;
+				bool Retry = true;
+				try
+				{
+					bool bSkipSizeCheck = false;
+					if (bFilterSpecialLinesFromIniFiles && Path.GetExtension(SourceName) == ".ini")
+					{
+						FilterIniFile(SourceName, TargetName);
+						// ini files may change size, don't check
+						bSkipSizeCheck = true;
+					}
+					else
+					{
+						File.Copy(SourceName, TargetName, overwrite: true);
+					}
+					Retry = !File.Exists(TargetName);
+					if (!Retry)
+					{
+						FileInfo SourceInfo = new FileInfo(SourceName);
+						FileInfo TargetInfo = new FileInfo(TargetName);
+						if (!bSkipSizeCheck && SourceInfo.Length != TargetInfo.Length)
+						{
+							Log.WriteLine(TraceEventType.Warning, "Size mismatch {0} = {1} to {2} = {3}", SourceName, SourceInfo.Length, TargetName, TargetInfo.Length);
+							Retry = true;
+						}
+						// Timestamps should be no more than 2 seconds out - assuming this as exFAT filesystems store timestamps at 2 second intervals:
+						// http://ntfs.com/exfat-time-stamp.htm
+						if (!((SourceInfo.LastWriteTimeUtc - TargetInfo.LastWriteTimeUtc).TotalSeconds < 2 && (SourceInfo.LastWriteTimeUtc - TargetInfo.LastWriteTimeUtc).TotalSeconds > -2))
+						{
+							Log.WriteLine(TraceEventType.Warning, "Date mismatch {0} = {1} to {2} = {3}", SourceName, SourceInfo.LastWriteTimeUtc, TargetName, TargetInfo.LastWriteTimeUtc);
+							Retry = true;
+						}
+					}
+				}
+				catch (Exception Ex)
+				{
+					Log.WriteLine(System.Diagnostics.TraceEventType.Warning, "SafeCopyFile Exception was {0}", LogUtils.FormatException(Ex));
+					Retry = true;
+				}
+
+				if (Retry)
+				{
+					if (Attempts + 1 < MaxAttempts)
+					{
+						Log.WriteLine(TraceEventType.Warning, "Failed to copy {0} to {1}, deleting, waiting 10s and retrying.", SourceName, TargetName);
+						if (File.Exists(TargetName))
+						{
+							SafeDeleteFile(TargetName);
+						}
+						Thread.Sleep(10000);
+					}
+					else
+					{
+						Log.WriteLine(TraceEventType.Warning, "Failed to copy {0} to {1}", SourceName, TargetName);
+					}
+					Result = false;
+				}
+			}
+			while (Result == false && ++Attempts < MaxAttempts);
+
+			return Result;
+		}
+
+
+		/// <summary>
+		/// Reads all lines from a file.
+		/// </summary>
+		/// <param name="Filename">Filename</param>
+		/// <returns>An array containing all lines read from the file or null if the file could not be read.</returns>
+		public static string[] SafeReadAllLines(string Filename)
+		{
+			Log.WriteLine(TraceEventType.Verbose, "SafeReadAllLines {0}", Filename);
+			string[] Result = null;
+			try
+			{
+				Result = File.ReadAllLines(Filename);
+			}
+			catch (Exception Ex)
+			{
+				Log.WriteLine(TraceEventType.Warning, "Failed to load {0}", Filename);
+				Log.WriteLine(TraceEventType.Warning, LogUtils.FormatException(Ex));
+			}
+			return Result;
+		}
+
+		/// <summary>
+		/// Reads all text from a file.
+		/// </summary>
+		/// <param name="Filename">Filename</param>
+		/// <returns>String containing all text read from the file or null if the file could not be read.</returns>
+		public static string SafeReadAllText(string Filename)
+		{
+			Log.WriteLine(TraceEventType.Verbose, "SafeReadAllLines {0}", Filename);
+			string Result = null;
+			try
+			{
+				Result = File.ReadAllText(Filename);
+			}
+			catch (Exception Ex)
+			{
+				Log.WriteLine(TraceEventType.Warning, "Failed to load {0}", Filename);
+				Log.WriteLine(TraceEventType.Warning, LogUtils.FormatException(Ex));
+			}
+			return Result;
+		}
+
+		/// <summary>
+		/// Finds files in the specified path.
+		/// </summary>
+		/// <param name="Path">Path</param>
+		/// <param name="SearchPattern">Search pattern</param>
+		/// <param name="Recursive">Whether to search recursively or not.</param>
+		/// <returns>List of all files found (can be empty) or null if the operation failed.</returns>
+		public static string[] FindFiles(string Path, string SearchPattern, bool Recursive, bool bQuiet = false)
+		{
+			if (!bQuiet)
+			{
+				Log.WriteLine(TraceEventType.Verbose, "FindFiles {0} {1} {2}", Path, SearchPattern, Recursive);
+			}
+
+			// On Linux, filter out symlinks since we (usually) create them to fix mispelled case-sensitive filenames in content, and if they aren't filtered, 
+			// UAT picks up both the symlink and the original file and considers them duplicates when packaging (pak files are case-insensitive).
+			// Windows needs the symlinks though because that's how deduplication works on Windows server, 
+			// see https://answers.unrealengine.com/questions/212888/automated-buildjenkins-failing-due-to-symlink-chec.html
+			// FIXME: ZFS, JFS and other fs that can be case-insensitive on Linux should use the faster path as well.
+			if (UnrealBuildTool.BuildHostPlatform.Current.Platform != UnrealTargetPlatform.Linux)
+			{
+				return Directory.GetFiles(Path, SearchPattern, Recursive ? SearchOption.AllDirectories : SearchOption.TopDirectoryOnly);
+			}
+			else
+			{
+				List<string> FileNames = new List<string>();
+				DirectoryInfo DirInfo = new DirectoryInfo(Path);
+				foreach( FileInfo File in DirInfo.EnumerateFiles(SearchPattern, Recursive ? SearchOption.AllDirectories : SearchOption.TopDirectoryOnly))
+				{
+					if (File.Attributes.HasFlag(FileAttributes.ReparsePoint))
+					{
+						if (!bQuiet)
+						{
+							Log.WriteLine(TraceEventType.Warning, "Ignoring symlink {0}", File.FullName);
+						}
+						continue;
+					}
+					
+					FileNames.Add(File.FullName);
+				}
+				
+				return FileNames.ToArray();
+			}
+		}
+
+		/// <summary>
+		/// Finds directories in the specified path.
+		/// </summary>
+		/// <param name="Path">Path</param>
+		/// <param name="SearchPattern">Search pattern</param>
+		/// <param name="Recursive">Whether to search recursively or not.</param>
+		/// <returns>List of all directories found (can be empty) or null if the operation failed.</returns>
+		public static string[] FindDirectories(string Path, string SearchPattern, bool Recursive, bool bQuiet = false)
+		{
+			if (!bQuiet)
+			{
+				Log.WriteLine(TraceEventType.Verbose, "FindDirectories {0} {1} {2}", Path, SearchPattern, Recursive);
+			}
+			return Directory.GetDirectories(Path, SearchPattern, Recursive ? SearchOption.AllDirectories : SearchOption.TopDirectoryOnly);
+		}
+
+		/// <summary>
+		/// Finds files in the specified path.
+		/// </summary>
+		/// <param name="Path">Path</param>
+		/// <param name="SearchPattern">Search pattern</param>
+		/// <param name="Recursive">Whether to search recursively or not.</param>
+		/// <returns>List of all files found (can be empty) or null if the operation failed.</returns>
+		public static string[] SafeFindFiles(string Path, string SearchPattern, bool Recursive, bool bQuiet = false)
+		{
+			if (!bQuiet)
+			{
+				Log.WriteLine(TraceEventType.Verbose, "SafeFindFiles {0} {1} {2}", Path, SearchPattern, Recursive);
+			}
+			string[] Files = null;
+			try
+			{
+				Files = FindFiles(Path, SearchPattern, Recursive, bQuiet);
+			}
+			catch (Exception Ex)
+			{
+				Log.WriteLine(TraceEventType.Warning, "Unable to Find Files in {0}", Path);
+				Log.WriteLine(TraceEventType.Warning, LogUtils.FormatException(Ex));
+			}
+			return Files;
+		}
+
+		/// <summary>
+		/// Finds directories in the specified path.
+		/// </summary>
+		/// <param name="Path">Path</param>
+		/// <param name="SearchPattern">Search pattern</param>
+		/// <param name="Recursive">Whether to search recursively or not.</param>
+		/// <returns>List of all files found (can be empty) or null if the operation failed.</returns>
+		public static string[] SafeFindDirectories(string Path, string SearchPattern, bool Recursive, bool bQuiet = false)
+		{
+			if (!bQuiet)
+			{
+				Log.WriteLine(TraceEventType.Verbose, "SafeFindDirectories {0} {1} {2}", Path, SearchPattern, Recursive);
+			}
+			string[] Directories = null;
+			try
+			{
+				Directories = FindDirectories(Path, SearchPattern, Recursive, bQuiet);
+			}
+			catch (Exception Ex)
+			{
+				Log.WriteLine(TraceEventType.Warning, "Unable to Find Directories in {0}", Path);
+				Log.WriteLine(TraceEventType.Warning, LogUtils.FormatException(Ex));
+			}
+			return Directories;
+		}
+
+		/// <summary>
+		/// Checks if a file exists.
+		/// </summary>
+		/// <param name="Path">Filename</param>
+		/// <param name="bQuiet">if true, do not print a message</param>
+		/// <returns>True if the file exists, false otherwise.</returns>
+		public static bool SafeFileExists(string Path, bool bQuiet = false)
+		{
+			bool Result = false;
+			try
+			{
+				Result = File.Exists(Path);
+				if (!bQuiet)
+				{
+					Log.WriteLine(TraceEventType.Verbose, "SafeFileExists {0}={1}", Path, Result);
+				}
+			}
+			catch (Exception Ex)
+			{
+				Log.WriteLine(TraceEventType.Warning, "Unable to check if file {0} exists.", Path);
+				Log.WriteLine(TraceEventType.Warning, LogUtils.FormatException(Ex));
+			}
+			return Result;
+		}
+
+
+		/// <summary>
+		/// Checks if a directory exists.
+		/// </summary>
+		/// <param name="Path">Directory</param>
+		/// <param name="bQuiet">if true, no longging</param>
+		/// <returns>True if the directory exists, false otherwise.</returns>
+		public static bool SafeDirectoryExists(string Path, bool bQuiet = false)
+		{
+			bool Result = false;
+			try
+			{
+				Result = Directory.Exists(Path);
+				if (!bQuiet)
+				{
+					Log.WriteLine(TraceEventType.Verbose, "SafeDirectoryExists {0}={1}", Path, Result);
+				}
+			}
+			catch (Exception Ex)
+			{
+				Log.WriteLine(TraceEventType.Warning, "Unable to check if directory {0} exists.", Path);
+				Log.WriteLine(TraceEventType.Warning, LogUtils.FormatException(Ex));
+			}
+			return Result;
+		}
+
+		/// <summary>
+		/// Writes lines to a file.
+		/// </summary>
+		/// <param name="Path">Filename</param>
+		/// <param name="Text">Text</param>
+		/// <returns>True if the operation was successful, false otherwise.</returns>
+		public static bool SafeWriteAllLines(string Path, string[] Text)
+		{
+			Log.WriteLine(TraceEventType.Verbose, "SafeWriteAllLines {0}", Path);
+			bool Result = false;
+			try
+			{
+				File.WriteAllLines(Path, Text);
+				Result = true;
+			}
+			catch (Exception Ex)
+			{
+				Log.WriteLine(TraceEventType.Warning, "Unable to write text to {0}", Path);
+				Log.WriteLine(TraceEventType.Warning, LogUtils.FormatException(Ex));
+			}
+			return Result;
+		}
+
+		/// <summary>
+		/// Writes text to a file.
+		/// </summary>
+		/// <param name="Path">Filename</param>
+		/// <param name="Text">Text</param>
+		/// <returns>True if the operation was successful, false otherwise.</returns>
+		public static bool SafeWriteAllText(string Path, string Text)
+		{
+			Log.WriteLine(TraceEventType.Verbose, "SafeWriteAllText {0}", Path);
+			bool Result = false;
+			try
+			{
+				File.WriteAllText(Path, Text);
+				Result = true;
+			}
+			catch (Exception Ex)
+			{
+				Log.WriteLine(TraceEventType.Warning, "Unable to write text to {0}", Path);
+				Log.WriteLine(TraceEventType.Warning, LogUtils.FormatException(Ex));
+			}
+			return Result;
+		}
+
+		/// <summary>
+		/// Writes text to a file.
+		/// </summary>
+		/// <param name="Path">Filename</param>
+		/// <param name="Text">Text</param>
+		/// <returns>True if the operation was successful, false otherwise.</returns>
+		public static bool SafeWriteAllBytes(string Path, byte[] Bytes)
+		{
+			Log.WriteLine(TraceEventType.Verbose, "SafeWriteAllBytes {0}", Path);
+			bool Result = false;
+			try
+			{
+				File.WriteAllBytes(Path, Bytes);
+				Result = true;
+			}
+			catch (Exception Ex)
+			{
+				Log.WriteLine(TraceEventType.Warning, "Unable to write text to {0}", Path);
+				Log.WriteLine(TraceEventType.Warning, LogUtils.FormatException(Ex));
+			}
+			return Result;
+		}
+
+		/// <summary>
+		/// Runs the specified delegate checking if this is the only instance of the application.
+		/// </summary>
+		/// <param name="Main"></param>
+		/// <param name="Param"></param>
+		public static void RunSingleInstance(System.Action<object> Main, object Param)
+		{
+			if (Environment.GetEnvironmentVariable("uebp_UATMutexNoWait") == "1")
+			{
+				Main(Param);
+			}
+			var bCreatedMutex = false;
+			var LocationHash = InternalUtils.ExecutingAssemblyLocation.GetHashCode();
+			var MutexName = "Global/" + Path.GetFileNameWithoutExtension(ExecutingAssemblyLocation) + "_" + LocationHash.ToString() + "_Mutex";
+			using (Mutex SingleInstanceMutex = new Mutex(true, MutexName, out bCreatedMutex))
+			{
+				if (!bCreatedMutex)
+				{
+                    throw new AutomationException("Another instance of {0} is already running.", ExecutingAssemblyLocation);
+				}
+				else
+				{
+					Log.WriteLine(TraceEventType.Verbose, "No other instance of {0} is running.", ExecutingAssemblyLocation);
+				}
+
+				Main(Param);
+
+				SingleInstanceMutex.ReleaseMutex();
+			}
+		}
+
+		/// <summary>
+		/// Path to the executable which runs this code.
+		/// </summary>
+		public static string ExecutingAssemblyDirectory
+		{
+			get
+			{
+				return CommandUtils.CombinePaths(Path.GetDirectoryName(ExecutingAssemblyLocation));
+			}
+		}
+
+		/// <summary>
+		/// Filename of the executable which runs this code.
+		/// </summary>
+		public static string ExecutingAssemblyLocation
+		{
+			get
+			{
+				return CommandUtils.CombinePaths(new Uri(System.Reflection.Assembly.GetExecutingAssembly().CodeBase).LocalPath);
+			}
+		}
+
+		/// <summary>
+		/// Version info of the executable which runs this code.
+		/// </summary>
+		public static FileVersionInfo ExecutableVersion
+		{
+			get
+			{
+				return FileVersionInfo.GetVersionInfo(ExecutingAssemblyLocation);
+			}
+		}
+	}
+
+	#endregion
+
+
+    #region VersionFileReader
+
+
+    /// <summary>
+    /// This is to ensure that UAT can produce version strings precisely compatible
+    /// with FEngineVersion.
+    /// 
+    /// WARNING: If FEngineVersion compatibility changes, this code needs to be updated.
+    /// </summary>
+    public class FEngineVersionSupport
+    {
+        /// <summary>
+        /// The version info read from the Version header. The populated fields will be Major, Minor, and Build from the MAJOR, MINOR, and PATCH lines, respectively.
+        /// Expects lines like:
+        /// #define APP_MAJOR_VERSION 0
+        /// #define APP_MINOR_VERSION 0
+        /// #define APP_PATCH_VERSION 0
+        /// </summary>
+        public readonly Version Version;
+
+        /// <summary>
+        /// The changelist this version is associated with
+        /// </summary>
+        public readonly int Changelist;
+
+        /// <summary>
+        /// The Branch name associated with the version
+        /// </summary>
+        public readonly string BranchName;
+
+        /// <summary>
+        /// Reads a Version.h file, looking for macros that define the MAJOR/MINOR/PATCH version fields. Expected to match the Version.h in the engine.
+        /// </summary>
+        /// <param name="Filename">Version.h file to read.</param>
+        /// <returns>Version that puts the Major/Minor/Patch fields in the Major/Minor/Build fields, respectively.</returns>
+        public static Version ReadVersionFromFile(string Filename)
+        {
+            var regex = new Regex(@"#define.+_(?<Type>MAJOR|MINOR|PATCH)_VERSION\s+(?<Value>.+)", RegexOptions.ExplicitCapture | RegexOptions.IgnoreCase);
+            var foundElements = new Dictionary<string, int>(3);
+            foreach (var line in File.ReadLines(Filename))
+            {
+                try
+                {
+                    var match = regex.Match(line);
+                    if (match.Success)
+                    {
+                        foundElements.Add(match.Groups["Type"].Value, int.Parse(match.Groups["Value"].Value));
+                    }
+                }
+                catch (Exception ex)
+                {
+                    throw new AutomationException(string.Format("Failed to parse line {0} in version file {1}", line, Filename), ex);
+                }
+            }
+
+            // must find all three parts to accept the version file.
+            if (foundElements.Keys.Intersect(new[] { "MAJOR", "MINOR", "PATCH" }).Count() != 3)
+            {
+                throw new AutomationException("Failed to find MAJOR, MINOR, and PATCH fields from version file {0}", Filename);
+            }
+            CommandUtils.Log("Read {0}.{1}.{2} from {3}", foundElements["MAJOR"], foundElements["MINOR"], foundElements["PATCH"], Filename);
+            return new Version(foundElements["MAJOR"], foundElements["MINOR"], foundElements["PATCH"]);
+        }
+
+        /// <summary>
+        /// Ctor that takes a pre-determined Version. Gets the Changelist and BranchName from the current <see cref="CommandUtils.P4Env"/>.
+        /// </summary>
+		/// <param name="InVersion">Predetermined version.</param>
+		/// <param name="InChangelist">Predetermined changelist (optional)</param>
+		/// <param name="InBranchName">Predetermined branch name (optional)</param>
+		public FEngineVersionSupport(Version InVersion, int InChangelist = -1, string InBranchName = null)
+        {
+			Version = InVersion;
+			if (InChangelist <= 0)
+			{
+				Changelist = CommandUtils.P4Enabled ? CommandUtils.P4Env.Changelist : 0;
+			}
+			else
+			{
+				Changelist = InChangelist;
+			}
+			if (String.IsNullOrEmpty(InBranchName))
+			{
+				BranchName = CommandUtils.P4Enabled ? CommandUtils.P4Env.BuildRootEscaped : "UnknownBranch";
+			}
+			else
+			{
+				BranchName = InBranchName;
+			}
+        }
+
+        /// <summary>
+        /// Gets a version string compatible with FEngineVersion's native parsing code.
+        /// </summary>
+        /// <remarks>
+        /// The format looks like: Major.Minor.Build-Changelist+BranchName.
+        /// </remarks>
+        /// <returns></returns>
+        public override string ToString()
+        {
+			return String.Format("{0}.{1}.{2}-{3}+{4}",
+				Version.Major,
+				Version.Minor,
+				Version.Build,
+                Changelist.ToString("0000000"),
+                BranchName);
+        }
+
+		public override bool Equals(object obj)
+		{
+			var that = obj as FEngineVersionSupport;
+			if (that == null)
+				return false;
+			return this.ToString() == that.ToString();
+		}
+
+		public override int GetHashCode()
+		{
+			return this.ToString().GetHashCode();
+		}
+
+        /// <summary>
+        /// Ctor initializes with the values from the supplied Version file. The BranchName and CL are also taken from the current <see cref="CommandUtils.P4Env"/>.
+        /// </summary>
+        /// <param name="Filename">Full path to the file with the version info.</param>
+		/// <param name="InChangelist">Predetermined changelist (optional)</param>
+		/// <param name="InBranchName">Predetermined branch name (optional)</param>
+        public static FEngineVersionSupport FromVersionFile(string Filename, int InChangelist = -1, string InBranchName = null)
+        {
+			return new FEngineVersionSupport(ReadVersionFromFile(Filename), InChangelist, InBranchName);
+        }
+
+        /// <summary>
+        /// Creates a <see cref="FEngineVersionSupport"/> from a string that matches the format given in <see cref="ToString"/>.
+        /// </summary>
+        /// <param name="versionString">Version string that should match the FEngineVersion::ToString() format.</param>
+		/// <param name="bAllowVersion">Optional parameter which if set to true, allows version strings with no version number specified.</param>
+        /// <returns>a new instance with fields initialized to the match those given in the string.</returns>
+        public static FEngineVersionSupport FromString(string versionString, bool bAllowNoVersion = false)
+        {
+            try
+            {
+				if (bAllowNoVersion && versionString.StartsWith("++depot"))
+				{
+					// This form of version is used when a product has no major.minor.patch version
+					// E.g. ++depot+UE4-ProdName-CL-12345678
+					var clSplit = versionString.Split(new string[] { "-CL-" }, 2, StringSplitOptions.None);
+					var dashSplit = clSplit[1].Split(new[] { '-' }, 2);
+					var changelist = int.Parse(dashSplit[0]);
+					var branchName = clSplit[0];
+					return new FEngineVersionSupport(new Version(0, 0, 0), changelist, branchName);
+				}
+				else
+				{
+					// This is the standard Rocket versioning scheme, e.g. "4.5.0-12345678+++depot+UE4"
+					var dotSplit = versionString.Split(new[] { '.' }, 3);
+					var dashSplit = dotSplit[2].Split(new[] { '-' }, 2);
+					var plusSplit = dashSplit[1].Split(new[] { '+' }, 2);
+					var major = int.Parse(dotSplit[0]);
+					var minor = int.Parse(dotSplit[1]);
+					var patch = int.Parse(dashSplit[0]);
+					var changelist = int.Parse(plusSplit[0]);
+					var branchName = plusSplit[1];
+					return new FEngineVersionSupport(new Version(major, minor, patch), changelist, branchName);
+				}
+            }
+            catch (Exception ex)
+            {
+                throw new AutomationException(string.Format("Failed to parse {0} as an FEngineVersion compatible string", versionString), ex);
+            }
+        }
+
+        public static DateTime BuildTime()
+        {
+            string VerFile = CommandUtils.CombinePaths(CommandUtils.CmdEnv.LocalRoot, "Engine", "Build", "build.properties");
+
+            var VerLines = CommandUtils.ReadAllText(VerFile);
+
+            var SearchFor = "TimestampForBVT=";
+
+            int Index = VerLines.IndexOf(SearchFor);
+
+            if (Index < 0)
+            {
+                throw new AutomationException("Could not find {0} in {1} for file {2}", SearchFor, VerLines, VerFile);
+            }
+            Index = Index + SearchFor.Length;
+
+
+            var Parts = VerLines.Substring(Index).Split('.', '_', '-', '\n', '\r', '\t');
+
+            DateTime Result = new DateTime(int.Parse(Parts[0]), int.Parse(Parts[1]), int.Parse(Parts[2]), int.Parse(Parts[3]), int.Parse(Parts[4]), int.Parse(Parts[5]));
+
+            CommandUtils.Log("Current Build Time is {0}", Result);
+            return Result;
+        }
+    }
+
+    #endregion
+
+    #region VersionFileUpdater
+
+    /// <summary>
+	/// VersionFileUpdater.
+	/// </summary>
+	public class VersionFileUpdater
+	{
+		/// <summary>
+		/// Constructor
+		/// </summary>
+		public VersionFileUpdater(string Filename)
+		{
+			MyFile = new FileInfo(Filename);
+			Lines = new List<string>(InternalUtils.SafeReadAllLines(Filename));
+            if (CommandUtils.IsNullOrEmpty(Lines))
+            {
+                throw new AutomationException("Version file {0} was empty or not found!", Filename);
+            }
+		}
+
+		/// <summary>
+		/// Doc
+		/// </summary>
+		public void ReplaceLine(string StartOfLine, string ReplacementRHS)
+		{
+			for (int Index = 0; Index < Lines.Count; ++Index)
+			{
+				if (Lines[Index].StartsWith(StartOfLine))
+				{
+					Lines[Index] = StartOfLine + ReplacementRHS;
+					return;
+				}
+			}
+            throw new AutomationException("Unable to find line {0} in {1}", StartOfLine, MyFile.FullName);
+		}
+
+		/// <summary>
+		/// Doc
+		/// </summary>
+		public void ReplaceOrAddLine(string StartOfLine, string ReplacementRHS)
+		{
+			if (Contains(StartOfLine))
+			{
+				ReplaceLine(StartOfLine, ReplacementRHS);
+			}
+			else
+			{
+				AddLine("");
+				AddLine(StartOfLine + ReplacementRHS);
+			}
+		}
+
+		/// <summary>
+		/// Adds a new line to the version file
+		/// </summary>
+		/// <param name="Line"></param>
+		public void AddLine(string Line)
+		{
+			Lines.Add(Line);
+		}
+
+		/// <summary>
+		/// Doc
+		/// </summary>
+        public void SetAssemblyInformationalVersion(string NewInformationalVersion)
+		{
+            // This searches for the AssemblyInformationalVersion string. Most the mess is to allow whitespace in places that are possible.
+            // Captures the string into a group called "Ver" for replacement.
+            var regex = new Regex(@"\[assembly:\s+AssemblyInformationalVersion\s*\(\s*""(?<Ver>.*)""\s*\)\s*]", RegexOptions.IgnoreCase | RegexOptions.ExplicitCapture);
+            foreach (var Index in Enumerable.Range(0, Lines.Count))
+            {
+                var line = Lines[Index];
+                var match = regex.Match(line);
+                if (match.Success)
+                {
+                    var verGroup = match.Groups["Ver"];
+                    var sb = new StringBuilder(line);
+                    sb.Remove(verGroup.Index, verGroup.Length);
+                    sb.Insert(verGroup.Index, NewInformationalVersion);
+                    Lines[Index] = sb.ToString();
+                    return;
+                }
+            }
+            throw new AutomationException("Failed to find the AssemblyInformationalVersion attribute in {1}", MyFile.FullName);
+		}
+
+		/// <summary>
+		/// Doc
+		/// </summary>
+		public void Commit()
+		{
+			MyFile.IsReadOnly = false;
+			if (!InternalUtils.SafeWriteAllLines(MyFile.FullName, Lines.ToArray()))
+			{
+				throw new AutomationException("Unable to update version info in {0}", MyFile.FullName);
+			}
+		}
+
+		/// <summary>
+		/// Checks if the version file contains the specified string.
+		/// </summary>
+		/// <param name="Text">String to look for.</param>
+		/// <returns></returns>
+		public bool Contains(string Text, bool CaseSensitive = true)
+		{
+			foreach (var Line in Lines)
+			{
+				if (Line.IndexOf(Text, CaseSensitive ? StringComparison.InvariantCulture : StringComparison.InvariantCultureIgnoreCase) >= 0)
+				{
+					return true;
+				}
+			}
+			return false;
+		}
+
+		/// <summary>
+		/// Doc
+		/// </summary>
+		protected FileInfo MyFile;
+		/// <summary>
+		/// Doc
+		/// </summary>
+		protected List<string> Lines;
+	}
+
+	#endregion
+}
diff --git a/Engine/Source/Programs/AutomationTool/AutomationUtils/WatchdogTimer.cs b/Engine/Source/Programs/AutomationTool/AutomationUtils/WatchdogTimer.cs
new file mode 100644
--- /dev/null
+++ b/Engine/Source/Programs/AutomationTool/AutomationUtils/WatchdogTimer.cs
@@ -0,0 +1,87 @@
+﻿// Copyright 1998-2015 Epic Games, Inc. All Rights Reserved.
+
+using System;
+using System.Collections.Generic;
+using System.Linq;
+using System.Text;
+using System.Timers;
+using System.Runtime.CompilerServices;
+using UnrealBuildTool;
+
+namespace AutomationTool
+{
+	/// <summary>
+	/// Timer that exits the application of the execution of the command takes too long.
+	/// </summary>
+	public class WatchdogTimer : IDisposable, IProcess
+	{
+		private Timer CountownTimer;
+		private System.Diagnostics.StackFrame TimerFrame;
+		private object SyncObject = new object();
+		private bool bTicking = false;
+
+		[MethodImplAttribute(MethodImplOptions.NoInlining)]
+		public WatchdogTimer(int Seconds)
+		{
+			TimerFrame = new System.Diagnostics.StackFrame(1);  
+			CountownTimer = new Timer(Seconds * 1000.0);
+			CountownTimer.Elapsed += Elapsed;
+			CountownTimer.Start();
+			bTicking = true;
+			ProcessManager.AddProcess(this);
+		}
+
+		void Elapsed(object sender, ElapsedEventArgs e)
+		{
+			lock (SyncObject)
+			{
+				bTicking = false;
+			}
+			int Interval = (int)(CountownTimer.Interval * 0.001);
+			Dispose();
+			var Method = TimerFrame.GetMethod();
+			Log.TraceError("BUILD FAILED: WatchdogTimer timed out after {0}s in {1}.{2}", Interval, Method.DeclaringType.Name, Method.Name);			
+			Environment.Exit(1);
+		}
+
+		public void Dispose()
+		{
+			Log.TraceInformation("WatchdogTimer.Dispose()");
+			lock (SyncObject)
+			{
+				if (CountownTimer != null)
+				{
+					CountownTimer.Stop();
+					CountownTimer.Dispose();
+					CountownTimer = null;
+				}
+				bTicking = false;
+			}
+			ProcessManager.RemoveProcess(this);
+		}
+
+		public void StopProcess(bool KillDescendants = true)
+		{
+			Dispose();			
+		}
+
+		public bool HasExited
+		{
+			get 
+			{
+				bool bResult = true;
+ 				lock (SyncObject)
+				{
+					bResult = bTicking;
+				}
+				return bResult;
+			}
+		}
+
+		public string GetProcessName()
+		{
+			var Method = TimerFrame.GetMethod();
+			return String.Format("WatchdogTimer_{0}.{1}", Method.DeclaringType.Name, Method.Name);
+		}
+	}
+}
