UEB-260 - Break AutomationTool into AutomationUtils that all automation projects depend on, and AutomationTool, which essentially only contains the startup code.

diff --git a/Engine/Source/Programs/AutomationTool/AutomationUtils/WindowsHostPlatform.cs b/Engine/Source/Programs/AutomationTool/AutomationUtils/WindowsHostPlatform.cs
new file mode 100644
--- /dev/null
+++ b/Engine/Source/Programs/AutomationTool/AutomationUtils/WindowsHostPlatform.cs
@@ -0,0 +1,237 @@
+﻿// Copyright 1998-2015 Epic Games, Inc. All Rights Reserved.
+
+using System;
+using System.Collections.Generic;
+using System.Linq;
+using System.Text;
+using System.IO;
+using System.Diagnostics;
+using UnrealBuildTool;
+
+namespace AutomationTool
+{
+	class WindowsHostPlatform : HostPlatform
+	{
+		/// <returns>The path to Windows SDK directory for the specified version.</returns>
+		private static string FindWindowsSDKInstallationFolder( string Version )
+		{
+			// Based on VCVarsQueryRegistry
+			string WinSDKPath = (string)Microsoft.Win32.Registry.GetValue( @"HKEY_CURRENT_USER\SOFTWARE\Microsoft\Microsoft SDKs\Windows\" + Version, "InstallationFolder", null );
+			if( WinSDKPath != null )
+			{
+				return WinSDKPath;
+			}
+			WinSDKPath = (string)Microsoft.Win32.Registry.GetValue( @"HKEY_LOCAL_MACHINE\SOFTWARE\Wow6432Node\Microsoft\Microsoft SDKs\Windows\" + Version, "InstallationFolder", null );
+			if( WinSDKPath != null )
+			{
+				return WinSDKPath;
+			}
+			WinSDKPath = (string)Microsoft.Win32.Registry.GetValue( @"HKEY_CURRENT_USER\SOFTWARE\Wow6432Node\Microsoft\Microsoft SDKs\Windows\" + Version, "InstallationFolder", null );
+			if( WinSDKPath != null )
+			{
+				return WinSDKPath;
+			}
+
+			throw new BuildException( "Can't find Windows SDK {0}.", Version );
+		}
+
+		/// <summary>
+		/// Figures out the path to Visual Studio's /Common7/Tools directory.  Note that if Visual Studio is not
+		/// installed, the Windows SDK path will be used, which also happens to be the same directory. (It installs
+		/// the toolchain into the folder where Visual Studio would have installed it to.
+		/// </summary>
+		/// <returns>The path to Visual Studio's /Common7/Tools directory</returns>
+		private bool FindBaseVSToolPaths(out string WindowsSDKDir, out string BaseVSToolPath)
+		{
+			WindowsSDKDir = "";
+			BaseVSToolPath = WindowsPlatform.GetVSComnToolsPath();
+
+			if (WindowsPlatform.Compiler == WindowsCompiler.VisualStudio2013)
+			{
+				WindowsSDKDir = FindWindowsSDKInstallationFolder( "v8.1" );
+
+				if (string.IsNullOrEmpty(BaseVSToolPath))
+				{
+					Log.TraceError("Visual Studio 2013 must be installed in order to build this target.");
+					return false;
+				}
+			}
+			else if (WindowsPlatform.Compiler == WindowsCompiler.VisualStudio2012)
+			{
+				WindowsSDKDir = FindWindowsSDKInstallationFolder( "v8.0" );
+
+				if (string.IsNullOrEmpty(BaseVSToolPath))
+				{
+					Log.TraceError("Visual Studio 2012 must be installed in order to build this target.");
+					return false;
+				}
+			}
+			else
+			{
+				throw new BuildException("Unexpected WindowsPlatform.Compiler version");
+			}
+			return true;
+		}
+
+		public override void SetFrameworkVars()
+		{
+			var FrameworkDir = CommandUtils.GetEnvVar(EnvVarNames.NETFrameworkDir);
+			var FrameworkVersion = CommandUtils.GetEnvVar(EnvVarNames.NETFrameworkVersion);
+
+			if (String.IsNullOrEmpty(FrameworkDir) || String.IsNullOrEmpty(FrameworkVersion))
+			{
+				Log.TraceInformation("Setting .Net Framework environment variables.");
+
+				// Determine 32 vs 64 bit. Worth noting that WOW64 will report x86.
+				bool Supports64bitExecutables = Environment.GetEnvironmentVariable("PROCESSOR_ARCHITECTURE") == "AMD64";
+				Log.TraceInformation("Supports64bitExecutables=" + Supports64bitExecutables.ToString());
+				Log.TraceInformation("WindowsPlatform.Compiler" + WindowsPlatform.Compiler.ToString());
+
+				string WindowsSDKDir, BaseVSToolPath;
+				if (FindBaseVSToolPaths(out WindowsSDKDir, out BaseVSToolPath) == false)
+				{
+					throw new AutomationException("Unable to find VS paths.");
+				}
+
+				Log.TraceInformation("WindowsSDKDir=" + WindowsSDKDir);
+				Log.TraceInformation("BaseVSToolPath=" + BaseVSToolPath);
+
+				string VCVarsBatchFile = "";
+
+				// 64 bit tool chain.
+				if (Supports64bitExecutables)
+				{
+					VCVarsBatchFile = CommandUtils.CombinePaths(BaseVSToolPath, "../../VC/bin/x86_amd64/vcvarsx86_amd64.bat");
+				}
+				// The 32 bit vars batch file in the binary folder simply points to the one in the common tools folder.
+				else
+				{
+					VCVarsBatchFile = CommandUtils.CombinePaths(BaseVSToolPath, "vsvars32.bat");
+				}
+
+				if (File.Exists(VCVarsBatchFile))
+				{
+					Log.TraceInformation("Setting VS environment variables via {0}.", VCVarsBatchFile);
+					var BatchVars = CommandUtils.GetEnvironmentVariablesFromBatchFile(VCVarsBatchFile);
+					// Note these env var names are hardcoded here because they're more
+					// related to the values stored inside the batch files rather than the mapping equivalents.
+					BatchVars.TryGetValue("FrameworkDir", out FrameworkDir);
+					BatchVars.TryGetValue("FrameworkVersion", out FrameworkVersion);
+					CommandUtils.SetEnvVar(EnvVarNames.NETFrameworkDir, FrameworkDir);
+					CommandUtils.SetEnvVar(EnvVarNames.NETFrameworkVersion, FrameworkVersion);
+				}
+				else
+				{
+					throw new AutomationException("{0} does not exist.", VCVarsBatchFile);
+				}
+				Log.TraceInformation("{0}={1}", EnvVarNames.NETFrameworkDir, FrameworkDir);
+				Log.TraceInformation("{0}={1}", EnvVarNames.NETFrameworkVersion, FrameworkVersion);
+			}
+		}
+
+		public override string GetMsBuildExe()
+		{
+			var DotNETToolsFolder = CommandUtils.CombinePaths(CommandUtils.GetEnvVar(EnvVarNames.NETFrameworkDir), CommandUtils.GetEnvVar(EnvVarNames.NETFrameworkVersion));
+			var MsBuildExe = CommandUtils.CombinePaths(DotNETToolsFolder, "MSBuild.exe");
+			if (!CommandUtils.FileExists_NoExceptions(MsBuildExe))
+			{				
+				throw new NotSupportedException("MsBuild.exe does not exist.");
+			}
+			return MsBuildExe;
+		}
+
+		public override string GetMsDevExe()
+		{
+			// Locate MsDevEnv executable
+			string PotentialMsDevExe = Path.Combine(WindowsPlatform.GetVSComnToolsPath(), "..", "IDE", "Devenv.com");
+
+			var Info = new System.IO.FileInfo(PotentialMsDevExe);
+			if (Info.Exists)
+			{
+				return PotentialMsDevExe;
+			}
+			else
+			{				
+				throw new NotSupportedException(String.Format("{0} does not exist.", PotentialMsDevExe));
+			}
+		}
+
+		public override string RelativeBinariesFolder
+		{
+			get { return @"Engine/Binaries/Win64/"; }
+		}
+
+		public override string GetUE4ExePath(string UE4Exe)
+		{
+			return CommandUtils.CombinePaths(CommandUtils.CmdEnv.LocalRoot, RelativeBinariesFolder, UE4Exe);
+		}
+
+		public override string LocalBuildsLogFolder
+		{
+			get { return CommandUtils.CombinePaths(PathSeparator.Slash, Environment.GetFolderPath(Environment.SpecialFolder.ApplicationData), "Unreal Engine", "AutomationTool", "Logs"); }
+		}
+
+		public override string P4Exe
+		{
+			get { return "p4.exe"; }
+		}
+
+		public override Process CreateProcess(string AppName)
+		{
+			var NewProcess = new Process();
+			return NewProcess;
+		}
+
+		public override void SetupOptionsForRun(ref string AppName, ref CommandUtils.ERunOptions Options, ref string CommandLine)
+		{
+		}
+
+		public override void SetConsoleCtrlHandler(ProcessManager.CtrlHandlerDelegate Handler)
+		{
+			ProcessManager.SetConsoleCtrlHandler(Handler, true);
+		}
+
+		public override bool IsScriptModuleSupported(string ModuleName)
+		{
+			return true;
+		}
+
+		public override string UBTProjectName
+		{
+			get { return "UnrealBuildTool"; }
+		}
+
+		public override UnrealTargetPlatform HostEditorPlatform
+		{
+			get { return UnrealTargetPlatform.Win64; }
+		}
+
+		public override string PdbExtension
+		{
+			get { return ".pdb"; }
+		}
+
+		static string[] SystemServices = new string[]
+		{
+			"winlogon",
+			"system idle process",
+			"taskmgr",
+			"spoolsv",
+			"csrss",
+			"smss",
+			"svchost",
+			"services",
+			"lsass",
+            "conhost",
+            "oobe",
+            "mmc"
+		};
+		public override string[] DontKillProcessList
+		{
+			get 
+			{
+				return SystemServices;
+			}
+		}
+	}
+}
diff --git a/Engine/Source/Programs/AutomationTool/GUBP/GUBP.Automation.csproj b/Engine/Source/Programs/AutomationTool/GUBP/GUBP.Automation.csproj
--- a/Engine/Source/Programs/AutomationTool/GUBP/GUBP.Automation.csproj
+++ b/Engine/Source/Programs/AutomationTool/GUBP/GUBP.Automation.csproj
@@ -72,9 +72,10 @@
       <Name>UnrealBuildTool</Name>
       <Private>False</Private>
     </ProjectReference>
-    <ProjectReference Include="..\AutomationTool.csproj">
-      <Project>{a35b0261-d868-475e-a098-eb14ce6dc89f}</Project>
-      <Name>AutomationTool</Name>
+    <ProjectReference Include="..\AutomationUtils\AutomationUtils.Automation.csproj">
+      <Project>{2c96a7f2-b1a3-4258-8e0a-e588ff41a53e}</Project>
+      <Name>AutomationUtils.Automation</Name>
+      <Private>False</Private>
     </ProjectReference>
   </ItemGroup>
   <ItemGroup>
diff --git a/Engine/Source/Programs/AutomationTool/HTML5/HTML5.Automation.csproj b/Engine/Source/Programs/AutomationTool/HTML5/HTML5.Automation.csproj
--- a/Engine/Source/Programs/AutomationTool/HTML5/HTML5.Automation.csproj
+++ b/Engine/Source/Programs/AutomationTool/HTML5/HTML5.Automation.csproj
@@ -55,16 +55,16 @@
     <Compile Include="Properties\AssemblyInfo.cs" />
   </ItemGroup>
   <ItemGroup>
-    <ProjectReference Include="..\AutomationTool.csproj">
-      <Project>{a35b0261-d868-475e-a098-eb14ce6dc89f}</Project>
-      <Name>AutomationTool</Name>
-      <Private>False</Private>
-    </ProjectReference>
     <ProjectReference Include="..\..\UnrealBuildTool\UnrealBuildTool.csproj">
       <Project>{fd7c5e1a-cfe4-4fd5-a525-1eb1599a39ac}</Project>
       <Name>UnrealBuildTool</Name>
       <Private>False</Private>
     </ProjectReference>
+    <ProjectReference Include="..\AutomationUtils\AutomationUtils.Automation.csproj">
+      <Project>{2c96a7f2-b1a3-4258-8e0a-e588ff41a53e}</Project>
+      <Name>AutomationUtils.Automation</Name>
+      <Private>False</Private>
+    </ProjectReference>
     <ProjectReference Include="..\Scripts\AutomationScripts.Automation.csproj">
       <Project>{8aa00d65-0954-4a27-ac0d-fb8b1106120f}</Project>
       <Name>AutomationScripts.Automation</Name>
diff --git a/Engine/Source/Programs/AutomationTool/HTML5/HTML5Platform.Automation.cs b/Engine/Source/Programs/AutomationTool/HTML5/HTML5Platform.Automation.cs
--- a/Engine/Source/Programs/AutomationTool/HTML5/HTML5Platform.Automation.cs
+++ b/Engine/Source/Programs/AutomationTool/HTML5/HTML5Platform.Automation.cs
@@ -87,8 +87,7 @@ public class HTML5Platform : Platform
         if (!SC.IsCodeBasedProject && !FileExists_NoExceptions(FullGameExePath))
         {
 	        Log("Failed to find game application " + FullGameExePath);
-	        AutomationTool.ErrorReporter.Error("Stage Failed.", (int)AutomationTool.ErrorCodes.Error_MissingExecutable);
-	        throw new AutomationException("Could not find application {0}. You may need to build the UE4 project with your target configuration and platform.", FullGameExePath);
+	        throw new AutomationException(ErrorCodes.Error_MissingExecutable, "Stage Failed. Could not find application {0}. You may need to build the UE4 project with your target configuration and platform.", FullGameExePath);
         }
 
         if (Path.Combine(Path.GetDirectoryName(Params.ProjectGameExeFilename), GameExe) != Path.Combine(PackagePath, GameExe))
diff --git a/Engine/Source/Programs/AutomationTool/IOS/IOS.Automation.csproj b/Engine/Source/Programs/AutomationTool/IOS/IOS.Automation.csproj
--- a/Engine/Source/Programs/AutomationTool/IOS/IOS.Automation.csproj
+++ b/Engine/Source/Programs/AutomationTool/IOS/IOS.Automation.csproj
@@ -69,16 +69,16 @@
       <Name>DeploymentInterface</Name>
       <Private>True</Private>
     </ProjectReference>
-    <ProjectReference Include="..\AutomationTool.csproj">
-      <Project>{a35b0261-d868-475e-a098-eb14ce6dc89f}</Project>
-      <Name>AutomationTool</Name>
-      <Private>False</Private>
-    </ProjectReference>
     <ProjectReference Include="..\..\UnrealBuildTool\UnrealBuildTool.csproj">
       <Project>{fd7c5e1a-cfe4-4fd5-a525-1eb1599a39ac}</Project>
       <Name>UnrealBuildTool</Name>
       <Private>False</Private>
     </ProjectReference>
+    <ProjectReference Include="..\AutomationUtils\AutomationUtils.Automation.csproj">
+      <Project>{2c96a7f2-b1a3-4258-8e0a-e588ff41a53e}</Project>
+      <Name>AutomationUtils.Automation</Name>
+      <Private>False</Private>
+    </ProjectReference>
   </ItemGroup>
   <Import Project="$(MSBuildToolsPath)\Microsoft.CSharp.targets" />
 </Project>
\ No newline at end of file
diff --git a/Engine/Source/Programs/AutomationTool/IOS/IOSPlatform.Automation.cs b/Engine/Source/Programs/AutomationTool/IOS/IOSPlatform.Automation.cs
--- a/Engine/Source/Programs/AutomationTool/IOS/IOSPlatform.Automation.cs
+++ b/Engine/Source/Programs/AutomationTool/IOS/IOSPlatform.Automation.cs
@@ -134,8 +134,7 @@ public class IOSPlatform : Platform
 		if (!SC.IsCodeBasedProject && !FileExists_NoExceptions(FullExePath))
 		{
 			Log("Failed to find game binary " + FullExePath);
-			AutomationTool.ErrorReporter.Error("Stage Failed.", (int)AutomationTool.ErrorCodes.Error_MissingExecutable);
-			throw new AutomationException("Could not find binary {0}. You may need to build the UE4 project with your target configuration and platform.", FullExePath);
+			throw new AutomationException(ErrorCodes.Error_MissingExecutable, "Stage Failed. Could not find binary {0}. You may need to build the UE4 project with your target configuration and platform.", FullExePath);
 		}
 
 		//@TODO: We should be able to use this code on both platforms, when the following issues are sorted:
@@ -323,8 +322,7 @@ public class IOSPlatform : Platform
 			// verify the .ipa exists
 			if (!FileExists(ProjectIPA))
 			{
-				ErrorReporter.Error(String.Format("PACKAGE FAILED - {0} was not created", ProjectIPA), (int)ErrorCodes.Error_FailedToCreateIPA);
-				throw new AutomationException("PACKAGE FAILED - {0} was not created", ProjectIPA);
+				throw new AutomationException(ErrorCodes.Error_FailedToCreateIPA, "PACKAGE FAILED - {0} was not created", ProjectIPA);
 			}
 
 			if (WorkingCL > 0)
@@ -415,8 +413,7 @@ public class IOSPlatform : Platform
 		}
 		if (Result.ExitCode != 0)
 		{
-			ErrorReporter.Error("CodeSign Failed", (int)ErrorCodes.Error_FailedToCodeSign);
-			throw new AutomationException("CodeSign Failed");
+			throw new AutomationException(ErrorCodes.Error_FailedToCodeSign, "CodeSign Failed");
 		}
 	}
 
@@ -824,12 +821,9 @@ public class IOSPlatform : Platform
 			switch (Result.ExitCode)
 			{
 				case 253:
-					ErrorReporter.Error("Launch Failure", (int)ErrorCodes.Error_DeviceNotSetupForDevelopment);
-					break;
-
+                    throw new AutomationException(ErrorCodes.Error_DeviceNotSetupForDevelopment, "Launch Failure");
 				case 255:
-					ErrorReporter.Error("Launch Failure", (int)ErrorCodes.Error_DeviceOSNewerThanSDK);
-					break;
+                    throw new AutomationException(ErrorCodes.Error_DeviceOSNewerThanSDK, "Launch Failure");
 			}
 		}
 	}
diff --git a/Engine/Source/Programs/AutomationTool/Linux/LinuxPlatform.Automation.cs b/Engine/Source/Programs/AutomationTool/Linux/LinuxPlatform.Automation.cs
--- a/Engine/Source/Programs/AutomationTool/Linux/LinuxPlatform.Automation.cs
+++ b/Engine/Source/Programs/AutomationTool/Linux/LinuxPlatform.Automation.cs
@@ -123,8 +123,7 @@ public abstract class BaseLinuxPlatform : Platform
 						if (!SC.IsCodeBasedProject && !FileExists_NoExceptions(Params.ProjectGameExeFilename) && !SC.bIsCombiningMultiplePlatforms)
 						{
 							Log("Failed to find game binary " + Params.ProjectGameExeFilename);
-							AutomationTool.ErrorReporter.Error("Stage Failed.", (int)AutomationTool.ErrorCodes.Error_MissingExecutable);
-							throw new AutomationException("Could not find game binary {0}. You may need to build the UE4 project with your target configuration and platform.", Params.ProjectGameExeFilename);
+                            throw new AutomationException(ErrorCodes.Error_MissingExecutable, "Stage Failed. Could not find game binary {0}. You may need to build the UE4 project with your target configuration and platform.", Params.ProjectGameExeFilename);
 						}
 
 					    SC.StageFiles(StagedFileType.NonUFS, CombinePaths(SC.LocalRoot, "Engine/Binaries", SC.PlatformDir), Path.GetFileNameWithoutExtension(Exe), true, null, CommandUtils.CombinePaths(SC.RelativeProjectRootForStage, "Binaries", SC.PlatformDir), false, true, SC.ShortProjectName);
diff --git a/Engine/Source/Programs/AutomationTool/Mac/MacPlatform.Automation.cs b/Engine/Source/Programs/AutomationTool/Mac/MacPlatform.Automation.cs
--- a/Engine/Source/Programs/AutomationTool/Mac/MacPlatform.Automation.cs
+++ b/Engine/Source/Programs/AutomationTool/Mac/MacPlatform.Automation.cs
@@ -115,8 +115,7 @@ public class MacPlatform : Platform
 					if (!SC.IsCodeBasedProject && !Directory.Exists(AbsoluteBundlePath) && !SC.bIsCombiningMultiplePlatforms)
 					{
 						Log("Failed to find app bundle " + AbsoluteBundlePath);
-						AutomationTool.ErrorReporter.Error("Stage Failed.", (int)AutomationTool.ErrorCodes.Error_MissingExecutable);
-						throw new AutomationException("Could not find app bundle {0}. You may need to build the UE4 project with your target configuration and platform.", AbsoluteBundlePath);
+						throw new AutomationException(ErrorCodes.Error_MissingExecutable, "Could not find app bundle {0}. You may need to build the UE4 project with your target configuration and platform.", AbsoluteBundlePath);
 					}
 
 					StageAppBundle(SC, WorkingFileType, CombinePaths(SC.LocalRoot, "Engine/Binaries", SC.PlatformDir, Path.GetFileNameWithoutExtension(Exe) + ".app"), AppBundlePath);
diff --git a/Engine/Source/Programs/AutomationTool/Program.cs b/Engine/Source/Programs/AutomationTool/Program.cs
--- a/Engine/Source/Programs/AutomationTool/Program.cs
+++ b/Engine/Source/Programs/AutomationTool/Program.cs
@@ -11,89 +11,23 @@ using System.Reflection;
 
 namespace AutomationTool
 {
-	// NOTE: this needs to be kept in sync with EditorAnalytics.h and iPhonePackager.cs
-	public enum ErrorCodes
-	{
-		Error_UATNotFound = -1,
-		Error_Success = 0,
-		Error_Unknown = 1,
-		Error_Arguments = 2,
-		Error_UnknownCommand = 3,
-		Error_SDKNotFound = 10,
-		Error_ProvisionNotFound = 11,
-		Error_CertificateNotFound = 12,
-		Error_ProvisionAndCertificateNotFound = 13,
-		Error_InfoPListNotFound = 14,
-		Error_KeyNotFoundInPList = 15,
-		Error_ProvisionExpired = 16,
-		Error_CertificateExpired = 17,
-		Error_CertificateProvisionMismatch = 18,
-		Error_CodeUnsupported = 19,
-		Error_PluginsUnsupported = 20,
-		Error_UnknownCookFailure = 25,
-		Error_UnknownDeployFailure = 26,
-		Error_UnknownBuildFailure = 27,
-		Error_UnknownPackageFailure = 28,
-		Error_UnknownLaunchFailure = 29,
-		Error_StageMissingFile = 30,
-		Error_FailedToCreateIPA = 31,
-		Error_FailedToCodeSign = 32,
-		Error_DeviceBackupFailed = 33,
-		Error_AppUninstallFailed = 34,
-		Error_AppInstallFailed = 35,
-		Error_AppNotFound = 36,
-		Error_StubNotSignedCorrectly = 37,
-		Error_IPAMissingInfoPList = 38,
-		Error_DeleteFile = 39,
-		Error_DeleteDirectory = 40,
-		Error_CreateDirectory = 41,
-		Error_CopyFile = 42,
-        Error_OnlyOneObbFileSupported = 50,
-        Error_FailureGettingPackageInfo = 51,
-        Error_OnlyOneTargetConfigurationSupported = 52,
-        Error_ObbNotFound = 53,
-        Error_AndroidBuildToolsPathNotFound = 54,
-        Error_NoApkSuitableForArchitecture = 55,
-		Error_FilesInstallFailed = 56,
-		Error_RemoteCertificatesNotFound = 57,
-		Error_LauncherFailed = 100,
-		Error_UATLaunchFailure = 101,
-		Error_FailedToDeleteStagingDirectory = 102,
-		Error_MissingExecutable = 103,
-		Error_DeviceNotSetupForDevelopment = 150,
-		Error_DeviceOSNewerThanSDK = 151,
-	};
-
-	public class ErrorReporter
-	{
-		static public void Error(string Line, int Code = 1)
-		{
-			if (Program.ReturnCode == 0)
-			{
-				Program.ReturnCode = Code;
-			}
-
-			Console.ForegroundColor = ConsoleColor.Red;
-			Log.WriteLine(TraceEventType.Error, "AutomationTool error: " + Line);
-			Console.ResetColor();
-		}
-	};
-
 	public class Program
 	{
 		// This needs to be static, otherwise SetConsoleCtrlHandler will result in a crash on exit.
 		static ProcessManager.CtrlHandlerDelegate ProgramCtrlHandler = new ProcessManager.CtrlHandlerDelegate(CtrlHandler);
-		static public int ReturnCode = 0;
 
 		[STAThread]
 		public static int Main()
 		{
-			var CommandLine = SharedUtils.ParseCommandLine();
-			HostPlatform.Initialize();
-			
-			LogUtils.InitLogging(CommandLine);
+            var CommandLine = SharedUtils.ParseCommandLine();
+            LogUtils.InitLogging(CommandLine);
+
+            ErrorCodes ReturnCode = ErrorCodes.Error_Success;
+
             try
             {
+                HostPlatform.Initialize();
+
                 Log.WriteLine(TraceEventType.Information, "Running on {0}", HostPlatform.Current.GetType().Name);
 
                 XmlConfigLoader.Init();
@@ -115,44 +49,65 @@ namespace AutomationTool
                 var Version = InternalUtils.ExecutableVersion;
                 Log.WriteLine(TraceEventType.Verbose, "{0} ver. {1}", Version.ProductName, Version.ProductVersion);
 
-                try
+                // Don't allow simultaneous execution of AT (in the same branch)
+                InternalUtils.RunSingleInstance(MainProc, CommandLine);
+
+            }
+            catch (Exception Ex)
+            {
+                // Catch all exceptions and propagate the ErrorCode if we are given one.
+                Log.WriteLine(TraceEventType.Error, "AutomationTool terminated with exception:");
+                Log.WriteLine(TraceEventType.Error, LogUtils.FormatException(Ex));
+                // set the exit code of the process
+                if (Ex is AutomationException)
                 {
-                    // Don't allow simultaneous execution of AT (in the same branch)
-                    ReturnCode = InternalUtils.RunSingleInstance(MainProc, CommandLine);
+                    ReturnCode = (Ex as AutomationException).ErrorCode;
                 }
-                catch (Exception Ex)
+                else
                 {
-                    Log.WriteLine(TraceEventType.Error, "AutomationTool terminated with exception:");
-                    Log.WriteLine(TraceEventType.Error, LogUtils.FormatException(Ex));
-                    Log.WriteLine(TraceEventType.Error, Ex.Message);
-                    if (ReturnCode == 0)
-                    {
-                        ReturnCode = (int)ErrorCodes.Error_Unknown;
-                    }
+                    ReturnCode = ErrorCodes.Error_Unknown;
                 }
+            }
+            finally
+            {
+                // In all cases, do necessary shut down stuff, but don't let any additional exceptions leak out while trying to shut down.
 
-                // Make sure there's no directiories on the stack.
-                CommandUtils.ClearDirStack();
-                Environment.ExitCode = ReturnCode;
+                // Make sure there's no directories on the stack.
+                NoThrow(() => CommandUtils.ClearDirStack(), "Clear Dir Stack");
 
                 // Try to kill process before app domain exits to leave the other KillAll call to extreme edge cases
-                if (ShouldKillProcesses && !Utils.IsRunningOnMono)
-                {
-                    ProcessManager.KillAll();
-                }
+                NoThrow(() => { if (ShouldKillProcesses && !Utils.IsRunningOnMono) ProcessManager.KillAll(); }, "Kill All Processes");
 
                 Log.WriteLine(TraceEventType.Information, "AutomationTool exiting with ExitCode={0}", ReturnCode);
-            }
-            finally
-            {
-                // ensure that logging is always shut down cleanly, as this potentially copies the log to a final resting place.
+
+                // Can't use NoThrow here because the code logs exceptions. We're shutting down logging!
                 LogUtils.ShutdownLogging();
             }
 
-			return ReturnCode;
-		}
+            // STOP: No code beyond the return statement should go beyond this point!
+            // Nothing should happen after the finally block above is finished. 
+            return (int)ReturnCode;
+        }
+
+        /// <summary>
+        /// Wraps an action in an exception block.
+        /// Ensures individual actions can be performed and exceptions won't prevent further actions from being executed.
+        /// Useful for shutdown code where shutdown may be in several stages and it's important that all stages get a chance to run.
+        /// </summary>
+        /// <param name="Action"></param>
+        private static void NoThrow(System.Action Action, string ActionDesc)
+        {
+            try
+            {
+                Action();
+            }
+            catch (Exception Ex)
+            {
+                Log.WriteLine(TraceEventType.Error, "Exception performing nothrow action \"{0}\": {1}", ActionDesc, LogUtils.FormatException(Ex));
+            }
+        }
 
-		static bool CtrlHandler(CtrlTypes EventType)
+        static bool CtrlHandler(CtrlTypes EventType)
 		{
 			Domain_ProcessExit(null, null);
 			if (EventType == CtrlTypes.CTRL_C_EVENT)
@@ -174,11 +129,10 @@ namespace AutomationTool
 			Trace.Close();
 		}
 
-		static int MainProc(object Param)
+		static void MainProc(object Param)
 		{
 			Automation.Process((string[])Param);
 			ShouldKillProcesses = Automation.ShouldKillProcesses;
-			return 0;
 		}
 
 		static bool ShouldKillProcesses = true;
diff --git a/Engine/Source/Programs/AutomationTool/Scripts/CopyBuildToStagingDirectory.Automation.cs b/Engine/Source/Programs/AutomationTool/Scripts/CopyBuildToStagingDirectory.Automation.cs
--- a/Engine/Source/Programs/AutomationTool/Scripts/CopyBuildToStagingDirectory.Automation.cs
+++ b/Engine/Source/Programs/AutomationTool/Scripts/CopyBuildToStagingDirectory.Automation.cs
@@ -1014,9 +1014,7 @@ public partial class Project : CommandUtils
             catch (Exception Ex)
             {
                 // Delete cooked data (if any) as it may be incomplete / corrupted.
-                Log("Failed to delete staging directory "+SC.StageDirectory);
-                AutomationTool.ErrorReporter.Error("Stage Failed.", (int)AutomationTool.ErrorCodes.Error_FailedToDeleteStagingDirectory);
-                throw Ex;   
+                throw new AutomationException(ErrorCodes.Error_FailedToDeleteStagingDirectory, Ex, "Stage Failed. Failed to delete staging directory " + SC.StageDirectory);
             }
 		}
 		else
@@ -1029,9 +1027,7 @@ public partial class Project : CommandUtils
             catch (Exception Ex)
             {
                 // Delete cooked data (if any) as it may be incomplete / corrupted.
-                Log("Failed to delete pak files in "+SC.StageDirectory);
-                AutomationTool.ErrorReporter.Error("Stage Failed.", (int)AutomationTool.ErrorCodes.Error_FailedToDeleteStagingDirectory);
-                throw Ex;   
+                throw new AutomationException(ErrorCodes.Error_FailedToDeleteStagingDirectory, Ex, "Stage Failed. Failed to delete pak files in " + SC.StageDirectory);
             }
 		}
 		if (ShouldCreatePak(Params, SC))
@@ -1483,8 +1479,7 @@ public partial class Project : CommandUtils
 							if (bRequireStagedFilesToExist)
 							{
 								// if we aren't collecting multiple platforms, then it is expected to exist
-								AutomationTool.ErrorReporter.Error("Stage Failed.", (int)AutomationTool.ErrorCodes.Error_MissingExecutable);
-								throw new AutomationException("Missing receipt '{0}'. Check that this target has been built.", Path.GetFileName(ReceiptFileName));
+                                throw new AutomationException(ErrorCodes.Error_MissingExecutable, "Stage Failed. Missing receipt '{0}'. Check that this target has been built.", Path.GetFileName(ReceiptFileName));
 							}
 							else
 							{
diff --git a/Engine/Source/Programs/AutomationTool/Scripts/RunProjectCommand.Automation.cs b/Engine/Source/Programs/AutomationTool/Scripts/RunProjectCommand.Automation.cs
--- a/Engine/Source/Programs/AutomationTool/Scripts/RunProjectCommand.Automation.cs
+++ b/Engine/Source/Programs/AutomationTool/Scripts/RunProjectCommand.Automation.cs
@@ -370,6 +370,8 @@ public partial class Project : CommandUtils
 
 				SC.StageTargetPlatform.PostRunClient(ClientProcess, Params);
 
+                // any non-zero exit code should propagate an exception. The Virtual function above may have
+                // already thrown a more specific exception or given a more specific ErrorCode, but this catches the rest.
 				if (ClientProcess.ExitCode != 0)
 				{
 					throw new AutomationException("Client exited with error code: " + ClientProcess.ExitCode);
