Remove ECProject; not referenced by anything.

diff --git a/Engine/Source/Programs/AutomationTool/GUBP/GUBP.cs b/Engine/Source/Programs/AutomationTool/GUBP/GUBP.cs
--- a/Engine/Source/Programs/AutomationTool/GUBP/GUBP.cs
+++ b/Engine/Source/Programs/AutomationTool/GUBP/GUBP.cs
@@ -24,7 +24,6 @@ public partial class GUBP : BuildCommand
     public bool bFake = false;
     public static bool bNoIOSOnPC = false;
     public static bool bForceIncrementalCompile = false;
-    public static string ECProject = null;
     public string EmailHint;
     static public bool bPreflightBuild = false;
     public int PreflightShelveCL = 0;
@@ -1440,11 +1439,6 @@ public partial class GUBP : BuildCommand
             bPreflightBuild = true;
         }
         
-        ECProject = ParseParamValue("ECProject");
-        if (ECProject == null)
-        {
-            ECProject = "";
-        }		
         HostPlatforms = new List<UnrealTargetPlatform>();
         if (!ParseParam("NoPC"))
         {
