Change GUBPNodesHistory into a local variable.

diff --git a/Engine/Source/Programs/AutomationTool/GUBP/GUBP.cs b/Engine/Source/Programs/AutomationTool/GUBP/GUBP.cs
--- a/Engine/Source/Programs/AutomationTool/GUBP/GUBP.cs
+++ b/Engine/Source/Programs/AutomationTool/GUBP/GUBP.cs
@@ -48,8 +48,6 @@ public partial class GUBP : BuildCommand
         public List<int> AllFailed = new List<int>();
     };
 
-    Dictionary<string, NodeHistory> GUBPNodesHistory;
-
     public string AddNode(GUBPNode Node)
     {
         string Name = Node.GetFullName();
@@ -497,7 +495,7 @@ public partial class GUBP : BuildCommand
         Log("Took {0}s to get P4 history", BuildDuration / 1000);
 
     }
-    void PrintNodes(GUBP bp, List<string> Nodes, bool LocalOnly, List<string> UnfinishedTriggers = null)
+    void PrintNodes(GUBP bp, List<string> Nodes, Dictionary<string, NodeHistory> GUBPNodesHistory, bool LocalOnly, List<string> UnfinishedTriggers = null)
     {
         bool bShowAllChanges = bp.ParseParam("AllChanges") && GUBPNodesHistory != null;
         bool bShowChanges = (bp.ParseParam("Changes") && GUBPNodesHistory != null) || bShowAllChanges;
@@ -1084,7 +1082,7 @@ public partial class GUBP : BuildCommand
         return ParentPath + "/" + GetJobStepPath(Dep);
     }
 
-    void UpdateNodeHistory(string Node, string CLString)
+    void UpdateNodeHistory(Dictionary<string, NodeHistory> GUBPNodesHistory, string Node, string CLString)
     {
         if (GUBPNodes[Node].RunInEC() && !GUBPNodes[Node].TriggerNode() && CLString != "")
         {
@@ -1128,7 +1126,7 @@ public partial class GUBP : BuildCommand
 			GUBPNodesHistory.Add(Node, History);
         }
     }
-	void GetFailureEmails(string NodeToDo, string CLString, bool OnlyLateUpdates = false)
+	void GetFailureEmails(Dictionary<string, NodeHistory> GUBPNodesHistory, string NodeToDo, string CLString, bool OnlyLateUpdates = false)
 	{
         string EMails;
         string FailCauserEMails = "";
@@ -1148,7 +1146,7 @@ public partial class GUBP : BuildCommand
                 {
                     if (OnlyLateUpdates)
                     {
-                        LastNonDuplicateFail = FindLastNonDuplicateFail(NodeToDo, CLString);
+                        LastNonDuplicateFail = FindLastNonDuplicateFail(GUBPNodesHistory, NodeToDo, CLString);
                         if (LastNonDuplicateFail < P4Env.Changelist)
                         {
                             Log("*** Red-after-red spam reduction, changed CL {0} to CL {1} because the errors didn't change.", P4Env.Changelist, LastNonDuplicateFail);
@@ -1255,7 +1253,7 @@ public partial class GUBP : BuildCommand
         return true;
     }
 
-    int FindLastNonDuplicateFail(string NodeToDo, string CLString)
+    int FindLastNonDuplicateFail(Dictionary<string, NodeHistory> GUBPNodesHistory, string NodeToDo, string CLString)
     {
         var History = GUBPNodesHistory[NodeToDo];
         int Result = P4Env.Changelist;
@@ -2310,7 +2308,7 @@ public partial class GUBP : BuildCommand
 		}
 
         GUBPNodesCompleted = new Dictionary<string, bool>();
-        GUBPNodesHistory = new Dictionary<string, NodeHistory>();
+        Dictionary<string, NodeHistory> GUBPNodesHistory = new Dictionary<string, NodeHistory>();
 
         LogVerbose("******* Caching completion");
         {
@@ -2358,7 +2356,7 @@ public partial class GUBP : BuildCommand
         }
 
         LogVerbose("*********** Desired And Dependent Nodes, in order.");
-        PrintNodes(this, OrdereredToDo, LocalOnly, UnfinishedTriggers);		
+        PrintNodes(this, OrdereredToDo, GUBPNodesHistory, LocalOnly, UnfinishedTriggers);		
         //check sorting
         {
             foreach (var NodeToDo in OrdereredToDo)
@@ -2466,7 +2464,7 @@ public partial class GUBP : BuildCommand
             var FinishFilterTime = DateTime.Now.ToString();
             PrintCSVFile(String.Format("UAT,FilterNodes,{0},{1}", StartFilterTime, FinishFilterTime));            
             Log("*********** EC Nodes, in order.");
-            PrintNodes(this, OrdereredToDo, LocalOnly, UnfinishedTriggers);
+            PrintNodes(this, OrdereredToDo, GUBPNodesHistory, LocalOnly, UnfinishedTriggers);
             var FinishNodePrint = DateTime.Now.ToString();
             PrintCSVFile(String.Format("UAT,SetupCommanderPrint,{0},{1}", FinishFilterTime, FinishNodePrint));
             // here we are just making sure everything before the explicit trigger is completed.
@@ -3090,7 +3088,7 @@ public partial class GUBP : BuildCommand
                     if (SaveSuccessRecords)
                     {
                         var StartNodeHistory = DateTime.Now.ToString();
-                        UpdateNodeHistory(NodeToDo, CLString);
+                        UpdateNodeHistory(GUBPNodesHistory, NodeToDo, CLString);
                         var FinishNodeHistory = DateTime.Now.ToString();                        
                         SaveStatus(NodeToDo, FailedTempStorageSuffix, NodeStoreName, bSaveSharedTempStorage, GameNameIfAny, ParseParamValue("MyJobStepId"));
                         var FinishSaveStatus = DateTime.Now.ToString();                        
@@ -3099,7 +3097,7 @@ public partial class GUBP : BuildCommand
                         
 						if (IsBuildMachine)
 						{
-							GetFailureEmails(NodeToDo, CLString);
+							GetFailureEmails(GUBPNodesHistory, NodeToDo, CLString);
                             var FinishFailEmails = DateTime.Now.ToString();
                             PrintCSVFile(String.Format("UAT,UpdateNodeHistory,{0},{1}", StartNodeHistory, FinishNodeHistory));
                             PrintCSVFile(String.Format("UAT,SaveNodeStatus,{0},{1}", FinishNodeHistory, FinishSaveStatus));
@@ -3157,7 +3155,7 @@ public partial class GUBP : BuildCommand
                 if (SaveSuccessRecords) 
                 {
                     var StartNodeHistory = DateTime.Now.ToString();
-                    UpdateNodeHistory(NodeToDo, CLString);
+                    UpdateNodeHistory(GUBPNodesHistory, NodeToDo, CLString);
                     var FinishNodeHistory = DateTime.Now.ToString();
                     SaveStatus(NodeToDo, SucceededTempStorageSuffix, NodeStoreName, bSaveSharedTempStorage, GameNameIfAny);
                     var FinishSaveStatus = DateTime.Now.ToString();                    
@@ -3166,7 +3164,7 @@ public partial class GUBP : BuildCommand
                     
 					if (IsBuildMachine)
 					{
-						GetFailureEmails(NodeToDo, CLString);
+						GetFailureEmails(GUBPNodesHistory, NodeToDo, CLString);
                         var FinishFailEmails = DateTime.Now.ToString();
                         PrintCSVFile(String.Format("UAT,UpdateNodeHistory,{0},{1}", StartNodeHistory, FinishNodeHistory));
                         PrintCSVFile(String.Format("UAT,SaveNodeStatus,{0},{1}", FinishNodeHistory, FinishSaveStatus));
