Convert GUBPNodesControllingTrigger to local variables.

diff --git a/Engine/Source/Programs/AutomationTool/GUBP/GUBP.cs b/Engine/Source/Programs/AutomationTool/GUBP/GUBP.cs
--- a/Engine/Source/Programs/AutomationTool/GUBP/GUBP.cs
+++ b/Engine/Source/Programs/AutomationTool/GUBP/GUBP.cs
@@ -32,8 +32,6 @@ public partial class GUBP : BuildCommand
 
     Dictionary<string, GUBPNode> GUBPNodes;
     Dictionary<string, bool> GUBPNodesCompleted;
-    Dictionary<string, string> GUBPNodesControllingTrigger;
-    Dictionary<string, string> GUBPNodesControllingTriggerDotName;
 
     class NodeHistory
     {
@@ -178,7 +176,7 @@ public partial class GUBP : BuildCommand
         return Deps.Contains(Rootward);
     }
 
-    string GetControllingTrigger(string NodeToDo)
+    string GetControllingTrigger(string NodeToDo, Dictionary<string, string> GUBPNodesControllingTrigger)
     {
         if (GUBPNodesControllingTrigger.ContainsKey(NodeToDo))
         {
@@ -202,7 +200,7 @@ public partial class GUBP : BuildCommand
             }
             else
             {
-                string NewResult = GetControllingTrigger(Node);
+                string NewResult = GetControllingTrigger(Node, GUBPNodesControllingTrigger);
                 if (!String.IsNullOrEmpty(NewResult))
                 {
                     if (NewResult != Result && !string.IsNullOrEmpty(Result))
@@ -230,7 +228,7 @@ public partial class GUBP : BuildCommand
             }
             else
             {
-                string NewResult = GetControllingTrigger(Node);
+                string NewResult = GetControllingTrigger(Node, GUBPNodesControllingTrigger);
                 if (!String.IsNullOrEmpty(NewResult))
                 {
                     if (NewResult != Result && !string.IsNullOrEmpty(Result))
@@ -244,7 +242,7 @@ public partial class GUBP : BuildCommand
         GUBPNodesControllingTrigger.Add(NodeToDo, Result);
         return Result;
     }
-    string GetControllingTriggerDotName(string NodeToDo)
+    string GetControllingTriggerDotName(string NodeToDo, Dictionary<string, string> GUBPNodesControllingTriggerDotName, Dictionary<string, string> GUBPNodesControllingTrigger)
     {
         if (GUBPNodesControllingTriggerDotName.ContainsKey(NodeToDo))
         {
@@ -254,7 +252,7 @@ public partial class GUBP : BuildCommand
         string WorkingNode = NodeToDo;
         while (true)
         {
-            string ThisResult = GetControllingTrigger(WorkingNode);
+            string ThisResult = GetControllingTrigger(WorkingNode, GUBPNodesControllingTrigger);
             if (ThisResult == "")
             {
                 break;
@@ -495,7 +493,7 @@ public partial class GUBP : BuildCommand
         Log("Took {0}s to get P4 history", BuildDuration / 1000);
 
     }
-    void PrintNodes(GUBP bp, List<string> Nodes, Dictionary<string, NodeHistory> GUBPNodesHistory, bool LocalOnly, List<string> UnfinishedTriggers = null)
+    void PrintNodes(GUBP bp, List<string> Nodes, Dictionary<string, NodeHistory> GUBPNodesHistory, Dictionary<string, string> GUBPNodesControllingTriggerDotName, Dictionary<string, string> GUBPNodesControllingTrigger, bool LocalOnly, List<string> UnfinishedTriggers = null)
     {
         bool bShowAllChanges = bp.ParseParam("AllChanges") && GUBPNodesHistory != null;
         bool bShowChanges = (bp.ParseParam("Changes") && GUBPNodesHistory != null) || bShowAllChanges;
@@ -524,7 +522,7 @@ public partial class GUBP : BuildCommand
             }
             if (bShowTriggers)
             {
-                string MyControllingTrigger = GetControllingTriggerDotName(NodeToDo);
+                string MyControllingTrigger = GetControllingTriggerDotName(NodeToDo, GUBPNodesControllingTriggerDotName, GUBPNodesControllingTrigger);
                 if (MyControllingTrigger != LastControllingTrigger)
                 {
                     LastControllingTrigger = MyControllingTrigger;
@@ -533,7 +531,7 @@ public partial class GUBP : BuildCommand
                         string Finished = "";
                         if (UnfinishedTriggers != null)
                         {
-                            string MyShortControllingTrigger = GetControllingTrigger(NodeToDo);
+                            string MyShortControllingTrigger = GetControllingTrigger(NodeToDo, GUBPNodesControllingTrigger);
                             if (UnfinishedTriggers.Contains(MyShortControllingTrigger))
                             {
                                 Finished = "(not yet triggered)";
@@ -815,7 +813,7 @@ public partial class GUBP : BuildCommand
         return Result;
     }
 
-    List<string> TopologicalSort(HashSet<string> NodesToDo, string ExplicitTrigger = "", bool LocalOnly = false, bool SubSort = false, bool DoNotConsiderCompletion = false)
+    List<string> TopologicalSort(HashSet<string> NodesToDo, Dictionary<string, string> GUBPNodesControllingTriggerDotName, Dictionary<string, string> GUBPNodesControllingTrigger, string ExplicitTrigger = "", bool LocalOnly = false, bool SubSort = false, bool DoNotConsiderCompletion = false)
     {
         var StartTime = DateTime.UtcNow;
 
@@ -842,7 +840,7 @@ public partial class GUBP : BuildCommand
             }
             foreach (var Chain in AgentGroupChains)
             {
-                SortedAgentGroupChains.Add(Chain.Key, TopologicalSort(new HashSet<string>(Chain.Value), ExplicitTrigger, LocalOnly, true, DoNotConsiderCompletion));
+                SortedAgentGroupChains.Add(Chain.Key, TopologicalSort(new HashSet<string>(Chain.Value), GUBPNodesControllingTriggerDotName, GUBPNodesControllingTrigger, ExplicitTrigger, LocalOnly, true, DoNotConsiderCompletion));
             }
             Log("***************Done with recursion");
         }
@@ -950,11 +948,11 @@ public partial class GUBP : BuildCommand
 
                 if (bReady && BestNode != "")
                 {
-                    if (String.Compare(GetControllingTriggerDotName(BestNode), GetControllingTriggerDotName(NodeToDo)) < 0) //sorted by controlling trigger
+                    if (String.Compare(GetControllingTriggerDotName(BestNode, GUBPNodesControllingTriggerDotName, GUBPNodesControllingTrigger), GetControllingTriggerDotName(NodeToDo, GUBPNodesControllingTriggerDotName, GUBPNodesControllingTrigger)) < 0) //sorted by controlling trigger
                     {
                         bReady = false;
                     }
-                    else if (String.Compare(GetControllingTriggerDotName(BestNode), GetControllingTriggerDotName(NodeToDo)) == 0) //sorted by controlling trigger
+                    else if (String.Compare(GetControllingTriggerDotName(BestNode, GUBPNodesControllingTriggerDotName, GUBPNodesControllingTrigger), GetControllingTriggerDotName(NodeToDo, GUBPNodesControllingTriggerDotName, GUBPNodesControllingTrigger)) == 0) //sorted by controlling trigger
                     {
                         if (GUBPNodes[BestNode].IsSticky() && !GUBPNodes[NodeToDo].IsSticky()) //sticky nodes first
                         {
@@ -1809,8 +1807,8 @@ public partial class GUBP : BuildCommand
             TempStorage.DeleteLocalTempStorageManifests(CmdEnv);
         }
 
-        GUBPNodesControllingTrigger = new Dictionary<string, string>();
-        GUBPNodesControllingTriggerDotName = new Dictionary<string, string>();
+		Dictionary<string, string> GUBPNodesControllingTrigger = new Dictionary<string, string>();
+        Dictionary<string, string> GUBPNodesControllingTriggerDotName = new Dictionary<string, string>();
 
         var FullNodeList = new Dictionary<string, string>();
         var FullNodeDirectDependencies = new Dictionary<string, string>();
@@ -1906,11 +1904,11 @@ public partial class GUBP : BuildCommand
         }	
         {
             Log("******* {0} GUBP Nodes", GUBPNodes.Count);
-            var SortedNodes = TopologicalSort(new HashSet<string>(GUBPNodes.Keys), LocalOnly: true, DoNotConsiderCompletion: true);
+            var SortedNodes = TopologicalSort(new HashSet<string>(GUBPNodes.Keys), GUBPNodesControllingTriggerDotName, GUBPNodesControllingTrigger, LocalOnly: true, DoNotConsiderCompletion: true);
             var DependencyStart = DateTime.Now.ToString();
             foreach (var Node in SortedNodes)
             {
-                string Note = GetControllingTriggerDotName(Node);
+                string Note = GetControllingTriggerDotName(Node, GUBPNodesControllingTriggerDotName, GUBPNodesControllingTrigger);
                 if (Note == "")
                 {
                     Note = CISFrequencyQuantumShiftString(Node);
@@ -2293,7 +2291,7 @@ public partial class GUBP : BuildCommand
 			var NewNodesToDo = new HashSet<string>();
 			foreach (var NodeToDo in NodesToDo)
 			{
-				var TriggerDot = GetControllingTriggerDotName(NodeToDo);
+				var TriggerDot = GetControllingTriggerDotName(NodeToDo, GUBPNodesControllingTriggerDotName, GUBPNodesControllingTrigger);
 				if (TriggerDot == "" && !GUBPNodes[NodeToDo].TriggerNode())
 				{
 					LogVerbose("  Keeping {0}", NodeToDo);
@@ -2317,7 +2315,7 @@ public partial class GUBP : BuildCommand
             {
                 LogVerbose("** {0}", Node);
                 NodeIsAlreadyComplete(Node, LocalOnly); // cache these now to avoid spam later
-                GetControllingTriggerDotName(Node);
+                GetControllingTriggerDotName(Node, GUBPNodesControllingTriggerDotName, GUBPNodesControllingTrigger);
             }
             var BuildDuration = (DateTime.UtcNow - StartTime).TotalMilliseconds;
 			LogVerbose("Took {0}s to cache completion for {1} nodes", BuildDuration / 1000, NodesToDo.Count);
@@ -2337,7 +2335,7 @@ public partial class GUBP : BuildCommand
             Log("Took {0}s to get history for {1} nodes", BuildDuration / 1000, NodesToDo.Count);
         }*/
 
-        var OrdereredToDo = TopologicalSort(NodesToDo, ExplicitTrigger, LocalOnly);
+        var OrdereredToDo = TopologicalSort(NodesToDo, GUBPNodesControllingTriggerDotName, GUBPNodesControllingTrigger, ExplicitTrigger, LocalOnly);
 
         // find all unfinished triggers, excepting the one we are triggering right now
         var UnfinishedTriggers = new List<string>();
@@ -2356,7 +2354,7 @@ public partial class GUBP : BuildCommand
         }
 
         LogVerbose("*********** Desired And Dependent Nodes, in order.");
-        PrintNodes(this, OrdereredToDo, GUBPNodesHistory, LocalOnly, UnfinishedTriggers);		
+        PrintNodes(this, OrdereredToDo, GUBPNodesHistory, GUBPNodesControllingTriggerDotName, GUBPNodesControllingTrigger, LocalOnly, UnfinishedTriggers);		
         //check sorting
         {
             foreach (var NodeToDo in OrdereredToDo)
@@ -2438,7 +2436,7 @@ public partial class GUBP : BuildCommand
             // remove nodes that have unfinished triggers
             foreach (var NodeToDo in OrdereredToDo)
             {
-                string ControllingTrigger = GetControllingTrigger(NodeToDo);
+                string ControllingTrigger = GetControllingTrigger(NodeToDo, GUBPNodesControllingTrigger);
                 bool bNoUnfinishedTriggers = !UnfinishedTriggers.Contains(ControllingTrigger);
 
                 if (bNoUnfinishedTriggers)
@@ -2464,7 +2462,7 @@ public partial class GUBP : BuildCommand
             var FinishFilterTime = DateTime.Now.ToString();
             PrintCSVFile(String.Format("UAT,FilterNodes,{0},{1}", StartFilterTime, FinishFilterTime));            
             Log("*********** EC Nodes, in order.");
-            PrintNodes(this, OrdereredToDo, GUBPNodesHistory, LocalOnly, UnfinishedTriggers);
+            PrintNodes(this, OrdereredToDo, GUBPNodesHistory, GUBPNodesControllingTriggerDotName, GUBPNodesControllingTrigger, LocalOnly, UnfinishedTriggers);
             var FinishNodePrint = DateTime.Now.ToString();
             PrintCSVFile(String.Format("UAT,SetupCommanderPrint,{0},{1}", FinishFilterTime, FinishNodePrint));
             // here we are just making sure everything before the explicit trigger is completed.
