Change completed node cache into a local variable.

diff --git a/Engine/Source/Programs/AutomationTool/GUBP/GUBP.cs b/Engine/Source/Programs/AutomationTool/GUBP/GUBP.cs
--- a/Engine/Source/Programs/AutomationTool/GUBP/GUBP.cs
+++ b/Engine/Source/Programs/AutomationTool/GUBP/GUBP.cs
@@ -31,7 +31,6 @@ public partial class GUBP : BuildCommand
     public GUBPBranchHacker.BranchOptions BranchOptions = null;	
 
     Dictionary<string, GUBPNode> GUBPNodes;
-    Dictionary<string, bool> GUBPNodesCompleted;
 
     class NodeHistory
     {
@@ -327,7 +326,7 @@ public partial class GUBP : BuildCommand
         return Result;
     }
 
-    bool NodeIsAlreadyComplete(string NodeToDo, bool LocalOnly)
+    bool NodeIsAlreadyComplete(Dictionary<string, bool> GUBPNodesCompleted, string NodeToDo, bool LocalOnly)
     {
         if (GUBPNodesCompleted.ContainsKey(NodeToDo))
         {
@@ -493,7 +492,7 @@ public partial class GUBP : BuildCommand
         Log("Took {0}s to get P4 history", BuildDuration / 1000);
 
     }
-    void PrintNodes(GUBP bp, List<string> Nodes, Dictionary<string, NodeHistory> GUBPNodesHistory, Dictionary<string, string> GUBPNodesControllingTriggerDotName, Dictionary<string, string> GUBPNodesControllingTrigger, bool LocalOnly, List<string> UnfinishedTriggers = null)
+    void PrintNodes(GUBP bp, List<string> Nodes, Dictionary<string, bool> GUBPNodesCompleted, Dictionary<string, NodeHistory> GUBPNodesHistory, Dictionary<string, string> GUBPNodesControllingTriggerDotName, Dictionary<string, string> GUBPNodesControllingTrigger, bool LocalOnly, List<string> UnfinishedTriggers = null)
     {
         bool bShowAllChanges = bp.ParseParam("AllChanges") && GUBPNodesHistory != null;
         bool bShowChanges = (bp.ParseParam("Changes") && GUBPNodesHistory != null) || bShowAllChanges;
@@ -571,7 +570,7 @@ public partial class GUBP : BuildCommand
                 (LastAgentGroup != "" ? "  " : ""),
                 NodeToDo,
                 FrequencyString,
-                NodeIsAlreadyComplete(NodeToDo, LocalOnly) ? " - (Completed)" : "",
+                NodeIsAlreadyComplete(GUBPNodesCompleted, NodeToDo, LocalOnly) ? " - (Completed)" : "",
                 GUBPNodes[NodeToDo].TriggerNode() ? " - (TriggerNode)" : "",
                 GUBPNodes[NodeToDo].IsSticky() ? " - (Sticky)" : "",				
                 Agent,
@@ -813,7 +812,7 @@ public partial class GUBP : BuildCommand
         return Result;
     }
 
-    List<string> TopologicalSort(HashSet<string> NodesToDo, Dictionary<string, string> GUBPNodesControllingTriggerDotName, Dictionary<string, string> GUBPNodesControllingTrigger, string ExplicitTrigger = "", bool LocalOnly = false, bool SubSort = false, bool DoNotConsiderCompletion = false)
+    List<string> TopologicalSort(HashSet<string> NodesToDo, Dictionary<string, bool> GUBPNodesCompleted, Dictionary<string, string> GUBPNodesControllingTriggerDotName, Dictionary<string, string> GUBPNodesControllingTrigger, string ExplicitTrigger = "", bool LocalOnly = false, bool SubSort = false, bool DoNotConsiderCompletion = false)
     {
         var StartTime = DateTime.UtcNow;
 
@@ -840,7 +839,7 @@ public partial class GUBP : BuildCommand
             }
             foreach (var Chain in AgentGroupChains)
             {
-                SortedAgentGroupChains.Add(Chain.Key, TopologicalSort(new HashSet<string>(Chain.Value), GUBPNodesControllingTriggerDotName, GUBPNodesControllingTrigger, ExplicitTrigger, LocalOnly, true, DoNotConsiderCompletion));
+                SortedAgentGroupChains.Add(Chain.Key, TopologicalSort(new HashSet<string>(Chain.Value), GUBPNodesCompleted, GUBPNodesControllingTriggerDotName, GUBPNodesControllingTrigger, ExplicitTrigger, LocalOnly, true, DoNotConsiderCompletion));
             }
             Log("***************Done with recursion");
         }
@@ -966,8 +965,8 @@ public partial class GUBP : BuildCommand
                             }
                             else if (BestPseudoReady == bPseudoReady)
                             {
-                                bool IamLateTrigger = !DoNotConsiderCompletion && GUBPNodes[NodeToDo].TriggerNode() && NodeToDo != ExplicitTrigger && !NodeIsAlreadyComplete(NodeToDo, LocalOnly);
-                                bool BestIsLateTrigger = !DoNotConsiderCompletion && GUBPNodes[BestNode].TriggerNode() && BestNode != ExplicitTrigger && !NodeIsAlreadyComplete(BestNode, LocalOnly);
+                                bool IamLateTrigger = !DoNotConsiderCompletion && GUBPNodes[NodeToDo].TriggerNode() && NodeToDo != ExplicitTrigger && !NodeIsAlreadyComplete(GUBPNodesCompleted, NodeToDo, LocalOnly);
+                                bool BestIsLateTrigger = !DoNotConsiderCompletion && GUBPNodes[BestNode].TriggerNode() && BestNode != ExplicitTrigger && !NodeIsAlreadyComplete(GUBPNodesCompleted, BestNode, LocalOnly);
                                 if (BestIsLateTrigger && !IamLateTrigger)
                                 {
                                     bReady = false;
@@ -1904,7 +1903,7 @@ public partial class GUBP : BuildCommand
         }	
         {
             Log("******* {0} GUBP Nodes", GUBPNodes.Count);
-            var SortedNodes = TopologicalSort(new HashSet<string>(GUBPNodes.Keys), GUBPNodesControllingTriggerDotName, GUBPNodesControllingTrigger, LocalOnly: true, DoNotConsiderCompletion: true);
+            var SortedNodes = TopologicalSort(new HashSet<string>(GUBPNodes.Keys), null/*GUBPNodesCompleted*/, GUBPNodesControllingTriggerDotName, GUBPNodesControllingTrigger, LocalOnly: true, DoNotConsiderCompletion: true);
             var DependencyStart = DateTime.Now.ToString();
             foreach (var Node in SortedNodes)
             {
@@ -2305,7 +2304,7 @@ public partial class GUBP : BuildCommand
 			NodesToDo = NewNodesToDo;
 		}
 
-        GUBPNodesCompleted = new Dictionary<string, bool>();
+        Dictionary<string, bool> GUBPNodesCompleted = new Dictionary<string, bool>();
         Dictionary<string, NodeHistory> GUBPNodesHistory = new Dictionary<string, NodeHistory>();
 
         LogVerbose("******* Caching completion");
@@ -2314,7 +2313,7 @@ public partial class GUBP : BuildCommand
             foreach (var Node in NodesToDo)
             {
                 LogVerbose("** {0}", Node);
-                NodeIsAlreadyComplete(Node, LocalOnly); // cache these now to avoid spam later
+                NodeIsAlreadyComplete(GUBPNodesCompleted, Node, LocalOnly); // cache these now to avoid spam later
                 GetControllingTriggerDotName(Node, GUBPNodesControllingTriggerDotName, GUBPNodesControllingTrigger);
             }
             var BuildDuration = (DateTime.UtcNow - StartTime).TotalMilliseconds;
@@ -2335,7 +2334,7 @@ public partial class GUBP : BuildCommand
             Log("Took {0}s to get history for {1} nodes", BuildDuration / 1000, NodesToDo.Count);
         }*/
 
-        var OrdereredToDo = TopologicalSort(NodesToDo, GUBPNodesControllingTriggerDotName, GUBPNodesControllingTrigger, ExplicitTrigger, LocalOnly);
+        var OrdereredToDo = TopologicalSort(NodesToDo, GUBPNodesCompleted, GUBPNodesControllingTriggerDotName, GUBPNodesControllingTrigger, ExplicitTrigger, LocalOnly);
 
         // find all unfinished triggers, excepting the one we are triggering right now
         var UnfinishedTriggers = new List<string>();
@@ -2343,7 +2342,7 @@ public partial class GUBP : BuildCommand
         {
             foreach (var NodeToDo in OrdereredToDo)
             {
-                if (GUBPNodes[NodeToDo].TriggerNode() && !NodeIsAlreadyComplete(NodeToDo, LocalOnly))
+                if (GUBPNodes[NodeToDo].TriggerNode() && !NodeIsAlreadyComplete(GUBPNodesCompleted, NodeToDo, LocalOnly))
                 {
                     if (String.IsNullOrEmpty(ExplicitTrigger) || ExplicitTrigger != NodeToDo)
                     {
@@ -2354,12 +2353,12 @@ public partial class GUBP : BuildCommand
         }
 
         LogVerbose("*********** Desired And Dependent Nodes, in order.");
-        PrintNodes(this, OrdereredToDo, GUBPNodesHistory, GUBPNodesControllingTriggerDotName, GUBPNodesControllingTrigger, LocalOnly, UnfinishedTriggers);		
+        PrintNodes(this, OrdereredToDo, GUBPNodesCompleted, GUBPNodesHistory, GUBPNodesControllingTriggerDotName, GUBPNodesControllingTrigger, LocalOnly, UnfinishedTriggers);		
         //check sorting
         {
             foreach (var NodeToDo in OrdereredToDo)
             {
-                if (GUBPNodes[NodeToDo].TriggerNode() && (GUBPNodes[NodeToDo].IsSticky() || NodeIsAlreadyComplete(NodeToDo, LocalOnly))) // these sticky triggers are ok, everything is already completed anyway
+                if (GUBPNodes[NodeToDo].TriggerNode() && (GUBPNodes[NodeToDo].IsSticky() || NodeIsAlreadyComplete(GUBPNodesCompleted, NodeToDo, LocalOnly))) // these sticky triggers are ok, everything is already completed anyway
                 {
                     continue;
                 }
@@ -2462,7 +2461,7 @@ public partial class GUBP : BuildCommand
             var FinishFilterTime = DateTime.Now.ToString();
             PrintCSVFile(String.Format("UAT,FilterNodes,{0},{1}", StartFilterTime, FinishFilterTime));            
             Log("*********** EC Nodes, in order.");
-            PrintNodes(this, OrdereredToDo, GUBPNodesHistory, GUBPNodesControllingTriggerDotName, GUBPNodesControllingTrigger, LocalOnly, UnfinishedTriggers);
+            PrintNodes(this, OrdereredToDo, GUBPNodesCompleted, GUBPNodesHistory, GUBPNodesControllingTriggerDotName, GUBPNodesControllingTrigger, LocalOnly, UnfinishedTriggers);
             var FinishNodePrint = DateTime.Now.ToString();
             PrintCSVFile(String.Format("UAT,SetupCommanderPrint,{0},{1}", FinishFilterTime, FinishNodePrint));
             // here we are just making sure everything before the explicit trigger is completed.
@@ -2470,7 +2469,7 @@ public partial class GUBP : BuildCommand
             {
                 foreach (var NodeToDo in FilteredOrdereredToDo)
                 {
-                    if (GUBPNodes[NodeToDo].RunInEC() && !NodeIsAlreadyComplete(NodeToDo, LocalOnly) && NodeToDo != ExplicitTrigger && !NodeDependsOn(ExplicitTrigger, NodeToDo)) // if something is already finished, we don't put it into EC
+                    if (GUBPNodes[NodeToDo].RunInEC() && !NodeIsAlreadyComplete(GUBPNodesCompleted, NodeToDo, LocalOnly) && NodeToDo != ExplicitTrigger && !NodeDependsOn(ExplicitTrigger, NodeToDo)) // if something is already finished, we don't put it into EC
                     {
                         throw new AutomationException("We are being asked to process node {0}, however, this is an explicit trigger {1}, so everything before it should already be handled. It seems likely that you waited too long to run the trigger. You will have to do a new build from scratch.", NodeToDo, ExplicitTrigger);
                     }
@@ -2490,7 +2489,7 @@ public partial class GUBP : BuildCommand
             // sticky nodes are ones that we run on the main agent. We run then first and they must not be intermixed with parallel jobs
             foreach (var NodeToDo in OrdereredToDo)
             {
-                if (GUBPNodes[NodeToDo].RunInEC() && !NodeIsAlreadyComplete(NodeToDo, LocalOnly)) // if something is already finished, we don't put it into EC
+                if (GUBPNodes[NodeToDo].RunInEC() && !NodeIsAlreadyComplete(GUBPNodesCompleted, NodeToDo, LocalOnly)) // if something is already finished, we don't put it into EC
                 {
                     bHaveECNodes = true;
                     if (GUBPNodes[NodeToDo].IsSticky())
@@ -2526,7 +2525,7 @@ public partial class GUBP : BuildCommand
             var StickyChain = new List<string>();
             foreach (var NodeToDo in OrdereredToDo)
             {
-                if (GUBPNodes[NodeToDo].RunInEC() && !NodeIsAlreadyComplete(NodeToDo, LocalOnly)) // if something is already finished, we don't put it into EC  
+                if (GUBPNodes[NodeToDo].RunInEC() && !NodeIsAlreadyComplete(GUBPNodesCompleted, NodeToDo, LocalOnly)) // if something is already finished, we don't put it into EC  
                 {
                     string MyAgentGroup = GUBPNodes[NodeToDo].AgentSharingGroup;
                     if (MyAgentGroup != "")
@@ -2551,7 +2550,7 @@ public partial class GUBP : BuildCommand
             }
             foreach (var NodeToDo in OrdereredToDo)
             {
-                if (GUBPNodes[NodeToDo].RunInEC() && !NodeIsAlreadyComplete(NodeToDo, LocalOnly)) // if something is already finished, we don't put it into EC  
+                if (GUBPNodes[NodeToDo].RunInEC() && !NodeIsAlreadyComplete(GUBPNodesCompleted, NodeToDo, LocalOnly)) // if something is already finished, we don't put it into EC  
                 {
                     string EMails;
                     var NodeProps = GetECPropsForNode(NodeToDo, CLString, out EMails);
@@ -2603,7 +2602,7 @@ public partial class GUBP : BuildCommand
                         var EcDeps = GetECDependencies(NodeToDo);
                         foreach (var Dep in EcDeps)
                         {
-                            if (GUBPNodes[Dep].RunInEC() && !NodeIsAlreadyComplete(Dep, LocalOnly) && OrdereredToDo.Contains(Dep)) // if something is already finished, we don't put it into EC
+                            if (GUBPNodes[Dep].RunInEC() && !NodeIsAlreadyComplete(GUBPNodesCompleted, Dep, LocalOnly) && OrdereredToDo.Contains(Dep)) // if something is already finished, we don't put it into EC
                             {
                                 if (OrdereredToDo.IndexOf(Dep) > OrdereredToDo.IndexOf(NodeToDo))
                                 {
@@ -2621,7 +2620,7 @@ public partial class GUBP : BuildCommand
                     var CompletedDeps = GetCompletedOnlyDependencies(NodeToDo);
 					foreach (var Dep in CompletedDeps)
 					{
-						if (GUBPNodes[Dep].RunInEC() && !NodeIsAlreadyComplete(Dep, LocalOnly) && OrdereredToDo.Contains(Dep)) // if something is already finished, we don't put it into EC
+						if (GUBPNodes[Dep].RunInEC() && !NodeIsAlreadyComplete(GUBPNodesCompleted, Dep, LocalOnly) && OrdereredToDo.Contains(Dep)) // if something is already finished, we don't put it into EC
 						{
 							if (OrdereredToDo.IndexOf(Dep) > OrdereredToDo.IndexOf(NodeToDo))
 							{
@@ -2652,7 +2651,7 @@ public partial class GUBP : BuildCommand
                                 var EcDeps = GetECDependencies(Chain);
                                 foreach (var Dep in EcDeps)
                                 {
-                                    if (GUBPNodes[Dep].RunInEC() && !NodeIsAlreadyComplete(Dep, LocalOnly) && OrdereredToDo.Contains(Dep)) // if something is already finished, we don't put it into EC
+                                    if (GUBPNodes[Dep].RunInEC() && !NodeIsAlreadyComplete(GUBPNodesCompleted, Dep, LocalOnly) && OrdereredToDo.Contains(Dep)) // if something is already finished, we don't put it into EC
                                     {
                                         if (OrdereredToDo.IndexOf(Dep) > OrdereredToDo.IndexOf(Chain))
                                         {
@@ -2673,7 +2672,7 @@ public partial class GUBP : BuildCommand
                         int MyIndex = MyChain.IndexOf(NodeToDo);
                         if (MyIndex > 0)
                         {
-                            if (!PreConditionUncompletedEcDeps.Contains(MyChain[MyIndex - 1]) && !NodeIsAlreadyComplete(MyChain[MyIndex - 1], LocalOnly))
+                            if (!PreConditionUncompletedEcDeps.Contains(MyChain[MyIndex - 1]) && !NodeIsAlreadyComplete(GUBPNodesCompleted, MyChain[MyIndex - 1], LocalOnly))
                             {
                                 PreConditionUncompletedEcDeps.Add(MyChain[MyIndex - 1]);
                             }
@@ -2683,7 +2682,7 @@ public partial class GUBP : BuildCommand
                             var EcDeps = GetECDependencies(NodeToDo);
                             foreach (var Dep in EcDeps)
                             {
-                                if (GUBPNodes[Dep].RunInEC() && !NodeIsAlreadyComplete(Dep, LocalOnly) && OrdereredToDo.Contains(Dep)) // if something is already finished, we don't put it into EC
+                                if (GUBPNodes[Dep].RunInEC() && !NodeIsAlreadyComplete(GUBPNodesCompleted, Dep, LocalOnly) && OrdereredToDo.Contains(Dep)) // if something is already finished, we don't put it into EC
                                 {
                                     if (OrdereredToDo.IndexOf(Dep) > OrdereredToDo.IndexOf(NodeToDo))
                                     {
@@ -2973,7 +2972,7 @@ public partial class GUBP : BuildCommand
                 && (GUBPNodes[NodeToDo].RunInEC() || !GUBPNodes[NodeToDo].IsAggregate()); //aggregates not in EC can be "run" multiple times, so we can't track those
 
             Log("***** Running GUBP Node {0} -> {1} : {2}", GUBPNodes[NodeToDo].GetFullName(), GameNameIfAny, NodeStoreName);
-            if (NodeIsAlreadyComplete(NodeToDo, LocalOnly))
+            if (NodeIsAlreadyComplete(GUBPNodesCompleted, NodeToDo, LocalOnly))
             {
                 if (NodeToDo == VersionFilesNode.StaticGetFullName() && !IsBuildMachine)
                 {
