Temporarily removed shared feature pack names from CurrentFeaturePacks list

diff --git a/Engine/Source/Programs/AutomationTool/Scripts/RocketBuild.Automation.cs b/Engine/Source/Programs/AutomationTool/Scripts/RocketBuild.Automation.cs
--- a/Engine/Source/Programs/AutomationTool/Scripts/RocketBuild.Automation.cs
+++ b/Engine/Source/Programs/AutomationTool/Scripts/RocketBuild.Automation.cs
@@ -67,23 +67,7 @@ namespace Rocket
 			"TP_VehicleAdvBP",
 			"StarterContent",
 			"MobileStarterContent",
-			"GeometryHigh",		//Shared packs
-			"GeometryStandard",
-			"MannequinHigh",
-			"MannequinStandard",
-			"RollingHigh",
-			"RollingStandard",
-			"2DSideScrollerStandard",
-			"FirstPersonStandard",
-			"FlyingStandard",
-			"PuzzleStandard",
-			"RollingStandard",
-			"RollingHigh",
-			"SideScrollerStandard",
-			"ThirdPersonStandard",
-			"TwinStickStandard",
-			"VehicleStandard",
-			"VehicleAdvStandard",
+
 		};
 
 		public RocketBuild()
