Restore Build Telemetry that was accidentally clobbered.

diff --git a/Engine/Source/Programs/AutomationTool/AutomationUtils/BuildUtils.cs b/Engine/Source/Programs/AutomationTool/AutomationUtils/BuildUtils.cs
--- a/Engine/Source/Programs/AutomationTool/AutomationUtils/BuildUtils.cs
+++ b/Engine/Source/Programs/AutomationTool/AutomationUtils/BuildUtils.cs
@@ -57,7 +57,7 @@ namespace AutomationTool
             var StartBuildSolution = DateTime.Now.ToString();
 			RunAndLog(Env, Env.MsDevExe, CmdLine, LogName);
             var FinishBuildSolution = DateTime.Now.ToString();
-            PrintCSVFile(String.Format("UAT,Compile,{0},{1}", StartBuildSolution, FinishBuildSolution));
+            PrintCSVFile(String.Format("UAT,Compile.{0}.{1}.{2},{3},{4}", Path.GetFileName(SolutionFile), "WinC#", BuildConfig, StartBuildSolution, FinishBuildSolution));
 		}
 
 		/// <summary>
diff --git a/Engine/Source/Programs/AutomationTool/AutomationUtils/UE4Build.cs b/Engine/Source/Programs/AutomationTool/AutomationUtils/UE4Build.cs
--- a/Engine/Source/Programs/AutomationTool/AutomationUtils/UE4Build.cs
+++ b/Engine/Source/Programs/AutomationTool/AutomationUtils/UE4Build.cs
@@ -189,7 +189,10 @@ namespace AutomationTool
 
 			ClearExportedXGEXML();
 
-			RunUBT(CmdEnv, UBTExecutable: UBTExecutable, Project: ProjectName, Target: TargetName, Platform: Platform.ToString(), Config: Config, AdditionalArgs: "-generatemanifest -nobuilduht -xgeexport" + AddArgs, EnvVars: EnvVars);
+            var StartXGEManifestExport = DateTime.Now.ToString();
+            RunUBT(CmdEnv, UBTExecutable: UBTExecutable, Project: ProjectName, Target: TargetName, Platform: Platform.ToString(), Config: Config, AdditionalArgs: "-generatemanifest -nobuilduht -xgeexport" + AddArgs, EnvVars: EnvVars);
+            var FinishXGEManifestExport = DateTime.Now.ToString();
+            PrintCSVFile(String.Format("UAT,GenerateXGEManifest.{0}.{1}.{2},{3},{4}", TargetName, Platform.ToString(), Config, StartXGEManifestExport, FinishXGEManifestExport));
 
 			PrepareManifest(UBTManifest, true);
 
@@ -282,9 +285,11 @@ namespace AutomationTool
 			}
 
 			PrepareUBT();
-
-			RunUBT(CmdEnv, UBTExecutable: UBTExecutable, Project: ProjectName, Target: TargetName, Platform: Platform.ToString(), Config: Config, AdditionalArgs: "-clean" + AddArgs, EnvVars: EnvVars);
-		}
+            var StartCleanUBT = DateTime.Now.ToString();
+            RunUBT(CmdEnv, UBTExecutable: UBTExecutable, Project: ProjectName, Target: TargetName, Platform: Platform.ToString(), Config: Config, AdditionalArgs: "-clean" + AddArgs, EnvVars: EnvVars);
+            var FinishCleanUBT = DateTime.Now.ToString();
+            PrintCSVFile(String.Format("UAT,CleanWithUBT.{0}.{1}.{2},{3},{4}", TargetName, Platform.ToString(), Config, StartCleanUBT, FinishCleanUBT));
+        }
 
 		void BuildWithUBT(string ProjectName, string TargetName, UnrealBuildTool.UnrealTargetPlatform TargetPlatform, string Config, string UprojectPath, bool ForceMonolithic = false, bool ForceNonUnity = false, bool ForceDebugInfo = false, bool ForceFlushMac = false, bool DisableXGE = false, string InAddArgs = "", bool ForceUnity = false, Dictionary<string, string> EnvVars = null)
 		{
@@ -329,15 +334,17 @@ namespace AutomationTool
                 string UBTManifest = GetUBTManifest(UprojectPath, AddArgs);
 
 				DeleteFile(UBTManifest);
-
-				RunUBT(CmdEnv, UBTExecutable: UBTExecutable, Project: ProjectName, Target: TargetName, Platform: TargetPlatform.ToString(), Config: Config, AdditionalArgs: AddArgs + " -generatemanifest" , EnvVars: EnvVars);
-
-				PrepareManifest(UBTManifest, false);
-			}
-
-			RunUBT(CmdEnv, UBTExecutable: UBTExecutable, Project: ProjectName, Target: TargetName, Platform: TargetPlatform.ToString(), Config: Config, AdditionalArgs: AddArgs, EnvVars: EnvVars);
-
-			// allow the platform to perform any actions after building a target (seems almost like this should be done in UBT)
+                var StartGenerateManifest = DateTime.Now.ToString();
+                RunUBT(CmdEnv, UBTExecutable: UBTExecutable, Project: ProjectName, Target: TargetName, Platform: TargetPlatform.ToString(), Config: Config, AdditionalArgs: AddArgs + " -generatemanifest", EnvVars: EnvVars);
+                var FinishGenerateManifest = DateTime.Now.ToString();
+                PrintCSVFile(String.Format("UAT,PrepareUBTManifest.{0}.{1}.{2},{3},{4}", TargetName, TargetPlatform.ToString(), Config, StartGenerateManifest, FinishGenerateManifest));
+                PrepareManifest(UBTManifest, false);
+			}
+            var StartCompile = DateTime.Now.ToString();
+            RunUBT(CmdEnv, UBTExecutable: UBTExecutable, Project: ProjectName, Target: TargetName, Platform: TargetPlatform.ToString(), Config: Config, AdditionalArgs: AddArgs, EnvVars: EnvVars);
+            var FinishCompile = DateTime.Now.ToString();
+            PrintCSVFile(String.Format("UAT,Compile.{0}.{1}.{2},{3},{4}", TargetName, TargetPlatform.ToString(), Config, StartCompile, FinishCompile));
+            // allow the platform to perform any actions after building a target (seems almost like this should be done in UBT)
 			Platform.Platforms[TargetPlatform].PostBuildTarget(this, string.IsNullOrEmpty(ProjectName) ? TargetName : ProjectName, UprojectPath, Config);
 
 			if (UseManifest)
