Separate code to execute nodes into its own function.

diff --git a/Engine/Source/Programs/AutomationTool/GUBP/GUBP.cs b/Engine/Source/Programs/AutomationTool/GUBP/GUBP.cs
--- a/Engine/Source/Programs/AutomationTool/GUBP/GUBP.cs
+++ b/Engine/Source/Programs/AutomationTool/GUBP/GUBP.cs
@@ -2904,6 +2904,11 @@ public partial class GUBP : BuildCommand
             Log("List only, done.");
             return;
         }    
+		ExecuteNodes(OrdereredToDo, bOnlyNode, bFakeEC, LocalOnly, bSaveSharedTempStorage, GUBPNodesCompleted, GUBPNodesHistory, CLString, FakeFail);
+	}
+
+	void ExecuteNodes(List<string> OrdereredToDo, bool bOnlyNode, bool bFakeEC, bool LocalOnly, bool bSaveSharedTempStorage, Dictionary<string, bool> GUBPNodesCompleted, Dictionary<string, NodeHistory> GUBPNodesHistory, string CLString, string FakeFail)
+	{
         var BuildProductToNodeMap = new Dictionary<string, string>();
 		foreach (var NodeToDo in OrdereredToDo)
         {
