Refactor out a few functions from GUBP.DoBuild().

diff --git a/Engine/Source/Programs/AutomationTool/GUBP/GUBP.cs b/Engine/Source/Programs/AutomationTool/GUBP/GUBP.cs
--- a/Engine/Source/Programs/AutomationTool/GUBP/GUBP.cs
+++ b/Engine/Source/Programs/AutomationTool/GUBP/GUBP.cs
@@ -2334,291 +2334,302 @@ public partial class GUBP : BuildCommand
             Log("Took {0}s to get history for {1} nodes", BuildDuration / 1000, NodesToDo.Count);
         }*/
 
-        var OrdereredToDo = TopologicalSort(NodesToDo, GUBPNodesCompleted, GUBPNodesControllingTriggerDotName, GUBPNodesControllingTrigger, ExplicitTrigger, LocalOnly);
+        List<string> OrdereredToDo = TopologicalSort(NodesToDo, GUBPNodesCompleted, GUBPNodesControllingTriggerDotName, GUBPNodesControllingTrigger, ExplicitTrigger, LocalOnly);
 
-        // find all unfinished triggers, excepting the one we are triggering right now
-        var UnfinishedTriggers = new List<string>();
-        if (!bSkipTriggers)
-        {
-            foreach (var NodeToDo in OrdereredToDo)
-            {
-                if (GUBPNodes[NodeToDo].TriggerNode() && !NodeIsAlreadyComplete(GUBPNodesCompleted, NodeToDo, LocalOnly))
-                {
-                    if (String.IsNullOrEmpty(ExplicitTrigger) || ExplicitTrigger != NodeToDo)
-                    {
-                        UnfinishedTriggers.Add(NodeToDo);
-                    }
-                }
-            }
-        }
+		List<string> UnfinishedTriggers = FindUnfinishedTriggers(bSkipTriggers, LocalOnly, ExplicitTrigger, GUBPNodesCompleted, OrdereredToDo);
 
         LogVerbose("*********** Desired And Dependent Nodes, in order.");
         PrintNodes(this, OrdereredToDo, GUBPNodesCompleted, GUBPNodesHistory, GUBPNodesControllingTriggerDotName, GUBPNodesControllingTrigger, LocalOnly, UnfinishedTriggers);		
         //check sorting
-        {
-            foreach (var NodeToDo in OrdereredToDo)
-            {
-                if (GUBPNodes[NodeToDo].TriggerNode() && (GUBPNodes[NodeToDo].IsSticky() || NodeIsAlreadyComplete(GUBPNodesCompleted, NodeToDo, LocalOnly))) // these sticky triggers are ok, everything is already completed anyway
-                {
-                    continue;
-                }
-				if(GUBPNodes[NodeToDo].IsTest())
-				{
-					bHasTests = true;
-				}
-                int MyIndex = OrdereredToDo.IndexOf(NodeToDo);
-                foreach (var Dep in GUBPNodes[NodeToDo].FullNamesOfDependencies)
-                {
-                    int DepIndex = OrdereredToDo.IndexOf(Dep);
-                    if (DepIndex >= MyIndex)
-                    {
-                        throw new AutomationException("Topological sort error, node {0} has a dependency of {1} which sorted after it.", NodeToDo, Dep);
-                    }
-                }
-                foreach (var Dep in GUBPNodes[NodeToDo].FullNamesOfPseudosependencies)
-                {
-                    int DepIndex = OrdereredToDo.IndexOf(Dep);
-                    if (DepIndex >= MyIndex)
-                    {
-                        throw new AutomationException("Topological sort error, node {0} has a pseduodependency of {1} which sorted after it.", NodeToDo, Dep);
-                    }
-                }
-            }
-        }
+		CheckSortOrder(LocalOnly, GUBPNodesCompleted, OrdereredToDo);
 
         string FakeFail = ParseParamValue("FakeFail");
         if (CommanderSetup)
         {
+			DoCommanderSetup(bSkipTriggers, bFakeEC, LocalOnly, CLString, ExplicitTrigger, GUBPNodesControllingTrigger, GUBPNodesControllingTriggerDotName, FullNodeList, FullNodeDirectDependencies, FullNodeDependedOnBy, FullNodeDependentPromotions, SeparatePromotables, FullNodeListSortKey, GUBPNodesCompleted, GUBPNodesHistory, OrdereredToDo, UnfinishedTriggers, FakeFail);
+            Log("Commander setup only, done.");            
+            PrintRunTime();
+            return;
 
-            if (OrdereredToDo.Count == 0)
-            {
-                throw new AutomationException("No nodes to do!");
-            }
-            var ECProps = new List<string>();
-            ECProps.Add(String.Format("TimeIndex={0}", TimeIndex));
-            foreach (var NodePair in FullNodeList)
-            {
-                ECProps.Add(string.Format("AllNodes/{0}={1}", NodePair.Key, NodePair.Value));
-            }
-            foreach (var NodePair in FullNodeDirectDependencies)
-            {
-                ECProps.Add(string.Format("DirectDependencies/{0}={1}", NodePair.Key, NodePair.Value));
-            }
-            foreach (var NodePair in FullNodeListSortKey)
-            {
-                ECProps.Add(string.Format("SortKey/{0}={1}", NodePair.Key, NodePair.Value));
-            }
-			foreach (var NodePair in FullNodeDependedOnBy)
-            {
-				ECProps.Add(string.Format("DependedOnBy/{0}={1}", NodePair.Key, NodePair.Value));
-            }
-			foreach (var NodePair in FullNodeDependentPromotions)
+        }
+        if (ParseParam("SaveGraph"))
+        {
+            SaveGraphVisualization(OrdereredToDo);
+        }
+        if (bListOnly)
+        {
+            Log("List only, done.");
+            return;
+        }    
+		ExecuteNodes(OrdereredToDo, bOnlyNode, bFakeEC, LocalOnly, bSaveSharedTempStorage, GUBPNodesCompleted, GUBPNodesHistory, CLString, FakeFail);
+	}
+
+	private List<string> FindUnfinishedTriggers(bool bSkipTriggers, bool LocalOnly, string ExplicitTrigger, Dictionary<string, bool> GUBPNodesCompleted, List<string> OrdereredToDo)
+	{
+		// find all unfinished triggers, excepting the one we are triggering right now
+		var UnfinishedTriggers = new List<string>();
+		if (!bSkipTriggers)
+		{
+			foreach (var NodeToDo in OrdereredToDo)
+			{
+				if (GUBPNodes[NodeToDo].TriggerNode() && !NodeIsAlreadyComplete(GUBPNodesCompleted, NodeToDo, LocalOnly))
+				{
+					if (String.IsNullOrEmpty(ExplicitTrigger) || ExplicitTrigger != NodeToDo)
+					{
+						UnfinishedTriggers.Add(NodeToDo);
+					}
+				}
+			}
+		}
+		return UnfinishedTriggers;
+	}
+
+	private void CheckSortOrder(bool LocalOnly, Dictionary<string, bool> GUBPNodesCompleted, List<string> OrdereredToDo)
+	{
+		foreach (var NodeToDo in OrdereredToDo)
+		{
+			if (GUBPNodes[NodeToDo].TriggerNode() && (GUBPNodes[NodeToDo].IsSticky() || NodeIsAlreadyComplete(GUBPNodesCompleted, NodeToDo, LocalOnly))) // these sticky triggers are ok, everything is already completed anyway
 			{
-				ECProps.Add(string.Format("DependentPromotions/{0}={1}", NodePair.Key, NodePair.Value));
+				continue;
 			}
-			foreach (var Node in SeparatePromotables)
+			if (GUBPNodes[NodeToDo].IsTest())
 			{
-				ECProps.Add(string.Format("PossiblePromotables/{0}={1}", Node, ""));
-			}			
-            var ECJobProps = new List<string>();
-            if (ExplicitTrigger != "")
-            {
-                ECJobProps.Add("IsRoot=0");
-            }
-            else
-            {
-                ECJobProps.Add("IsRoot=1");
-            }
+				bHasTests = true;
+			}
+			int MyIndex = OrdereredToDo.IndexOf(NodeToDo);
+			foreach (var Dep in GUBPNodes[NodeToDo].FullNamesOfDependencies)
+			{
+				int DepIndex = OrdereredToDo.IndexOf(Dep);
+				if (DepIndex >= MyIndex)
+				{
+					throw new AutomationException("Topological sort error, node {0} has a dependency of {1} which sorted after it.", NodeToDo, Dep);
+				}
+			}
+			foreach (var Dep in GUBPNodes[NodeToDo].FullNamesOfPseudosependencies)
+			{
+				int DepIndex = OrdereredToDo.IndexOf(Dep);
+				if (DepIndex >= MyIndex)
+				{
+					throw new AutomationException("Topological sort error, node {0} has a pseduodependency of {1} which sorted after it.", NodeToDo, Dep);
+				}
+			}
+		}
+	}
 
-            var FilteredOrdereredToDo = new List<string>();
-            var StartFilterTime = DateTime.Now.ToString();
-            // remove nodes that have unfinished triggers
-            foreach (var NodeToDo in OrdereredToDo)
-            {
-                string ControllingTrigger = GetControllingTrigger(NodeToDo, GUBPNodesControllingTrigger);
-                bool bNoUnfinishedTriggers = !UnfinishedTriggers.Contains(ControllingTrigger);
+	private void DoCommanderSetup(bool bSkipTriggers, bool bFakeEC, bool LocalOnly, string CLString, string ExplicitTrigger, Dictionary<string, string> GUBPNodesControllingTrigger, Dictionary<string, string> GUBPNodesControllingTriggerDotName, Dictionary<string, string> FullNodeList, Dictionary<string, string> FullNodeDirectDependencies, Dictionary<string, string> FullNodeDependedOnBy, Dictionary<string, string> FullNodeDependentPromotions, List<string> SeparatePromotables, Dictionary<string, int> FullNodeListSortKey, Dictionary<string, bool> GUBPNodesCompleted, Dictionary<string, NodeHistory> GUBPNodesHistory, List<string> OrdereredToDo, List<string> UnfinishedTriggers, string FakeFail)
+	{
 
-                if (bNoUnfinishedTriggers)
-                {
-                    // if we are triggering, then remove nodes that are not controlled by the trigger or are dependencies of this trigger
-                    if (!String.IsNullOrEmpty(ExplicitTrigger))
-                    {
-                        if (ExplicitTrigger != NodeToDo && !NodeDependsOn(NodeToDo, ExplicitTrigger) && !NodeDependsOn(ExplicitTrigger, NodeToDo))
-                        {
-                            continue; // this wasn't on the chain related to the trigger we are triggering, so it is not relevant
-                        }
-                    }
-					if (bPreflightBuild && !bSkipTriggers && GUBPNodes[NodeToDo].TriggerNode())
+		if (OrdereredToDo.Count == 0)
+		{
+			throw new AutomationException("No nodes to do!");
+		}
+		var ECProps = new List<string>();
+		ECProps.Add(String.Format("TimeIndex={0}", TimeIndex));
+		foreach (var NodePair in FullNodeList)
+		{
+			ECProps.Add(string.Format("AllNodes/{0}={1}", NodePair.Key, NodePair.Value));
+		}
+		foreach (var NodePair in FullNodeDirectDependencies)
+		{
+			ECProps.Add(string.Format("DirectDependencies/{0}={1}", NodePair.Key, NodePair.Value));
+		}
+		foreach (var NodePair in FullNodeListSortKey)
+		{
+			ECProps.Add(string.Format("SortKey/{0}={1}", NodePair.Key, NodePair.Value));
+		}
+		foreach (var NodePair in FullNodeDependedOnBy)
+		{
+			ECProps.Add(string.Format("DependedOnBy/{0}={1}", NodePair.Key, NodePair.Value));
+		}
+		foreach (var NodePair in FullNodeDependentPromotions)
+		{
+			ECProps.Add(string.Format("DependentPromotions/{0}={1}", NodePair.Key, NodePair.Value));
+		}
+		foreach (var Node in SeparatePromotables)
+		{
+			ECProps.Add(string.Format("PossiblePromotables/{0}={1}", Node, ""));
+		}
+		var ECJobProps = new List<string>();
+		if (ExplicitTrigger != "")
+		{
+			ECJobProps.Add("IsRoot=0");
+		}
+		else
+		{
+			ECJobProps.Add("IsRoot=1");
+		}
+
+		var FilteredOrdereredToDo = new List<string>();
+		var StartFilterTime = DateTime.Now.ToString();
+		// remove nodes that have unfinished triggers
+		foreach (var NodeToDo in OrdereredToDo)
+		{
+			string ControllingTrigger = GetControllingTrigger(NodeToDo, GUBPNodesControllingTrigger);
+			bool bNoUnfinishedTriggers = !UnfinishedTriggers.Contains(ControllingTrigger);
+
+			if (bNoUnfinishedTriggers)
+			{
+				// if we are triggering, then remove nodes that are not controlled by the trigger or are dependencies of this trigger
+				if (!String.IsNullOrEmpty(ExplicitTrigger))
+				{
+					if (ExplicitTrigger != NodeToDo && !NodeDependsOn(NodeToDo, ExplicitTrigger) && !NodeDependsOn(ExplicitTrigger, NodeToDo))
 					{
-						// in preflight builds, we are either skipping triggers (and running things downstream) or we just stop at triggers and don't make them available for triggering.
-						continue;
+						continue; // this wasn't on the chain related to the trigger we are triggering, so it is not relevant
 					}
-                    FilteredOrdereredToDo.Add(NodeToDo);
-                }
-            }
-            
-            OrdereredToDo = FilteredOrdereredToDo;
-            var FinishFilterTime = DateTime.Now.ToString();
-            PrintCSVFile(String.Format("UAT,FilterNodes,{0},{1}", StartFilterTime, FinishFilterTime));            
-            Log("*********** EC Nodes, in order.");
-            PrintNodes(this, OrdereredToDo, GUBPNodesCompleted, GUBPNodesHistory, GUBPNodesControllingTriggerDotName, GUBPNodesControllingTrigger, LocalOnly, UnfinishedTriggers);
-            var FinishNodePrint = DateTime.Now.ToString();
-            PrintCSVFile(String.Format("UAT,SetupCommanderPrint,{0},{1}", FinishFilterTime, FinishNodePrint));
-            // here we are just making sure everything before the explicit trigger is completed.
-            if (!String.IsNullOrEmpty(ExplicitTrigger))
-            {
-                foreach (var NodeToDo in FilteredOrdereredToDo)
-                {
-                    if (GUBPNodes[NodeToDo].RunInEC() && !NodeIsAlreadyComplete(GUBPNodesCompleted, NodeToDo, LocalOnly) && NodeToDo != ExplicitTrigger && !NodeDependsOn(ExplicitTrigger, NodeToDo)) // if something is already finished, we don't put it into EC
-                    {
-                        throw new AutomationException("We are being asked to process node {0}, however, this is an explicit trigger {1}, so everything before it should already be handled. It seems likely that you waited too long to run the trigger. You will have to do a new build from scratch.", NodeToDo, ExplicitTrigger);
-                    }
-                }
-            }
+				}
+				if (bPreflightBuild && !bSkipTriggers && GUBPNodes[NodeToDo].TriggerNode())
+				{
+					// in preflight builds, we are either skipping triggers (and running things downstream) or we just stop at triggers and don't make them available for triggering.
+					continue;
+				}
+				FilteredOrdereredToDo.Add(NodeToDo);
+			}
+		}
 
-            string LastSticky = "";
-            bool HitNonSticky = false;
-            bool bHaveECNodes = false;
-            List<string> StepList = new List<string>();
-            StepList.Add("use strict;");
-            StepList.Add("use diagnostics;");
-            StepList.Add("use ElectricCommander();");
-            StepList.Add("my $ec = new ElectricCommander;");
-            StepList.Add("$ec->setTimeout(600);");
-            StepList.Add("my $batch = $ec->newBatch(\"serial\");");
-            // sticky nodes are ones that we run on the main agent. We run then first and they must not be intermixed with parallel jobs
-            foreach (var NodeToDo in OrdereredToDo)
-            {
-                if (GUBPNodes[NodeToDo].RunInEC() && !NodeIsAlreadyComplete(GUBPNodesCompleted, NodeToDo, LocalOnly)) // if something is already finished, we don't put it into EC
-                {
-                    bHaveECNodes = true;
-                    if (GUBPNodes[NodeToDo].IsSticky())
-                    {
-                        LastSticky = NodeToDo;
-                        if (HitNonSticky && !bSkipTriggers)
-                        {
-                            throw new AutomationException("Sticky and non-sticky jobs did not sort right.");
-                        }
-                    }
-                    else
-                    {
-                        HitNonSticky = true;
-                    }
-                }
-            }
-			
-            string ParentPath = ParseParamValue("ParentPath");
-            var StartPerlOutput = DateTime.Now.ToString();
-            string BaseArgs = String.Format("$batch->createJobStep({{parentPath => '{0}'", ParentPath);
+		OrdereredToDo = FilteredOrdereredToDo;
+		var FinishFilterTime = DateTime.Now.ToString();
+		PrintCSVFile(String.Format("UAT,FilterNodes,{0},{1}", StartFilterTime, FinishFilterTime));
+		Log("*********** EC Nodes, in order.");
+		PrintNodes(this, OrdereredToDo, GUBPNodesCompleted, GUBPNodesHistory, GUBPNodesControllingTriggerDotName, GUBPNodesControllingTrigger, LocalOnly, UnfinishedTriggers);
+		var FinishNodePrint = DateTime.Now.ToString();
+		PrintCSVFile(String.Format("UAT,SetupCommanderPrint,{0},{1}", FinishFilterTime, FinishNodePrint));
+		// here we are just making sure everything before the explicit trigger is completed.
+		if (!String.IsNullOrEmpty(ExplicitTrigger))
+		{
+			foreach (var NodeToDo in FilteredOrdereredToDo)
+			{
+				if (GUBPNodes[NodeToDo].RunInEC() && !NodeIsAlreadyComplete(GUBPNodesCompleted, NodeToDo, LocalOnly) && NodeToDo != ExplicitTrigger && !NodeDependsOn(ExplicitTrigger, NodeToDo)) // if something is already finished, we don't put it into EC
+				{
+					throw new AutomationException("We are being asked to process node {0}, however, this is an explicit trigger {1}, so everything before it should already be handled. It seems likely that you waited too long to run the trigger. You will have to do a new build from scratch.", NodeToDo, ExplicitTrigger);
+				}
+			}
+		}
 
-            bool bHasNoop = false;
-            if (LastSticky == "" && bHaveECNodes)
-            {
-                // if we don't have any sticky nodes and we have other nodes, we run a fake noop just to release the resource 
-                string Args = String.Format("{0}, subprocedure => 'GUBP_UAT_Node', parallel => '0', jobStepName => 'Noop', actualParameter => [{{actualParameterName => 'NodeName', value => 'Noop'}}, {{actualParameterName => 'Sticky', value =>'1' }}], releaseMode => 'release'}});", BaseArgs);
-                StepList.Add(Args);
-                bHasNoop = true;
-            }
+		string LastSticky = "";
+		bool HitNonSticky = false;
+		bool bHaveECNodes = false;
+		List<string> StepList = new List<string>();
+		StepList.Add("use strict;");
+		StepList.Add("use diagnostics;");
+		StepList.Add("use ElectricCommander();");
+		StepList.Add("my $ec = new ElectricCommander;");
+		StepList.Add("$ec->setTimeout(600);");
+		StepList.Add("my $batch = $ec->newBatch(\"serial\");");
+		// sticky nodes are ones that we run on the main agent. We run then first and they must not be intermixed with parallel jobs
+		foreach (var NodeToDo in OrdereredToDo)
+		{
+			if (GUBPNodes[NodeToDo].RunInEC() && !NodeIsAlreadyComplete(GUBPNodesCompleted, NodeToDo, LocalOnly)) // if something is already finished, we don't put it into EC
+			{
+				bHaveECNodes = true;
+				if (GUBPNodes[NodeToDo].IsSticky())
+				{
+					LastSticky = NodeToDo;
+					if (HitNonSticky && !bSkipTriggers)
+					{
+						throw new AutomationException("Sticky and non-sticky jobs did not sort right.");
+					}
+				}
+				else
+				{
+					HitNonSticky = true;
+				}
+			}
+		}
 
-            var FakeECArgs = new List<string>();
-            var AgentGroupChains = new Dictionary<string, List<string>>();
-            var StickyChain = new List<string>();
-            foreach (var NodeToDo in OrdereredToDo)
-            {
-                if (GUBPNodes[NodeToDo].RunInEC() && !NodeIsAlreadyComplete(GUBPNodesCompleted, NodeToDo, LocalOnly)) // if something is already finished, we don't put it into EC  
-                {
-                    string MyAgentGroup = GUBPNodes[NodeToDo].AgentSharingGroup;
-                    if (MyAgentGroup != "")
-                    {
-                        if (!AgentGroupChains.ContainsKey(MyAgentGroup))
-                        {
-                            AgentGroupChains.Add(MyAgentGroup, new List<string>{NodeToDo});
-                        }
-                        else
-                        {
-                            AgentGroupChains[MyAgentGroup].Add(NodeToDo);
-                        }
-                    }
-                }
-                if(GUBPNodes[NodeToDo].IsSticky())
-                {
-                    if(!StickyChain.Contains(NodeToDo))
-                    {
-                        StickyChain.Add(NodeToDo);
-                    }
-                }
-            }
-            foreach (var NodeToDo in OrdereredToDo)
-            {
-                if (GUBPNodes[NodeToDo].RunInEC() && !NodeIsAlreadyComplete(GUBPNodesCompleted, NodeToDo, LocalOnly)) // if something is already finished, we don't put it into EC  
-                {
-                    string EMails;
-                    var NodeProps = GetECPropsForNode(NodeToDo, CLString, out EMails);
-                    ECProps.AddRange(NodeProps);
+		string ParentPath = ParseParamValue("ParentPath");
+		var StartPerlOutput = DateTime.Now.ToString();
+		string BaseArgs = String.Format("$batch->createJobStep({{parentPath => '{0}'", ParentPath);
 
-                    bool Sticky = GUBPNodes[NodeToDo].IsSticky();
-                    bool DoParallel = !Sticky;
-                    if(GUBPNodes[NodeToDo].ECProcedure() == "GUBP_UAT_Node_Parallel_AgentShare_Editor")
-                    {
-                        DoParallel = true;
-                    }
-                    if (Sticky && GUBPNodes[NodeToDo].ECAgentString() != "")
-                    {
-                        throw new AutomationException("Node {1} is sticky but has agent requirements.", NodeToDo);
-                    }
-                    string Procedure = GUBPNodes[NodeToDo].ECProcedure();
-                    if (GUBPNodes[NodeToDo].IsSticky() && NodeToDo == LastSticky)
-                    {
-                        Procedure = Procedure + "_Release";
-                    }
-                    string Args = String.Format("{0}, subprocedure => '{1}', parallel => '{2}', jobStepName => '{3}', actualParameter => [{{actualParameterName => 'NodeName', value =>'{4}'}}",
-                        BaseArgs, Procedure, DoParallel ? 1 : 0, NodeToDo, NodeToDo);
-                    string ProcedureParams = GUBPNodes[NodeToDo].ECProcedureParams();
-                    if (!String.IsNullOrEmpty(ProcedureParams))
-                    {
-                        Args = Args + ProcedureParams;
-                    }
+		bool bHasNoop = false;
+		if (LastSticky == "" && bHaveECNodes)
+		{
+			// if we don't have any sticky nodes and we have other nodes, we run a fake noop just to release the resource 
+			string Args = String.Format("{0}, subprocedure => 'GUBP_UAT_Node', parallel => '0', jobStepName => 'Noop', actualParameter => [{{actualParameterName => 'NodeName', value => 'Noop'}}, {{actualParameterName => 'Sticky', value =>'1' }}], releaseMode => 'release'}});", BaseArgs);
+			StepList.Add(Args);
+			bHasNoop = true;
+		}
 
-                    if ((Procedure == "GUBP_UAT_Trigger" || Procedure == "GUBP_Hardcoded_Trigger") && !String.IsNullOrEmpty(EMails))
-                    {
-                        Args = Args + ", {actualParameterName => 'EmailsForTrigger', value => \'" + EMails + "\'}";
-                    }
-                    Args = Args + "]";
-                    string PreCondition = "";
-                    string RunCondition = "";
-                    var UncompletedEcDeps = new List<string>();
-					var UncompletedCompletedDeps = new List<string>();
-                    var PreConditionUncompletedEcDeps = new List<string>();
-                    string MyAgentGroup = GUBPNodes[NodeToDo].AgentSharingGroup;
-                    bool bDoNestedJobstep = false;
-                    bool bDoFirstNestedJobstep = false;
+		var FakeECArgs = new List<string>();
+		var AgentGroupChains = new Dictionary<string, List<string>>();
+		var StickyChain = new List<string>();
+		foreach (var NodeToDo in OrdereredToDo)
+		{
+			if (GUBPNodes[NodeToDo].RunInEC() && !NodeIsAlreadyComplete(GUBPNodesCompleted, NodeToDo, LocalOnly)) // if something is already finished, we don't put it into EC  
+			{
+				string MyAgentGroup = GUBPNodes[NodeToDo].AgentSharingGroup;
+				if (MyAgentGroup != "")
+				{
+					if (!AgentGroupChains.ContainsKey(MyAgentGroup))
+					{
+						AgentGroupChains.Add(MyAgentGroup, new List<string> { NodeToDo });
+					}
+					else
+					{
+						AgentGroupChains[MyAgentGroup].Add(NodeToDo);
+					}
+				}
+			}
+			if (GUBPNodes[NodeToDo].IsSticky())
+			{
+				if (!StickyChain.Contains(NodeToDo))
+				{
+					StickyChain.Add(NodeToDo);
+				}
+			}
+		}
+		foreach (var NodeToDo in OrdereredToDo)
+		{
+			if (GUBPNodes[NodeToDo].RunInEC() && !NodeIsAlreadyComplete(GUBPNodesCompleted, NodeToDo, LocalOnly)) // if something is already finished, we don't put it into EC  
+			{
+				string EMails;
+				var NodeProps = GetECPropsForNode(NodeToDo, CLString, out EMails);
+				ECProps.AddRange(NodeProps);
 
+				bool Sticky = GUBPNodes[NodeToDo].IsSticky();
+				bool DoParallel = !Sticky;
+				if (GUBPNodes[NodeToDo].ECProcedure() == "GUBP_UAT_Node_Parallel_AgentShare_Editor")
+				{
+					DoParallel = true;
+				}
+				if (Sticky && GUBPNodes[NodeToDo].ECAgentString() != "")
+				{
+					throw new AutomationException("Node {1} is sticky but has agent requirements.", NodeToDo);
+				}
+				string Procedure = GUBPNodes[NodeToDo].ECProcedure();
+				if (GUBPNodes[NodeToDo].IsSticky() && NodeToDo == LastSticky)
+				{
+					Procedure = Procedure + "_Release";
+				}
+				string Args = String.Format("{0}, subprocedure => '{1}', parallel => '{2}', jobStepName => '{3}', actualParameter => [{{actualParameterName => 'NodeName', value =>'{4}'}}",
+					BaseArgs, Procedure, DoParallel ? 1 : 0, NodeToDo, NodeToDo);
+				string ProcedureParams = GUBPNodes[NodeToDo].ECProcedureParams();
+				if (!String.IsNullOrEmpty(ProcedureParams))
+				{
+					Args = Args + ProcedureParams;
+				}
 
-                    string NodeParentPath = ParentPath;
-					string PreconditionParentPath;
+				if ((Procedure == "GUBP_UAT_Trigger" || Procedure == "GUBP_Hardcoded_Trigger") && !String.IsNullOrEmpty(EMails))
+				{
+					Args = Args + ", {actualParameterName => 'EmailsForTrigger', value => \'" + EMails + "\'}";
+				}
+				Args = Args + "]";
+				string PreCondition = "";
+				string RunCondition = "";
+				var UncompletedEcDeps = new List<string>();
+				var UncompletedCompletedDeps = new List<string>();
+				var PreConditionUncompletedEcDeps = new List<string>();
+				string MyAgentGroup = GUBPNodes[NodeToDo].AgentSharingGroup;
+				bool bDoNestedJobstep = false;
+				bool bDoFirstNestedJobstep = false;
 
-                    PreconditionParentPath = ParentPath;
-					{
-                        var EcDeps = GetECDependencies(NodeToDo);
-                        foreach (var Dep in EcDeps)
-                        {
-                            if (GUBPNodes[Dep].RunInEC() && !NodeIsAlreadyComplete(GUBPNodesCompleted, Dep, LocalOnly) && OrdereredToDo.Contains(Dep)) // if something is already finished, we don't put it into EC
-                            {
-                                if (OrdereredToDo.IndexOf(Dep) > OrdereredToDo.IndexOf(NodeToDo))
-                                {
-                                    throw new AutomationException("Topological sort error, node {0} has a dependency of {1} which sorted after it.", NodeToDo, Dep);
-                                }
-                                UncompletedEcDeps.Add(Dep);
-                            }
-                        }
-                    }
 
-                    foreach (var Dep in UncompletedEcDeps)
-                    {
-                        PreConditionUncompletedEcDeps.Add(Dep);
-                    }
-                    var CompletedDeps = GetCompletedOnlyDependencies(NodeToDo);
-					foreach (var Dep in CompletedDeps)
+				string NodeParentPath = ParentPath;
+				string PreconditionParentPath;
+
+				PreconditionParentPath = ParentPath;
+				{
+					var EcDeps = GetECDependencies(NodeToDo);
+					foreach (var Dep in EcDeps)
 					{
 						if (GUBPNodes[Dep].RunInEC() && !NodeIsAlreadyComplete(GUBPNodesCompleted, Dep, LocalOnly) && OrdereredToDo.Contains(Dep)) // if something is already finished, we don't put it into EC
 						{
@@ -2626,190 +2637,208 @@ public partial class GUBP : BuildCommand
 							{
 								throw new AutomationException("Topological sort error, node {0} has a dependency of {1} which sorted after it.", NodeToDo, Dep);
 							}
-							UncompletedCompletedDeps.Add(Dep);
+							UncompletedEcDeps.Add(Dep);
 						}
 					}
-                    if (MyAgentGroup != "")
-                    {
-                        bDoNestedJobstep = true;
-                        NodeParentPath = ParentPath + "/jobSteps[" + MyAgentGroup + "]";
+				}
 
-                        PreConditionUncompletedEcDeps = new List<string>();
+				foreach (var Dep in UncompletedEcDeps)
+				{
+					PreConditionUncompletedEcDeps.Add(Dep);
+				}
+				var CompletedDeps = GetCompletedOnlyDependencies(NodeToDo);
+				foreach (var Dep in CompletedDeps)
+				{
+					if (GUBPNodes[Dep].RunInEC() && !NodeIsAlreadyComplete(GUBPNodesCompleted, Dep, LocalOnly) && OrdereredToDo.Contains(Dep)) // if something is already finished, we don't put it into EC
+					{
+						if (OrdereredToDo.IndexOf(Dep) > OrdereredToDo.IndexOf(NodeToDo))
+						{
+							throw new AutomationException("Topological sort error, node {0} has a dependency of {1} which sorted after it.", NodeToDo, Dep);
+						}
+						UncompletedCompletedDeps.Add(Dep);
+					}
+				}
+				if (MyAgentGroup != "")
+				{
+					bDoNestedJobstep = true;
+					NodeParentPath = ParentPath + "/jobSteps[" + MyAgentGroup + "]";
 
-                        var MyChain = AgentGroupChains[MyAgentGroup];
-                        int MyIndex = MyChain.IndexOf(NodeToDo);
-                        if (MyIndex  > 0)
-                        {
-                            PreConditionUncompletedEcDeps.Add(MyChain[MyIndex - 1]);
-                        }
-                        else
-                        {
-                            bDoFirstNestedJobstep = bDoNestedJobstep;
-                            // to avoid idle agents (and also EC doesn't actually reserve our agent!), we promote all dependencies to the first one
-                            foreach (var Chain in MyChain)
-                            {
-                                var EcDeps = GetECDependencies(Chain);
-                                foreach (var Dep in EcDeps)
-                                {
-                                    if (GUBPNodes[Dep].RunInEC() && !NodeIsAlreadyComplete(GUBPNodesCompleted, Dep, LocalOnly) && OrdereredToDo.Contains(Dep)) // if something is already finished, we don't put it into EC
-                                    {
-                                        if (OrdereredToDo.IndexOf(Dep) > OrdereredToDo.IndexOf(Chain))
-                                        {
-                                            throw new AutomationException("Topological sort error, node {0} has a dependency of {1} which sorted after it.", Chain, Dep);
-                                        }
-                                        if (!MyChain.Contains(Dep) && !PreConditionUncompletedEcDeps.Contains(Dep))
-                                        {
-                                            PreConditionUncompletedEcDeps.Add(Dep);
-                                        }
-                                    }
-                                }
-                            }
-                        }
-                    }
-                    if (GUBPNodes[NodeToDo].IsSticky())
-                    {
-                        var MyChain = StickyChain;
-                        int MyIndex = MyChain.IndexOf(NodeToDo);
-                        if (MyIndex > 0)
-                        {
-                            if (!PreConditionUncompletedEcDeps.Contains(MyChain[MyIndex - 1]) && !NodeIsAlreadyComplete(GUBPNodesCompleted, MyChain[MyIndex - 1], LocalOnly))
-                            {
-                                PreConditionUncompletedEcDeps.Add(MyChain[MyIndex - 1]);
-                            }
-                        }
-                        else
-                        {
-                            var EcDeps = GetECDependencies(NodeToDo);
-                            foreach (var Dep in EcDeps)
-                            {
-                                if (GUBPNodes[Dep].RunInEC() && !NodeIsAlreadyComplete(GUBPNodesCompleted, Dep, LocalOnly) && OrdereredToDo.Contains(Dep)) // if something is already finished, we don't put it into EC
-                                {
-                                    if (OrdereredToDo.IndexOf(Dep) > OrdereredToDo.IndexOf(NodeToDo))
-                                    {
-                                        throw new AutomationException("Topological sort error, node {0} has a dependency of {1} which sorted after it.", NodeToDo, Dep);
-                                    }
-                                    if (!MyChain.Contains(Dep) && !PreConditionUncompletedEcDeps.Contains(Dep))
-                                    {
-                                        PreConditionUncompletedEcDeps.Add(Dep);
-                                    }
-                                }
-                            }
-                        }
-                    }
-                    if (bHasNoop && PreConditionUncompletedEcDeps.Count == 0)
-                    {
-                        PreConditionUncompletedEcDeps.Add("Noop");
-                    }
-                    if (PreConditionUncompletedEcDeps.Count > 0)
-                    {
-                        PreCondition = "\"\\$\" . \"[/javascript if(";
-                        // these run "parallel", but we add preconditions to serialize them
-                        int Index = 0;
-                        foreach (var Dep in PreConditionUncompletedEcDeps)
-                        {
-                            PreCondition = PreCondition + "getProperty('" + GetJobStep(PreconditionParentPath, Dep) + "/status\') == \'completed\'";
-                            Index++;
-                            if (Index != PreConditionUncompletedEcDeps.Count)
-                            {
-                                PreCondition = PreCondition + " && ";
-                            }
-                        }
-                        PreCondition = PreCondition + ") true;]\"";
-                    }
-                    if (UncompletedCompletedDeps.Count > 0)
+					PreConditionUncompletedEcDeps = new List<string>();
+
+					var MyChain = AgentGroupChains[MyAgentGroup];
+					int MyIndex = MyChain.IndexOf(NodeToDo);
+					if (MyIndex > 0)
+					{
+						PreConditionUncompletedEcDeps.Add(MyChain[MyIndex - 1]);
+					}
+					else
+					{
+						bDoFirstNestedJobstep = bDoNestedJobstep;
+						// to avoid idle agents (and also EC doesn't actually reserve our agent!), we promote all dependencies to the first one
+						foreach (var Chain in MyChain)
+						{
+							var EcDeps = GetECDependencies(Chain);
+							foreach (var Dep in EcDeps)
+							{
+								if (GUBPNodes[Dep].RunInEC() && !NodeIsAlreadyComplete(GUBPNodesCompleted, Dep, LocalOnly) && OrdereredToDo.Contains(Dep)) // if something is already finished, we don't put it into EC
+								{
+									if (OrdereredToDo.IndexOf(Dep) > OrdereredToDo.IndexOf(Chain))
+									{
+										throw new AutomationException("Topological sort error, node {0} has a dependency of {1} which sorted after it.", Chain, Dep);
+									}
+									if (!MyChain.Contains(Dep) && !PreConditionUncompletedEcDeps.Contains(Dep))
+									{
+										PreConditionUncompletedEcDeps.Add(Dep);
+									}
+								}
+							}
+						}
+					}
+				}
+				if (GUBPNodes[NodeToDo].IsSticky())
+				{
+					var MyChain = StickyChain;
+					int MyIndex = MyChain.IndexOf(NodeToDo);
+					if (MyIndex > 0)
+					{
+						if (!PreConditionUncompletedEcDeps.Contains(MyChain[MyIndex - 1]) && !NodeIsAlreadyComplete(GUBPNodesCompleted, MyChain[MyIndex - 1], LocalOnly))
+						{
+							PreConditionUncompletedEcDeps.Add(MyChain[MyIndex - 1]);
+						}
+					}
+					else
 					{
-						PreCondition = "\"\\$\" . \"[/javascript if(";
-						// these run "parallel", but we add preconditions to serialize them
-						int Index = 0;
-						foreach (var Dep in CompletedDeps)
+						var EcDeps = GetECDependencies(NodeToDo);
+						foreach (var Dep in EcDeps)
 						{
-							if (GUBPNodes[Dep].RunInEC())
+							if (GUBPNodes[Dep].RunInEC() && !NodeIsAlreadyComplete(GUBPNodesCompleted, Dep, LocalOnly) && OrdereredToDo.Contains(Dep)) // if something is already finished, we don't put it into EC
 							{
-								PreCondition = PreCondition + "getProperty('" + GetJobStep(PreconditionParentPath, Dep) + "/status\') == \'completed\'";
-								Index++;
-								if (Index != CompletedDeps.Count)
+								if (OrdereredToDo.IndexOf(Dep) > OrdereredToDo.IndexOf(NodeToDo))
+								{
+									throw new AutomationException("Topological sort error, node {0} has a dependency of {1} which sorted after it.", NodeToDo, Dep);
+								}
+								if (!MyChain.Contains(Dep) && !PreConditionUncompletedEcDeps.Contains(Dep))
 								{
-									PreCondition = PreCondition + " && ";
+									PreConditionUncompletedEcDeps.Add(Dep);
 								}
 							}
 						}
-						PreCondition = PreCondition + ") true;]\"";
 					}
+				}
+				if (bHasNoop && PreConditionUncompletedEcDeps.Count == 0)
+				{
+					PreConditionUncompletedEcDeps.Add("Noop");
+				}
+				if (PreConditionUncompletedEcDeps.Count > 0)
+				{
+					PreCondition = "\"\\$\" . \"[/javascript if(";
+					// these run "parallel", but we add preconditions to serialize them
+					int Index = 0;
+					foreach (var Dep in PreConditionUncompletedEcDeps)
+					{
+						PreCondition = PreCondition + "getProperty('" + GetJobStep(PreconditionParentPath, Dep) + "/status\') == \'completed\'";
+						Index++;
+						if (Index != PreConditionUncompletedEcDeps.Count)
+						{
+							PreCondition = PreCondition + " && ";
+						}
+					}
+					PreCondition = PreCondition + ") true;]\"";
+				}
+				if (UncompletedCompletedDeps.Count > 0)
+				{
+					PreCondition = "\"\\$\" . \"[/javascript if(";
+					// these run "parallel", but we add preconditions to serialize them
+					int Index = 0;
+					foreach (var Dep in CompletedDeps)
+					{
+						if (GUBPNodes[Dep].RunInEC())
+						{
+							PreCondition = PreCondition + "getProperty('" + GetJobStep(PreconditionParentPath, Dep) + "/status\') == \'completed\'";
+							Index++;
+							if (Index != CompletedDeps.Count)
+							{
+								PreCondition = PreCondition + " && ";
+							}
+						}
+					}
+					PreCondition = PreCondition + ") true;]\"";
+				}
 
-                    if (UncompletedEcDeps.Count > 0)
-                    {
-                        RunCondition = "\"\\$\" . \"[/javascript if(";
-                        int Index = 0;
-                        foreach (var Dep in UncompletedEcDeps)
-                        {
-							RunCondition = RunCondition + "((\'\\$\" . \"[" + GetJobStep(PreconditionParentPath, Dep) + "/outcome]\' == \'success\') || ";
-							RunCondition = RunCondition + "(\'\\$\" . \"[" + GetJobStep(PreconditionParentPath, Dep) + "/outcome]\' == \'warning\'))";
+				if (UncompletedEcDeps.Count > 0)
+				{
+					RunCondition = "\"\\$\" . \"[/javascript if(";
+					int Index = 0;
+					foreach (var Dep in UncompletedEcDeps)
+					{
+						RunCondition = RunCondition + "((\'\\$\" . \"[" + GetJobStep(PreconditionParentPath, Dep) + "/outcome]\' == \'success\') || ";
+						RunCondition = RunCondition + "(\'\\$\" . \"[" + GetJobStep(PreconditionParentPath, Dep) + "/outcome]\' == \'warning\'))";
 
-                            Index++;
-                            if (Index != UncompletedEcDeps.Count)
-                            {
-                                RunCondition = RunCondition + " && ";
-                            }
-                        }
-                        RunCondition = RunCondition + ")true; else false;]\"";
-                    }
+						Index++;
+						if (Index != UncompletedEcDeps.Count)
+						{
+							RunCondition = RunCondition + " && ";
+						}
+					}
+					RunCondition = RunCondition + ")true; else false;]\"";
+				}
 
-                    if (bDoNestedJobstep)
-                    {
-                        if (bDoFirstNestedJobstep)
-                        {
-                            {
-                                string NestArgs = String.Format("$batch->createJobStep({{parentPath => '{0}', jobStepName => '{1}', parallel => '1'",
-                                    ParentPath, MyAgentGroup);
-                                if (!String.IsNullOrEmpty(PreCondition))
-                                {
-                                    NestArgs = NestArgs + ", precondition => " + PreCondition;
-                                }
-                                NestArgs = NestArgs + "});";
-                                StepList.Add(NestArgs);
-                            }
-                            {
-                                string NestArgs = String.Format("$batch->createJobStep({{parentPath => '{0}/jobSteps[{1}]', jobStepName => '{2}_GetPool', subprocedure => 'GUBP{3}_AgentShare_GetPool', parallel => '1', actualParameter => [{{actualParameterName => 'AgentSharingGroup', value => '{4}'}}, {{actualParameterName => 'NodeName', value => '{5}'}}]",
-                                    ParentPath, MyAgentGroup, MyAgentGroup, GUBPNodes[NodeToDo].ECProcedureInfix(), MyAgentGroup, NodeToDo);
-                                if (!String.IsNullOrEmpty(PreCondition))
-                                {
-                                    NestArgs = NestArgs + ", precondition => " + PreCondition;
-                                }
-                                NestArgs = NestArgs + "});";
-                                StepList.Add(NestArgs);
-                            }
-                            {
-                                string NestArgs = String.Format("$batch->createJobStep({{parentPath => '{0}/jobSteps[{1}]', jobStepName => '{2}_GetAgent', subprocedure => 'GUBP{3}_AgentShare_GetAgent', parallel => '1', exclusiveMode => 'call', resourceName => '{4}', actualParameter => [{{actualParameterName => 'AgentSharingGroup', value => '{5}'}}, {{actualParameterName => 'NodeName', value=> '{6}'}}]",
-                                    ParentPath, MyAgentGroup, MyAgentGroup, GUBPNodes[NodeToDo].ECProcedureInfix(),
-                                    String.Format("$[/myJob/jobSteps[{0}]/ResourcePool]", MyAgentGroup),
-                                    MyAgentGroup, NodeToDo);
-                                {
-                                    NestArgs = NestArgs + ", precondition  => ";
-                                    NestArgs = NestArgs + "\"\\$\" . \"[/javascript if(";
-                                    NestArgs = NestArgs + "getProperty('" + PreconditionParentPath + "/jobSteps[" + MyAgentGroup + "]/jobSteps[" + MyAgentGroup + "_GetPool]/status') == 'completed'";
-                                    NestArgs = NestArgs + ") true;]\"";
-                                }
-                                NestArgs = NestArgs + "});";
-                                StepList.Add(NestArgs);
-                            }
-                            {
-                                PreCondition = "\"\\$\" . \"[/javascript if(";
-                                PreCondition = PreCondition + "getProperty('" + PreconditionParentPath + "/jobSteps[" + MyAgentGroup + "]/jobSteps[" + MyAgentGroup + "_GetAgent]/status') == 'completed'";
-                                PreCondition = PreCondition + ") true;]\"";
-                            }
-                        }
-                        Args = Args.Replace(String.Format("parentPath => '{0}'", ParentPath), String.Format("parentPath => '{0}'", NodeParentPath));
-                        Args = Args.Replace("UAT_Node_Parallel_AgentShare", "UAT_Node_Parallel_AgentShare3");
-                    }
+				if (bDoNestedJobstep)
+				{
+					if (bDoFirstNestedJobstep)
+					{
+						{
+							string NestArgs = String.Format("$batch->createJobStep({{parentPath => '{0}', jobStepName => '{1}', parallel => '1'",
+								ParentPath, MyAgentGroup);
+							if (!String.IsNullOrEmpty(PreCondition))
+							{
+								NestArgs = NestArgs + ", precondition => " + PreCondition;
+							}
+							NestArgs = NestArgs + "});";
+							StepList.Add(NestArgs);
+						}
+						{
+							string NestArgs = String.Format("$batch->createJobStep({{parentPath => '{0}/jobSteps[{1}]', jobStepName => '{2}_GetPool', subprocedure => 'GUBP{3}_AgentShare_GetPool', parallel => '1', actualParameter => [{{actualParameterName => 'AgentSharingGroup', value => '{4}'}}, {{actualParameterName => 'NodeName', value => '{5}'}}]",
+								ParentPath, MyAgentGroup, MyAgentGroup, GUBPNodes[NodeToDo].ECProcedureInfix(), MyAgentGroup, NodeToDo);
+							if (!String.IsNullOrEmpty(PreCondition))
+							{
+								NestArgs = NestArgs + ", precondition => " + PreCondition;
+							}
+							NestArgs = NestArgs + "});";
+							StepList.Add(NestArgs);
+						}
+						{
+							string NestArgs = String.Format("$batch->createJobStep({{parentPath => '{0}/jobSteps[{1}]', jobStepName => '{2}_GetAgent', subprocedure => 'GUBP{3}_AgentShare_GetAgent', parallel => '1', exclusiveMode => 'call', resourceName => '{4}', actualParameter => [{{actualParameterName => 'AgentSharingGroup', value => '{5}'}}, {{actualParameterName => 'NodeName', value=> '{6}'}}]",
+								ParentPath, MyAgentGroup, MyAgentGroup, GUBPNodes[NodeToDo].ECProcedureInfix(),
+								String.Format("$[/myJob/jobSteps[{0}]/ResourcePool]", MyAgentGroup),
+								MyAgentGroup, NodeToDo);
+							{
+								NestArgs = NestArgs + ", precondition  => ";
+								NestArgs = NestArgs + "\"\\$\" . \"[/javascript if(";
+								NestArgs = NestArgs + "getProperty('" + PreconditionParentPath + "/jobSteps[" + MyAgentGroup + "]/jobSteps[" + MyAgentGroup + "_GetPool]/status') == 'completed'";
+								NestArgs = NestArgs + ") true;]\"";
+							}
+							NestArgs = NestArgs + "});";
+							StepList.Add(NestArgs);
+						}
+						{
+							PreCondition = "\"\\$\" . \"[/javascript if(";
+							PreCondition = PreCondition + "getProperty('" + PreconditionParentPath + "/jobSteps[" + MyAgentGroup + "]/jobSteps[" + MyAgentGroup + "_GetAgent]/status') == 'completed'";
+							PreCondition = PreCondition + ") true;]\"";
+						}
+					}
+					Args = Args.Replace(String.Format("parentPath => '{0}'", ParentPath), String.Format("parentPath => '{0}'", NodeParentPath));
+					Args = Args.Replace("UAT_Node_Parallel_AgentShare", "UAT_Node_Parallel_AgentShare3");
+				}
 
-                    if (!String.IsNullOrEmpty(PreCondition))
-                    {
-                        Args = Args + ", precondition => " + PreCondition;
-                    }
-                    if (!String.IsNullOrEmpty(RunCondition))
-                    {
-                        Args = Args + ", condition => " + RunCondition;
-                    }
+				if (!String.IsNullOrEmpty(PreCondition))
+				{
+					Args = Args + ", precondition => " + PreCondition;
+				}
+				if (!String.IsNullOrEmpty(RunCondition))
+				{
+					Args = Args + ", condition => " + RunCondition;
+				}
 #if false
                     // this doesn't work because it includes precondition time
                     if (GUBPNodes[NodeToDo].TimeoutInMinutes() > 0)
@@ -2817,94 +2846,79 @@ public partial class GUBP : BuildCommand
                         Args = Args + String.Format(" --timeLimitUnits minutes --timeLimit {0}", GUBPNodes[NodeToDo].TimeoutInMinutes());
                     }
 #endif
-                    if (Sticky && NodeToDo == LastSticky)
-                    {
-                        Args = Args + ", releaseMode => 'release'";
-                    }
-                    Args = Args + "});";
-                    StepList.Add(Args);
-                    if (bFakeEC &&
-                        !UnfinishedTriggers.Contains(NodeToDo) &&
-                        (GUBPNodes[NodeToDo].ECProcedure().StartsWith("GUBP_UAT_Node") || GUBPNodes[NodeToDo].ECProcedure().StartsWith("GUBP_Mac_UAT_Node")) // other things we really can't test
-                        ) // unfinished triggers are never run directly by EC, rather it does another job setup
-                    {
-                        string Arg = String.Format("gubp -Node={0} -FakeEC {1} {2} {3} {4} {5}",
-                            NodeToDo,
-                            bFake ? "-Fake" : "",
-                            ParseParam("AllPlatforms") ? "-AllPlatforms" : "",
-                            ParseParam("UnfinishedTriggersFirst") ? "-UnfinishedTriggersFirst" : "",
-                            ParseParam("UnfinishedTriggersParallel") ? "-UnfinishedTriggersParallel" : "",
-                            ParseParam("WithMac") ? "-WithMac" : ""
-                            );
-
-                        string Node = ParseParamValue("-Node");
-                        if (!String.IsNullOrEmpty(Node))
-                        {
-                            Arg = Arg + " -Node=" + Node;
-                        }
-                        if (!String.IsNullOrEmpty(FakeFail))
-                        {
-                            Arg = Arg + " -FakeFail=" + FakeFail;
-                        }
-                        FakeECArgs.Add(Arg);
-                    }
-
-                    if (MyAgentGroup != "" && !bDoNestedJobstep)
-                    {
-                        var MyChain = AgentGroupChains[MyAgentGroup];
-                        int MyIndex = MyChain.IndexOf(NodeToDo);
-                        if (MyIndex == MyChain.Count - 1)
-                        {
-                            var RelPreCondition = "\"\\$\" . \"[/javascript if(";
-                            // this runs "parallel", but we a precondition to serialize it
-                            RelPreCondition = RelPreCondition + "getProperty('" + PreconditionParentPath + "/jobSteps[" + NodeToDo + "]/status') == 'completed'";
-                            RelPreCondition = RelPreCondition + ") true;]\"";
-                            // we need to release the resource
-                            string RelArgs = String.Format("{0}, subprocedure => 'GUBP_Release_AgentShare', parallel => '1', jobStepName => 'Release_{1}', actualParameter => [{{actualParameterName => 'AgentSharingGroup', valued => '{2}'}}], releaseMode => 'release', precondition => '{3}'",
-                                BaseArgs, MyAgentGroup, MyAgentGroup, RelPreCondition);
-                            StepList.Add(RelArgs);
-                        }
-                    }
-                }
-            }
-            WriteECPerl(StepList);
-            var FinishPerlOutput = DateTime.Now.ToString();
-            PrintCSVFile(String.Format("UAT,PerlOutput,{0},{1}", StartPerlOutput, FinishPerlOutput));
-            RunECTool(String.Format("setProperty \"/myWorkflow/HasTests\" \"{0}\"", bHasTests));
-            {
-                ECProps.Add("GUBP_LoadedProps=1");
-                string BranchDefFile = CommandUtils.CombinePaths(CommandUtils.CmdEnv.LogFolder, "BranchDef.properties");
-                CommandUtils.WriteAllLines(BranchDefFile, ECProps.ToArray());
-                RunECTool(String.Format("setProperty \"/myWorkflow/BranchDefFile\" \"{0}\"", BranchDefFile.Replace("\\", "\\\\")));
-            }
-            {
-                ECProps.Add("GUBP_LoadedJobProps=1");
-                string BranchJobDefFile = CommandUtils.CombinePaths(CommandUtils.CmdEnv.LogFolder, "BranchJobDef.properties");
-                CommandUtils.WriteAllLines(BranchJobDefFile, ECProps.ToArray());
-                RunECTool(String.Format("setProperty \"/myJob/BranchJobDefFile\" \"{0}\"", BranchJobDefFile.Replace("\\", "\\\\")));
-            }
-            if (bFakeEC)
-            {
-                foreach (var Args in FakeECArgs)
-                {
-                    RunUAT(CmdEnv, Args);
-                }
-            }
-            Log("Commander setup only, done.");            
-            PrintRunTime();
-            return;
+				if (Sticky && NodeToDo == LastSticky)
+				{
+					Args = Args + ", releaseMode => 'release'";
+				}
+				Args = Args + "});";
+				StepList.Add(Args);
+				if (bFakeEC &&
+					!UnfinishedTriggers.Contains(NodeToDo) &&
+					(GUBPNodes[NodeToDo].ECProcedure().StartsWith("GUBP_UAT_Node") || GUBPNodes[NodeToDo].ECProcedure().StartsWith("GUBP_Mac_UAT_Node")) // other things we really can't test
+					) // unfinished triggers are never run directly by EC, rather it does another job setup
+				{
+					string Arg = String.Format("gubp -Node={0} -FakeEC {1} {2} {3} {4} {5}",
+						NodeToDo,
+						bFake ? "-Fake" : "",
+						ParseParam("AllPlatforms") ? "-AllPlatforms" : "",
+						ParseParam("UnfinishedTriggersFirst") ? "-UnfinishedTriggersFirst" : "",
+						ParseParam("UnfinishedTriggersParallel") ? "-UnfinishedTriggersParallel" : "",
+						ParseParam("WithMac") ? "-WithMac" : ""
+						);
+
+					string Node = ParseParamValue("-Node");
+					if (!String.IsNullOrEmpty(Node))
+					{
+						Arg = Arg + " -Node=" + Node;
+					}
+					if (!String.IsNullOrEmpty(FakeFail))
+					{
+						Arg = Arg + " -FakeFail=" + FakeFail;
+					}
+					FakeECArgs.Add(Arg);
+				}
 
-        }
-        if (ParseParam("SaveGraph"))
-        {
-            SaveGraphVisualization(OrdereredToDo);
-        }
-        if (bListOnly)
-        {
-            Log("List only, done.");
-            return;
-        }    
-		ExecuteNodes(OrdereredToDo, bOnlyNode, bFakeEC, LocalOnly, bSaveSharedTempStorage, GUBPNodesCompleted, GUBPNodesHistory, CLString, FakeFail);
+				if (MyAgentGroup != "" && !bDoNestedJobstep)
+				{
+					var MyChain = AgentGroupChains[MyAgentGroup];
+					int MyIndex = MyChain.IndexOf(NodeToDo);
+					if (MyIndex == MyChain.Count - 1)
+					{
+						var RelPreCondition = "\"\\$\" . \"[/javascript if(";
+						// this runs "parallel", but we a precondition to serialize it
+						RelPreCondition = RelPreCondition + "getProperty('" + PreconditionParentPath + "/jobSteps[" + NodeToDo + "]/status') == 'completed'";
+						RelPreCondition = RelPreCondition + ") true;]\"";
+						// we need to release the resource
+						string RelArgs = String.Format("{0}, subprocedure => 'GUBP_Release_AgentShare', parallel => '1', jobStepName => 'Release_{1}', actualParameter => [{{actualParameterName => 'AgentSharingGroup', valued => '{2}'}}], releaseMode => 'release', precondition => '{3}'",
+							BaseArgs, MyAgentGroup, MyAgentGroup, RelPreCondition);
+						StepList.Add(RelArgs);
+					}
+				}
+			}
+		}
+		WriteECPerl(StepList);
+		var FinishPerlOutput = DateTime.Now.ToString();
+		PrintCSVFile(String.Format("UAT,PerlOutput,{0},{1}", StartPerlOutput, FinishPerlOutput));
+		RunECTool(String.Format("setProperty \"/myWorkflow/HasTests\" \"{0}\"", bHasTests));
+		{
+			ECProps.Add("GUBP_LoadedProps=1");
+			string BranchDefFile = CommandUtils.CombinePaths(CommandUtils.CmdEnv.LogFolder, "BranchDef.properties");
+			CommandUtils.WriteAllLines(BranchDefFile, ECProps.ToArray());
+			RunECTool(String.Format("setProperty \"/myWorkflow/BranchDefFile\" \"{0}\"", BranchDefFile.Replace("\\", "\\\\")));
+		}
+		{
+			ECProps.Add("GUBP_LoadedJobProps=1");
+			string BranchJobDefFile = CommandUtils.CombinePaths(CommandUtils.CmdEnv.LogFolder, "BranchJobDef.properties");
+			CommandUtils.WriteAllLines(BranchJobDefFile, ECProps.ToArray());
+			RunECTool(String.Format("setProperty \"/myJob/BranchJobDefFile\" \"{0}\"", BranchJobDefFile.Replace("\\", "\\\\")));
+		}
+		if (bFakeEC)
+		{
+			foreach (var Args in FakeECArgs)
+			{
+				RunUAT(CmdEnv, Args);
+			}
+		}
 	}
 
 	void ExecuteNodes(List<string> OrdereredToDo, bool bOnlyNode, bool bFakeEC, bool LocalOnly, bool bSaveSharedTempStorage, Dictionary<string, bool> GUBPNodesCompleted, Dictionary<string, NodeHistory> GUBPNodesHistory, string CLString, string FakeFail)
