Move the CIS counter update into its own function.

diff --git a/Engine/Source/Programs/AutomationTool/GUBP/GUBP.cs b/Engine/Source/Programs/AutomationTool/GUBP/GUBP.cs
--- a/Engine/Source/Programs/AutomationTool/GUBP/GUBP.cs
+++ b/Engine/Source/Programs/AutomationTool/GUBP/GUBP.cs
@@ -14,7 +14,6 @@ public partial class GUBP : BuildCommand
     public string StoreName = null;
 	public string BranchName;
     public int CL = 0;
-    public int TimeIndex = 0;
     public bool bSignBuildProducts = false;
     public bool bHasTests = false;
     public List<UnrealTargetPlatform> ActivePlatforms = null;
@@ -1489,7 +1488,7 @@ public partial class GUBP : BuildCommand
         bool bSkipTriggers = ParseParam("SkipTriggers");
         bFake = ParseParam("fake");
         bool bFakeEC = ParseParam("FakeEC");
-        TimeIndex = ParseParamInt("TimeIndex", 0);
+        int TimeIndex = ParseParamInt("TimeIndex", 0);
         if (TimeIndex == 0)
         {
             TimeIndex = ParseParamInt("UserTimeIndex", 0);
@@ -1556,126 +1555,7 @@ public partial class GUBP : BuildCommand
 
         if (ParseParam("CIS") && ExplicitTrigger == "" && CommanderSetup) // explicit triggers will already have a time index assigned
         {
-            if (!P4Enabled)
-            {
-                throw new AutomationException("Can't have -CIS without P4 support");
-            }
-            var P4IndexFileP4 = CombinePaths(PathSeparator.Slash, CommandUtils.P4Env.BuildRootP4, "Engine", "Build", "CISCounter.txt");
-            var P4IndexFileLocal = CombinePaths(CmdEnv.LocalRoot, "Engine", "Build", "CISCounter.txt");
-            int Retry = 0;
-            bool bDone = false;
-            while (++Retry < 20 && !bDone)
-            {
-                int NowMinutes = (int)((DateTime.UtcNow - new DateTime(2014, 1, 1)).TotalMinutes);
-                if (NowMinutes < 3 * 30 * 24)
-                {
-                    throw new AutomationException("bad date calc");
-                }
-                if (!FileExists_NoExceptions(P4IndexFileLocal))
-                {
-                    LogVerbose("{0} doesn't exist, checking in a new one", P4IndexFileP4);
-                    WriteAllText(P4IndexFileLocal, "-1 0");
-                    int WorkingCL = -1;
-                    try
-                    {
-						WorkingCL = P4.CreateChange(P4Env.Client, "Adding new CIS Counter");
-						P4.Add(WorkingCL, P4IndexFileP4);
-                        int SubmittedCL;
-						P4.Submit(WorkingCL, out SubmittedCL);
-                    }
-                    catch (Exception)
-                    {
-                        LogWarning("Add of CIS counter failed, assuming it now exists.");
-                        if (WorkingCL > 0)
-                        {
-							P4.DeleteChange(WorkingCL);
-                        }
-                    }
-                }
-				P4.Sync("-f " + P4IndexFileP4 + "#head");
-                if (!FileExists_NoExceptions(P4IndexFileLocal))
-                {
-                    LogVerbose("{0} doesn't exist, checking in a new one", P4IndexFileP4);
-                    WriteAllText(P4IndexFileLocal, "-1 0");
-                    int WorkingCL = -1;
-                    try
-                    {
-                        WorkingCL = P4.CreateChange(P4Env.Client, "Adding new CIS Counter");
-                        P4.Add(WorkingCL, P4IndexFileP4);
-                        int SubmittedCL;
-                        P4.Submit(WorkingCL, out SubmittedCL);
-                    }
-                    catch (Exception)
-                    {
-                        LogWarning("Add of CIS counter failed, assuming it now exists.");
-                        if (WorkingCL > 0)
-                        {
-                            P4.DeleteChange(WorkingCL);
-                        }
-                    }
-                }
-                var Data = ReadAllText(P4IndexFileLocal);
-                var Parts = Data.Split(" ".ToCharArray());
-                int Index = int.Parse(Parts[0]);
-                int Minutes = int.Parse(Parts[1]);
-
-                int DeltaMinutes = NowMinutes - Minutes;
-
-                int TimeQuantum = 20;
-				if(BranchOptions.QuantumOverride != 0)
-				{
-					TimeQuantum = BranchOptions.QuantumOverride;
-				}
-                int NewIndex = Index + 1;
-
-                if (DeltaMinutes > TimeQuantum * 2)
-                {
-                    if (DeltaMinutes > TimeQuantum * (1 << 8))
-                    {
-                        // it has been forever, lets just start over
-                        NewIndex = 0;
-                    }
-                    else
-                    {
-                        int WorkingIndex = NewIndex + 1;
-                        for (int WorkingDelta = DeltaMinutes - TimeQuantum; WorkingDelta > 0; WorkingDelta -= TimeQuantum, WorkingIndex++)
-                        {
-                            if (CountZeros(NewIndex) < CountZeros(WorkingIndex))
-                            {
-                                NewIndex = WorkingIndex;
-                            }
-                        }
-                    }
-                }
-                {
-                    var Line = String.Format("{0} {1}", NewIndex, NowMinutes);
-                    LogVerbose("Attempting to write {0} with {1}", P4IndexFileP4, Line);
-                    int WorkingCL = -1;
-                    try
-                    {
-						WorkingCL = P4.CreateChange(P4Env.Client, "Updating CIS Counter");
-						P4.Edit(WorkingCL, P4IndexFileP4);
-                        WriteAllText(P4IndexFileLocal, Line);
-                        int SubmittedCL;
-						P4.Submit(WorkingCL, out SubmittedCL);
-                        bDone = true;
-                        TimeIndex = NewIndex;
-                    }
-                    catch (Exception)
-                    {
-                        LogWarning("Edit of CIS counter failed, assuming someone else checked in, retrying.");
-                        if (WorkingCL > 0)
-                        {
-							P4.DeleteChange(WorkingCL);
-                        }
-                        System.Threading.Thread.Sleep(30000);
-                    }
-                }
-            }
-            if (!bDone)
-            {
-                throw new AutomationException("Failed to update the CIS counter after 20 tries.");
-            }
+			TimeIndex = UpdateCISCounter();
             Log("Setting TimeIndex to {0}", TimeIndex);
         }
 
@@ -1773,7 +1653,7 @@ public partial class GUBP : BuildCommand
             }
         }
 
-		AddNodesForBranch(bNoAutomatedTesting);
+		AddNodesForBranch(TimeIndex, bNoAutomatedTesting);
 
         foreach (var NodeToDo in GUBPNodes)
         {
@@ -2346,7 +2226,7 @@ public partial class GUBP : BuildCommand
         string FakeFail = ParseParamValue("FakeFail");
         if (CommanderSetup)
         {
-			DoCommanderSetup(bSkipTriggers, bFakeEC, LocalOnly, CLString, ExplicitTrigger, GUBPNodesControllingTrigger, GUBPNodesControllingTriggerDotName, FullNodeList, FullNodeDirectDependencies, FullNodeDependedOnBy, FullNodeDependentPromotions, SeparatePromotables, FullNodeListSortKey, GUBPNodesCompleted, GUBPNodesHistory, OrdereredToDo, UnfinishedTriggers, FakeFail);
+			DoCommanderSetup(TimeIndex, bSkipTriggers, bFakeEC, LocalOnly, CLString, ExplicitTrigger, GUBPNodesControllingTrigger, GUBPNodesControllingTriggerDotName, FullNodeList, FullNodeDirectDependencies, FullNodeDependedOnBy, FullNodeDependentPromotions, SeparatePromotables, FullNodeListSortKey, GUBPNodesCompleted, GUBPNodesHistory, OrdereredToDo, UnfinishedTriggers, FakeFail);
             Log("Commander setup only, done.");            
             PrintRunTime();
             return;
@@ -2364,6 +2244,124 @@ public partial class GUBP : BuildCommand
 		ExecuteNodes(OrdereredToDo, bOnlyNode, bFakeEC, LocalOnly, bSaveSharedTempStorage, GUBPNodesCompleted, GUBPNodesHistory, CLString, FakeFail);
 	}
 
+	private int UpdateCISCounter()
+	{
+		if (!P4Enabled)
+		{
+			throw new AutomationException("Can't have -CIS without P4 support");
+		}
+		var P4IndexFileP4 = CombinePaths(PathSeparator.Slash, CommandUtils.P4Env.BuildRootP4, "Engine", "Build", "CISCounter.txt");
+		var P4IndexFileLocal = CombinePaths(CmdEnv.LocalRoot, "Engine", "Build", "CISCounter.txt");
+		for(int Retry = 0; Retry < 20; Retry++)
+		{
+			int NowMinutes = (int)((DateTime.UtcNow - new DateTime(2014, 1, 1)).TotalMinutes);
+			if (NowMinutes < 3 * 30 * 24)
+			{
+				throw new AutomationException("bad date calc");
+			}
+			if (!FileExists_NoExceptions(P4IndexFileLocal))
+			{
+				LogVerbose("{0} doesn't exist, checking in a new one", P4IndexFileP4);
+				WriteAllText(P4IndexFileLocal, "-1 0");
+				int WorkingCL = -1;
+				try
+				{
+					WorkingCL = P4.CreateChange(P4Env.Client, "Adding new CIS Counter");
+					P4.Add(WorkingCL, P4IndexFileP4);
+					int SubmittedCL;
+					P4.Submit(WorkingCL, out SubmittedCL);
+				}
+				catch (Exception)
+				{
+					LogWarning("Add of CIS counter failed, assuming it now exists.");
+					if (WorkingCL > 0)
+					{
+						P4.DeleteChange(WorkingCL);
+					}
+				}
+			}
+			P4.Sync("-f " + P4IndexFileP4 + "#head");
+			if (!FileExists_NoExceptions(P4IndexFileLocal))
+			{
+				LogVerbose("{0} doesn't exist, checking in a new one", P4IndexFileP4);
+				WriteAllText(P4IndexFileLocal, "-1 0");
+				int WorkingCL = -1;
+				try
+				{
+					WorkingCL = P4.CreateChange(P4Env.Client, "Adding new CIS Counter");
+					P4.Add(WorkingCL, P4IndexFileP4);
+					int SubmittedCL;
+					P4.Submit(WorkingCL, out SubmittedCL);
+				}
+				catch (Exception)
+				{
+					LogWarning("Add of CIS counter failed, assuming it now exists.");
+					if (WorkingCL > 0)
+					{
+						P4.DeleteChange(WorkingCL);
+					}
+				}
+			}
+			var Data = ReadAllText(P4IndexFileLocal);
+			var Parts = Data.Split(" ".ToCharArray());
+			int Index = int.Parse(Parts[0]);
+			int Minutes = int.Parse(Parts[1]);
+
+			int DeltaMinutes = NowMinutes - Minutes;
+
+			int TimeQuantum = 20;
+			if (BranchOptions.QuantumOverride != 0)
+			{
+				TimeQuantum = BranchOptions.QuantumOverride;
+			}
+			int NewIndex = Index + 1;
+
+			if (DeltaMinutes > TimeQuantum * 2)
+			{
+				if (DeltaMinutes > TimeQuantum * (1 << 8))
+				{
+					// it has been forever, lets just start over
+					NewIndex = 0;
+				}
+				else
+				{
+					int WorkingIndex = NewIndex + 1;
+					for (int WorkingDelta = DeltaMinutes - TimeQuantum; WorkingDelta > 0; WorkingDelta -= TimeQuantum, WorkingIndex++)
+					{
+						if (CountZeros(NewIndex) < CountZeros(WorkingIndex))
+						{
+							NewIndex = WorkingIndex;
+						}
+					}
+				}
+			}
+			{
+				var Line = String.Format("{0} {1}", NewIndex, NowMinutes);
+				LogVerbose("Attempting to write {0} with {1}", P4IndexFileP4, Line);
+				int WorkingCL = -1;
+				try
+				{
+					WorkingCL = P4.CreateChange(P4Env.Client, "Updating CIS Counter");
+					P4.Edit(WorkingCL, P4IndexFileP4);
+					WriteAllText(P4IndexFileLocal, Line);
+					int SubmittedCL;
+					P4.Submit(WorkingCL, out SubmittedCL);
+					return NewIndex;
+				}
+				catch (Exception)
+				{
+					LogWarning("Edit of CIS counter failed, assuming someone else checked in, retrying.");
+					if (WorkingCL > 0)
+					{
+						P4.DeleteChange(WorkingCL);
+					}
+					System.Threading.Thread.Sleep(30000);
+				}
+			}
+		}
+		throw new AutomationException("Failed to update the CIS counter after 20 tries.");
+	}
+
 	private List<string> FindUnfinishedTriggers(bool bSkipTriggers, bool LocalOnly, string ExplicitTrigger, Dictionary<string, bool> GUBPNodesCompleted, List<string> OrdereredToDo)
 	{
 		// find all unfinished triggers, excepting the one we are triggering right now
@@ -2416,7 +2414,7 @@ public partial class GUBP : BuildCommand
 		}
 	}
 
-	private void DoCommanderSetup(bool bSkipTriggers, bool bFakeEC, bool LocalOnly, string CLString, string ExplicitTrigger, Dictionary<string, string> GUBPNodesControllingTrigger, Dictionary<string, string> GUBPNodesControllingTriggerDotName, Dictionary<string, string> FullNodeList, Dictionary<string, string> FullNodeDirectDependencies, Dictionary<string, string> FullNodeDependedOnBy, Dictionary<string, string> FullNodeDependentPromotions, List<string> SeparatePromotables, Dictionary<string, int> FullNodeListSortKey, Dictionary<string, bool> GUBPNodesCompleted, Dictionary<string, NodeHistory> GUBPNodesHistory, List<string> OrdereredToDo, List<string> UnfinishedTriggers, string FakeFail)
+	private void DoCommanderSetup(int TimeIndex, bool bSkipTriggers, bool bFakeEC, bool LocalOnly, string CLString, string ExplicitTrigger, Dictionary<string, string> GUBPNodesControllingTrigger, Dictionary<string, string> GUBPNodesControllingTriggerDotName, Dictionary<string, string> FullNodeList, Dictionary<string, string> FullNodeDirectDependencies, Dictionary<string, string> FullNodeDependedOnBy, Dictionary<string, string> FullNodeDependentPromotions, List<string> SeparatePromotables, Dictionary<string, int> FullNodeListSortKey, Dictionary<string, bool> GUBPNodesCompleted, Dictionary<string, NodeHistory> GUBPNodesHistory, List<string> OrdereredToDo, List<string> UnfinishedTriggers, string FakeFail)
 	{
 
 		if (OrdereredToDo.Count == 0)
diff --git a/Engine/Source/Programs/AutomationTool/GUBP/LegacyBranchSetup.cs b/Engine/Source/Programs/AutomationTool/GUBP/LegacyBranchSetup.cs
--- a/Engine/Source/Programs/AutomationTool/GUBP/LegacyBranchSetup.cs
+++ b/Engine/Source/Programs/AutomationTool/GUBP/LegacyBranchSetup.cs
@@ -241,14 +241,14 @@ partial class GUBP
         return Frequency;
     }    
 
-	void AddNodesForBranch(bool bNoAutomatedTesting)
+	void AddNodesForBranch(int TimeIndex, bool bNoAutomatedTesting)
 	{
         AddNode(new VersionFilesNode());
 
 
         foreach (var HostPlatform in HostPlatforms)
         {
-			AddNode(new ToolsForCompileNode(HostPlatform));
+			AddNode(new ToolsForCompileNode(HostPlatform, TimeIndex));
 			
 			if (!BranchOptions.ExcludePlatformsForEditor.Contains(HostPlatform))
 			{
diff --git a/Engine/Source/Programs/AutomationTool/GUBP/LegacyNodes.cs b/Engine/Source/Programs/AutomationTool/GUBP/LegacyNodes.cs
--- a/Engine/Source/Programs/AutomationTool/GUBP/LegacyNodes.cs
+++ b/Engine/Source/Programs/AutomationTool/GUBP/LegacyNodes.cs
@@ -420,13 +420,17 @@ partial class GUBP
 
     public class ToolsForCompileNode : CompileNode
     {
-        public ToolsForCompileNode(UnrealTargetPlatform InHostPlatform)
+		int TimeIndex;
+
+        public ToolsForCompileNode(UnrealTargetPlatform InHostPlatform, int InTimeIndex)
             : base(InHostPlatform, false)
         {
-            if (InHostPlatform != UnrealTargetPlatform.Win64)
+			TimeIndex = InTimeIndex;
+
+			if (InHostPlatform != UnrealTargetPlatform.Win64)
             {
-            AgentSharingGroup = "Editor" + StaticGetHostPlatformSuffix(InHostPlatform);
-        }
+				AgentSharingGroup = "Editor" + StaticGetHostPlatformSuffix(InHostPlatform);
+			}
         }
         public static string StaticGetFullName(UnrealTargetPlatform InHostPlatform)
         {
@@ -458,7 +462,7 @@ partial class GUBP
         }
         public override int AgentMemoryRequirement(GUBP bp)
         {
-            if (bp.ParseParam("Launcher") || bp.TimeIndex != 0  && HostPlatform != UnrealTargetPlatform.Mac)
+            if (bp.ParseParam("Launcher") || TimeIndex != 0  && HostPlatform != UnrealTargetPlatform.Mac)
             {
                 return base.AgentMemoryRequirement(bp);
             }
