Move branch validation stuff into LegacyBranchSetup.

diff --git a/Engine/Source/Programs/AutomationTool/GUBP/GUBP.cs b/Engine/Source/Programs/AutomationTool/GUBP/GUBP.cs
--- a/Engine/Source/Programs/AutomationTool/GUBP/GUBP.cs
+++ b/Engine/Source/Programs/AutomationTool/GUBP/GUBP.cs
@@ -1629,30 +1629,6 @@ public partial class GUBP : BuildCommand
             LogVerbose("Active Platform: {0}", Plat.ToString());
         }
 
-        if (HostPlatforms.Count >= 2)
-        {
-            // make sure each project is set up with the right assumptions on monolithics that prefer a platform.
-            foreach (var CodeProj in Branch.CodeProjects)
-            {
-                var OptionsMac = CodeProj.Options(UnrealTargetPlatform.Mac);
-                var OptionsPC = CodeProj.Options(UnrealTargetPlatform.Win64);
-
-                var MacMonos = GetMonolithicPlatformsForUProject(UnrealTargetPlatform.Mac, CodeProj, false);
-                var PCMonos = GetMonolithicPlatformsForUProject(UnrealTargetPlatform.Win64, CodeProj, false);
-
-                if (!OptionsMac.bIsPromotable && OptionsPC.bIsPromotable && 
-                    (MacMonos.Contains(UnrealTargetPlatform.IOS) || PCMonos.Contains(UnrealTargetPlatform.IOS)))
-                {
-                    throw new AutomationException("Project {0} is promotable for PC, not promotable for Mac and uses IOS monothics. Since Mac is the preferred platform for IOS, please add Mac as a promotable platform.", CodeProj.GameName);
-                }
-                if (OptionsMac.bIsPromotable && !OptionsPC.bIsPromotable &&
-                    (MacMonos.Contains(UnrealTargetPlatform.Android) || PCMonos.Contains(UnrealTargetPlatform.Android)))
-                {
-                    throw new AutomationException("Project {0} is not promotable for PC, promotable for Mac and uses Android monothics. Since PC is the preferred platform for Android, please add PC as a promotable platform.", CodeProj.GameName);
-                }
-            }
-        }
-
 		AddNodesForBranch(TimeIndex, bNoAutomatedTesting);
 
         foreach (var NodeToDo in GUBPNodes)
diff --git a/Engine/Source/Programs/AutomationTool/GUBP/LegacyBranchSetup.cs b/Engine/Source/Programs/AutomationTool/GUBP/LegacyBranchSetup.cs
--- a/Engine/Source/Programs/AutomationTool/GUBP/LegacyBranchSetup.cs
+++ b/Engine/Source/Programs/AutomationTool/GUBP/LegacyBranchSetup.cs
@@ -243,6 +243,31 @@ partial class GUBP
 
 	void AddNodesForBranch(int TimeIndex, bool bNoAutomatedTesting)
 	{
+        if (HostPlatforms.Count >= 2)
+        {
+            // make sure each project is set up with the right assumptions on monolithics that prefer a platform.
+            foreach (var CodeProj in Branch.CodeProjects)
+            {
+                var OptionsMac = CodeProj.Options(UnrealTargetPlatform.Mac);
+                var OptionsPC = CodeProj.Options(UnrealTargetPlatform.Win64);
+
+                var MacMonos = GetMonolithicPlatformsForUProject(UnrealTargetPlatform.Mac, CodeProj, false);
+                var PCMonos = GetMonolithicPlatformsForUProject(UnrealTargetPlatform.Win64, CodeProj, false);
+
+                if (!OptionsMac.bIsPromotable && OptionsPC.bIsPromotable && 
+                    (MacMonos.Contains(UnrealTargetPlatform.IOS) || PCMonos.Contains(UnrealTargetPlatform.IOS)))
+                {
+                    throw new AutomationException("Project {0} is promotable for PC, not promotable for Mac and uses IOS monothics. Since Mac is the preferred platform for IOS, please add Mac as a promotable platform.", CodeProj.GameName);
+                }
+                if (OptionsMac.bIsPromotable && !OptionsPC.bIsPromotable &&
+                    (MacMonos.Contains(UnrealTargetPlatform.Android) || PCMonos.Contains(UnrealTargetPlatform.Android)))
+                {
+                    throw new AutomationException("Project {0} is not promotable for PC, promotable for Mac and uses Android monothics. Since PC is the preferred platform for Android, please add PC as a promotable platform.", CodeProj.GameName);
+                }
+            }
+        }
+
+
         AddNode(new VersionFilesNode());
 
 
