Add ParallelExecutor telemetry

diff --git a/Engine/Source/Programs/AutomationTool/AutomationUtils/UE4Build.cs b/Engine/Source/Programs/AutomationTool/AutomationUtils/UE4Build.cs
--- a/Engine/Source/Programs/AutomationTool/AutomationUtils/UE4Build.cs
+++ b/Engine/Source/Programs/AutomationTool/AutomationUtils/UE4Build.cs
@@ -785,7 +785,8 @@ namespace AutomationTool
 
 		public bool ProcessXGEItems(List<XGEItem> Actions, string XGETool, string Args, string TaskFilePath, bool DoRetries, bool SpecialTestFlag, bool ShowProgress)
 		{
-			if (string.IsNullOrEmpty(XGETool))
+            var StartProcessTime = DateTime.Now.ToString();
+            if (string.IsNullOrEmpty(XGETool))
 			{
 				throw new AutomationException("Unable to find xge tool: " + XGETool);
 			}
@@ -842,7 +843,10 @@ namespace AutomationTool
 					throw new AutomationException("Unable to find xge xml: " + TaskFilePath);
 				}
 
-				int Retries = DoRetries ? 2 : 1;
+                var FinishCombineXGEItemFilesTime = DateTime.Now.ToString();
+                PrintCSVFile(String.Format("UAT,CombineXGEItemFiles.{0},{1},{2}", Path.GetFileNameWithoutExtension(XGETool), StartProcessTime, FinishCombineXGEItemFilesTime));
+                
+                int Retries = DoRetries ? 2 : 1;
                 int ConnectionRetries = 4;
                 for (int i = 0; i < Retries; i++)
 				{
@@ -891,6 +895,8 @@ namespace AutomationTool
 						}
 					}
 				}
+                var FinishXGEToolTime = DateTime.Now.ToString();
+                PrintCSVFile(String.Format("UAT,ProcessXGE.{0},{1},{2}", Path.GetFileNameWithoutExtension(XGETool), FinishCombineXGEItemFilesTime, FinishXGEToolTime));
 
 			}
 			foreach (var Item in Actions)
