Add node name to DoBuild telemetry

diff --git a/Engine/Source/Programs/AutomationTool/GUBP/GUBP.cs b/Engine/Source/Programs/AutomationTool/GUBP/GUBP.cs
--- a/Engine/Source/Programs/AutomationTool/GUBP/GUBP.cs
+++ b/Engine/Source/Programs/AutomationTool/GUBP/GUBP.cs
@@ -3025,7 +3025,7 @@ public partial class GUBP : BuildCommand
                         var FinishBuild = DateTime.Now.ToString();
                         if (IsBuildMachine)
                         {
-                            PrintCSVFile(String.Format("UAT,DoBuild,{0},{1}", StartBuild, FinishBuild));
+                            PrintCSVFile(String.Format("UAT,DoBuild.{0},{1},{2}", NodeToDo, StartBuild, FinishBuild));
                         }
 						BuildDuration = (DateTime.UtcNow - StartTime).TotalMilliseconds / 1000;
 						
