New method for determining unique log filenames.

diff --git a/Engine/Source/Programs/AutomationTool/AutomationUtils/Log.cs b/Engine/Source/Programs/AutomationTool/AutomationUtils/Log.cs
--- a/Engine/Source/Programs/AutomationTool/AutomationUtils/Log.cs
+++ b/Engine/Source/Programs/AutomationTool/AutomationUtils/Log.cs
@@ -168,25 +168,30 @@ namespace AutomationTool
 		{
 			const int MaxAttempts = 1000;
 
-			string LogFilename = string.Empty;
-
+			var Now     = DateTime.Now;
 			int Attempt = 0;
-			do
+
+			string LogFilename;
+			for (;;)
 			{
-				if (Attempt == 0)
+				var NowStr = Now.ToString("yyyy.MM.dd-HH.mm.ss");
+
+				LogFilename = String.Format("{0}-{1}.txt", Base, NowStr);
+
+				if (!File.Exists(LogFilename))
 				{
-					LogFilename = String.Format("{0}.txt", Base);
+					break;
 				}
-				else
+
+				++Attempt;
+				if (Attempt == MaxAttempts)
 				{
-					LogFilename = String.Format("{0}.{1}.txt", Base, Attempt);
+					throw new AutomationException(String.Format("Failed to create logfile {0}.", LogFilename));
 				}
-			} while (File.Exists(LogFilename) && ++Attempt < MaxAttempts);
 
-			if (File.Exists(LogFilename))
-			{
-				throw new AutomationException(String.Format("Failed to create logfile {0}.", LogFilename));
+				Now.AddSeconds(1);
 			}
+
 			return LogFilename;
 		}
 
