Fix for two logfiles being created in the same second.

diff --git a/Engine/Source/Programs/AutomationTool/AutomationUtils/Log.cs b/Engine/Source/Programs/AutomationTool/AutomationUtils/Log.cs
--- a/Engine/Source/Programs/AutomationTool/AutomationUtils/Log.cs
+++ b/Engine/Source/Programs/AutomationTool/AutomationUtils/Log.cs
@@ -189,7 +189,7 @@ namespace AutomationTool
 					throw new AutomationException(String.Format("Failed to create logfile {0}.", LogFilename));
 				}
 
-				Now.AddSeconds(1);
+				Now = Now.AddSeconds(1);
 			}
 
 			return LogFilename;
