Separate out a couple more blocks of code into separate functions.

diff --git a/Engine/Source/Programs/AutomationTool/GUBP/GUBP.cs b/Engine/Source/Programs/AutomationTool/GUBP/GUBP.cs
--- a/Engine/Source/Programs/AutomationTool/GUBP/GUBP.cs
+++ b/Engine/Source/Programs/AutomationTool/GUBP/GUBP.cs
@@ -1550,31 +1550,7 @@ public partial class GUBP : BuildCommand
 
 		AddNodesForBranch(TimeIndex, bNoAutomatedTesting);
 
-        foreach (var NodeToDo in GUBPNodes)
-        {
-            foreach (var Dep in GUBPNodes[NodeToDo.Key].FullNamesOfDependencies)
-            {
-                if (!GUBPNodes.ContainsKey(Dep))
-                {
-                    throw new AutomationException("Node {0} is not in the full graph. It is a dependency of {1}.", Dep, NodeToDo.Key);
-                }
-                if (Dep == NodeToDo.Key)
-                {
-                    throw new AutomationException("Node {0} has a self arc.", NodeToDo.Key);
-                }
-            }
-            foreach (var Dep in GUBPNodes[NodeToDo.Key].FullNamesOfPseudosependencies)
-            {
-                if (!GUBPNodes.ContainsKey(Dep))
-                {
-                    throw new AutomationException("Node {0} is not in the full graph. It is a pseudodependency of {1}.", Dep, NodeToDo.Key);
-                }
-                if (Dep == NodeToDo.Key)
-                {
-                    throw new AutomationException("Node {0} has a self pseudoarc.", NodeToDo.Key);
-                }
-            }
-        }
+		ValidateDependencyNames();
 
         if (bCleanLocalTempStorage)  // shared temp storage can never be wiped
         {
@@ -1584,8 +1560,6 @@ public partial class GUBP : BuildCommand
 		Dictionary<string, string> GUBPNodesControllingTrigger = new Dictionary<string, string>();
         Dictionary<string, string> GUBPNodesControllingTriggerDotName = new Dictionary<string, string>();
 
-        var FullNodeList = new Dictionary<string, string>();
-        var FullNodeDirectDependencies = new Dictionary<string, string>();
 		var FullNodeDependedOnBy = new Dictionary<string, string>();
 		var FullNodeDependentPromotions = new Dictionary<string, string>();		
 		var SeparatePromotables = new List<string>();
@@ -1676,45 +1650,11 @@ public partial class GUBP : BuildCommand
                 FullNodeDependentPromotions.Add(NodeToDo.Key, All);
             }
         }	
-        {
-            Log("******* {0} GUBP Nodes", GUBPNodes.Count);
-            var SortedNodes = TopologicalSort(new HashSet<string>(GUBPNodes.Keys), null/*GUBPNodesCompleted*/, GUBPNodesControllingTriggerDotName, GUBPNodesControllingTrigger, LocalOnly: true, DoNotConsiderCompletion: true);
-            var DependencyStart = DateTime.Now.ToString();
-            foreach (var Node in SortedNodes)
-            {
-                string Note = GetControllingTriggerDotName(Node, GUBPNodesControllingTriggerDotName, GUBPNodesControllingTrigger);
-                if (Note == "")
-                {
-                    Note = CISFrequencyQuantumShiftString(Node);
-                }
-                if (Note == "")
-                {
-                    Note = "always";
-                }
-                if (GUBPNodes[Node].RunInEC())
-                {
-                    var Deps = GetECDependencies(Node);
-                    string All = "";
-                    foreach (var Dep in Deps)
-                    {
-                        if (All != "")
-                        {
-                            All += " ";
-                        }
-                        All += Dep;
-                    }
-                    LogVerbose("  {0}: {1}      {2}", Node, Note, All);
-                    FullNodeList.Add(Node, Note);
-                    FullNodeDirectDependencies.Add(Node, All);
-                }
-                else
-                {
-                    LogVerbose("  {0}: {1} [Aggregate]", Node, Note);
-                }
-            }
-            var DependencyFinish = DateTime.Now.ToString();
-            PrintCSVFile(String.Format("UAT,GetDependencies,{0},{1}", DependencyStart, DependencyFinish));
-        }
+
+        Dictionary<string, string> FullNodeList;
+        Dictionary<string, string> FullNodeDirectDependencies;
+		GetFullNodeListAndDirectDependencies(GUBPNodesControllingTrigger, GUBPNodesControllingTriggerDotName, out FullNodeList, out FullNodeDirectDependencies);
+
 		Dictionary<string, int> FullNodeListSortKey = GetDisplayOrder(FullNodeList.Keys.ToList(), FullNodeDirectDependencies, GUBPNodes);
 
         bool bOnlyNode = false;
@@ -2139,6 +2079,79 @@ public partial class GUBP : BuildCommand
 		ExecuteNodes(OrdereredToDo, bOnlyNode, bFakeEC, LocalOnly, bSaveSharedTempStorage, GUBPNodesCompleted, GUBPNodesHistory, CLString, FakeFail);
 	}
 
+	private void GetFullNodeListAndDirectDependencies(Dictionary<string, string> GUBPNodesControllingTrigger, Dictionary<string, string> GUBPNodesControllingTriggerDotName, out Dictionary<string, string> FullNodeList, out Dictionary<string, string> FullNodeDirectDependencies)
+	{
+		FullNodeList = new Dictionary<string,string>();
+		FullNodeDirectDependencies = new Dictionary<string,string>();
+
+		Log("******* {0} GUBP Nodes", GUBPNodes.Count);
+		var SortedNodes = TopologicalSort(new HashSet<string>(GUBPNodes.Keys), null/*GUBPNodesCompleted*/, GUBPNodesControllingTriggerDotName, GUBPNodesControllingTrigger, LocalOnly: true, DoNotConsiderCompletion: true);
+		var DependencyStart = DateTime.Now.ToString();
+		foreach (var Node in SortedNodes)
+		{
+			string Note = GetControllingTriggerDotName(Node, GUBPNodesControllingTriggerDotName, GUBPNodesControllingTrigger);
+			if (Note == "")
+			{
+				Note = CISFrequencyQuantumShiftString(Node);
+			}
+			if (Note == "")
+			{
+				Note = "always";
+			}
+			if (GUBPNodes[Node].RunInEC())
+			{
+				var Deps = GetECDependencies(Node);
+				string All = "";
+				foreach (var Dep in Deps)
+				{
+					if (All != "")
+					{
+						All += " ";
+					}
+					All += Dep;
+				}
+				LogVerbose("  {0}: {1}      {2}", Node, Note, All);
+				FullNodeList.Add(Node, Note);
+				FullNodeDirectDependencies.Add(Node, All);
+			}
+			else
+			{
+				LogVerbose("  {0}: {1} [Aggregate]", Node, Note);
+			}
+		}
+		var DependencyFinish = DateTime.Now.ToString();
+		PrintCSVFile(String.Format("UAT,GetDependencies,{0},{1}", DependencyStart, DependencyFinish));
+	}
+
+	private void ValidateDependencyNames()
+	{
+		foreach (var NodeToDo in GUBPNodes)
+		{
+			foreach (var Dep in GUBPNodes[NodeToDo.Key].FullNamesOfDependencies)
+			{
+				if (!GUBPNodes.ContainsKey(Dep))
+				{
+					throw new AutomationException("Node {0} is not in the full graph. It is a dependency of {1}.", Dep, NodeToDo.Key);
+				}
+				if (Dep == NodeToDo.Key)
+				{
+					throw new AutomationException("Node {0} has a self arc.", NodeToDo.Key);
+				}
+			}
+			foreach (var Dep in GUBPNodes[NodeToDo.Key].FullNamesOfPseudosependencies)
+			{
+				if (!GUBPNodes.ContainsKey(Dep))
+				{
+					throw new AutomationException("Node {0} is not in the full graph. It is a pseudodependency of {1}.", Dep, NodeToDo.Key);
+				}
+				if (Dep == NodeToDo.Key)
+				{
+					throw new AutomationException("Node {0} has a self pseudoarc.", NodeToDo.Key);
+				}
+			}
+		}
+	}
+
 	private int UpdateCISCounter()
 	{
 		if (!P4Enabled)
