Move NodesToDo into its own function.

diff --git a/Engine/Source/Programs/AutomationTool/GUBP/GUBP.cs b/Engine/Source/Programs/AutomationTool/GUBP/GUBP.cs
--- a/Engine/Source/Programs/AutomationTool/GUBP/GUBP.cs
+++ b/Engine/Source/Programs/AutomationTool/GUBP/GUBP.cs
@@ -1657,311 +1657,9 @@ public partial class GUBP : BuildCommand
 
 		Dictionary<string, int> FullNodeListSortKey = GetDisplayOrder(FullNodeList.Keys.ToList(), FullNodeDirectDependencies, GUBPNodes);
 
-        bool bOnlyNode = false;
-        bool bRelatedToNode = false;		
-        var NodesToDo = new HashSet<string>();
-
-        {
-            string NodeSpec = ParseParamValue("Node");
-            if (String.IsNullOrEmpty(NodeSpec))
-            {
-                NodeSpec = ParseParamValue("RelatedToNode");
-                if (!String.IsNullOrEmpty(NodeSpec))
-                {
-                    bRelatedToNode = true;
-                }
-            }
-            if (String.IsNullOrEmpty(NodeSpec) && CommanderSetup)
-            {
-                NodeSpec = ParseParamValue("SetupNode");
-                if (String.IsNullOrEmpty(NodeSpec))
-                {
-                    NodeSpec = ParseParamValue("SetupRelatedToNode");
-                    if (!String.IsNullOrEmpty(NodeSpec))
-                    {
-                        bRelatedToNode = true;
-                    }
-                }
-            }
-            if (String.IsNullOrEmpty(NodeSpec))
-            {
-                NodeSpec = ParseParamValue("OnlyNode");
-                if (!String.IsNullOrEmpty(NodeSpec))
-                {
-                    bOnlyNode = true;
-                }
-            }
-            if (!String.IsNullOrEmpty(NodeSpec))
-            {				
-                if (NodeSpec.Equals("Noop", StringComparison.InvariantCultureIgnoreCase))
-                {
-                    LogVerbose("Request for Noop node, done.");
-                    PrintRunTime();
-                    return;
-                }
-                List<string> Nodes = new List<string>(NodeSpec.Split('+'));
-                foreach (var NodeArg in Nodes)
-                {
-                    var NodeName = NodeArg.Trim();
-                    bool bFoundAnything = false;
-                    if (!String.IsNullOrEmpty(NodeName))
-                    {
-                        foreach (var Node in GUBPNodes)
-                        {
-                            if (Node.Value.GetFullName().Equals(NodeArg, StringComparison.InvariantCultureIgnoreCase) ||
-                                Node.Value.AgentSharingGroup.Equals(NodeArg, StringComparison.InvariantCultureIgnoreCase)
-                                )
-                            {
-                                if (!NodesToDo.Contains(Node.Key))
-                                {
-                                    NodesToDo.Add(Node.Key);
-                                }
-                                bFoundAnything = true;
-                            }
-                        }
-                        if (!bFoundAnything)
-                        {
-                            throw new AutomationException("Could not find node named {0}", NodeName);
-                        }
-                    }
-                }
-            }
-        }
-        string GameSpec = ParseParamValue("Game");
-        if (!String.IsNullOrEmpty(GameSpec))
-        {			
-            List<string> Games = new List<string>(GameSpec.Split('+'));
-            foreach (var GameArg in Games)
-            {
-                var GameName = GameArg.Trim();
-                if (!String.IsNullOrEmpty(GameName))
-                {
-                    foreach (var GameProj in Branch.CodeProjects)
-                    {
-                        if (GameProj.GameName.Equals(GameName, StringComparison.InvariantCultureIgnoreCase))
-                        {
-                            if (GameProj.Options(UnrealTargetPlatform.Win64).bIsPromotable)
-                            {
-                                NodesToDo.Add(GameAggregatePromotableNode.StaticGetFullName(GameProj));
-                            }
-                            foreach (var Node in GUBPNodes)
-                            {
-                                if (Node.Value.GameNameIfAnyForTempStorage() == GameProj.GameName)
-                                {
-                                    NodesToDo.Add(Node.Key);
-                                }
-                            }
- 
-                            GameName = null;
-                        }
-                    }
-                    if (GameName != null)
-                    {
-                        foreach (var GameProj in Branch.NonCodeProjects)
-                        {
-                            if (GameProj.GameName.Equals(GameName, StringComparison.InvariantCultureIgnoreCase))
-                            {
-                                foreach (var Node in GUBPNodes)
-                                {
-                                    if (Node.Value.GameNameIfAnyForTempStorage() == GameProj.GameName)
-                                    {
-                                        NodesToDo.Add(Node.Key);
-                                    }
-                                }
-                                GameName = null;
-                            }
-                        }
-                    }
-                    if (GameName != null)
-                    {
-                        throw new AutomationException("Could not find game named {0}", GameName);
-                    }
-                }
-            }
-        }
-        if (NodesToDo.Count == 0)
-        {
-            LogVerbose("No nodes specified, adding all nodes");
-            foreach (var Node in GUBPNodes)
-            {
-                NodesToDo.Add(Node.Key);
-            }
-        }
-        else if (TimeIndex != 0)
-        {
-            LogVerbose("Check to make sure we didn't ask for nodes that will be culled by time index");
-            foreach (var NodeToDo in NodesToDo)
-            {
-                if (TimeIndex % (1 << GUBPNodes[NodeToDo].DependentCISFrequencyQuantumShift()) != 0)
-                {
-                    throw new AutomationException("You asked specifically for node {0}, but it is culled by the time quantum: TimeIndex = {1}, DependentCISFrequencyQuantumShift = {2}.", NodeToDo, TimeIndex, GUBPNodes[NodeToDo].DependentCISFrequencyQuantumShift());
-                }
-            }
-        }
-
-        LogVerbose("Desired Nodes");
-        foreach (var NodeToDo in NodesToDo)
-        {
-            LogVerbose("  {0}", NodeToDo);			
-        }
-        // if we are doing related to, then find things that depend on the selected nodes
-        if (bRelatedToNode)
-        {
-            bool bDoneWithDependencies = false;
-
-            while (!bDoneWithDependencies)
-            {
-                bDoneWithDependencies = true;
-                var Fringe = new HashSet<string>();
-                foreach (var NodeToDo in GUBPNodes)
-                {
-                    if (!NodesToDo.Contains(NodeToDo.Key))
-                    {
-                        foreach (var Dep in GUBPNodes[NodeToDo.Key].FullNamesOfDependencies)
-                        {
-                            if (!GUBPNodes.ContainsKey(Dep))
-                            {
-                                throw new AutomationException("Node {0} is not in the graph. It is a dependency of {1}.", Dep, NodeToDo.Key);
-                            }
-                            if (NodesToDo.Contains(Dep))
-                            {
-                                Fringe.Add(NodeToDo.Key);
-                                bDoneWithDependencies = false;
-                            }
-                        }
-                        foreach (var Dep in GUBPNodes[NodeToDo.Key].FullNamesOfPseudosependencies)
-                        {
-                            if (!GUBPNodes.ContainsKey(Dep))
-                            {
-                                throw new AutomationException("Node {0} is not in the graph. It is a pseudodependency of {1}.", Dep, NodeToDo.Key);
-                            }
-                        }
-                    }
-                }
-                NodesToDo.UnionWith(Fringe);
-            }
-        }
-
-        // find things that our nodes depend on
-        if (!bOnlyNode)
-        {
-            bool bDoneWithDependencies = false;
-
-            while (!bDoneWithDependencies)
-            {
-                bDoneWithDependencies = true;
-                var Fringe = new HashSet<string>(); 
-                foreach (var NodeToDo in NodesToDo)
-                {
-                    if (BranchOptions.NodesToRemovePseudoDependencies.Count > 0)
-                    {
-                        var ListOfRetainedNodesToRemovePseudoDependencies = BranchOptions.NodesToRemovePseudoDependencies;
-                        foreach (var NodeToRemovePseudoDependencies in BranchOptions.NodesToRemovePseudoDependencies)
-                        {
-                            if (NodeToDo == NodeToRemovePseudoDependencies)
-                            {
-                                RemoveAllPseudodependenciesFromNode(NodeToDo);
-                                ListOfRetainedNodesToRemovePseudoDependencies.Remove(NodeToDo);
-                                break;
-                            }
-                        }
-                        BranchOptions.NodesToRemovePseudoDependencies = ListOfRetainedNodesToRemovePseudoDependencies;
-                    }
-                    foreach (var Dep in GUBPNodes[NodeToDo].FullNamesOfDependencies)
-                    {
-                        if (!GUBPNodes.ContainsKey(Dep))
-                        {
-                            throw new AutomationException("Node {0} is not in the graph. It is a dependency of {1}.", Dep, NodeToDo);
-                        }
-                        if (!NodesToDo.Contains(Dep))
-                        {
-                            Fringe.Add(Dep);
-                            bDoneWithDependencies = false;
-                        }
-                    }
-                    foreach (var Dep in GUBPNodes[NodeToDo].FullNamesOfPseudosependencies)
-                    {
-                        if (!GUBPNodes.ContainsKey(Dep))
-                        {
-                            throw new AutomationException("Node {0} is not in the graph. It is a pseudodependency of {1}.", Dep, NodeToDo);
-                        }
-                    }
-                }
-                NodesToDo.UnionWith(Fringe);
-            }
-		}		
-		if (TimeIndex != 0)
-		{
-			LogVerbose("Culling based on time index");
-			var NewNodesToDo = new HashSet<string>();
-			foreach (var NodeToDo in NodesToDo)
-			{
-				if (TimeIndex % (1 << GUBPNodes[NodeToDo].DependentCISFrequencyQuantumShift()) == 0)
-				{
-					LogVerbose("  Keeping {0}", NodeToDo);
-					NewNodesToDo.Add(NodeToDo);
-				}
-				else
-				{
-					LogVerbose("  Rejecting {0}", NodeToDo);
-				}
-			}
-			NodesToDo = NewNodesToDo;
-		}
-		//Remove Plat if specified
-		if(WithoutLinux)
-		{
-			var NewNodesToDo = new HashSet<string>();
-			foreach(var NodeToDo in NodesToDo)
-			{
-				if(!GUBPNodes[NodeToDo].GetFullName().Contains("Linux"))
-				{					
-					NewNodesToDo.Add(NodeToDo);
-				}
-				else
-				{
-					LogVerbose(" Rejecting {0} because -NoLinux was requested", NodeToDo);
-				}
-			}
-			NodesToDo = NewNodesToDo;
-		}
-		//find things that depend on our nodes and setup commander dictionary
-		if (!bOnlyNode)
-		{			
-			foreach(var NodeToDo in NodesToDo)
-			{
-				if (!GUBPNodes[NodeToDo].IsAggregate() && !GUBPNodes[NodeToDo].IsTest())
-				{
-					List<string> ECDependencies = new List<string>();
-					ECDependencies = GetECDependencies(NodeToDo);
-					foreach (var Dep in ECDependencies)
-					{
-						if (!GUBPNodes.ContainsKey(Dep))
-						{
-							throw new AutomationException("Node {0} is not in the graph. It is a dependency of {1}.", Dep, NodeToDo);
-						}
-						if (!GUBPNodes[Dep].FullNamesOfDependedOn.Contains(NodeToDo))
-						{
-							GUBPNodes[Dep].FullNamesOfDependedOn.Add(NodeToDo);
-						}						
-					}
-				}				
-			}	
-			foreach(var NodeToDo in NodesToDo)
-			{
-				var Deps = GUBPNodes[NodeToDo].FullNamesOfDependedOn;
-				string All = "";
-				foreach (var Dep in Deps)
-				{
-					if (All != "")
-					{
-						All += " ";
-					}
-					All += Dep;
-				}
-				FullNodeDependedOnBy.Add(NodeToDo, All);
-			}
-		}
+		bool bOnlyNode;
+		HashSet<string> NodesToDo;
+		ParseNodesToDo(WithoutLinux, TimeIndex, CommanderSetup, FullNodeDependedOnBy, out bOnlyNode, out NodesToDo);
 
         if (CommanderSetup)
         {
@@ -2079,6 +1777,309 @@ public partial class GUBP : BuildCommand
 		ExecuteNodes(OrdereredToDo, bOnlyNode, bFakeEC, LocalOnly, bSaveSharedTempStorage, GUBPNodesCompleted, GUBPNodesHistory, CLString, FakeFail);
 	}
 
+	private void ParseNodesToDo(bool WithoutLinux, int TimeIndex, bool CommanderSetup, Dictionary<string, string> FullNodeDependedOnBy, out bool bOnlyNode, out HashSet<string> NodesToDo)
+	{
+		bOnlyNode = false;
+		bool bRelatedToNode = false;
+		NodesToDo = new HashSet<string>();
+
+		{
+			string NodeSpec = ParseParamValue("Node");
+			if (String.IsNullOrEmpty(NodeSpec))
+			{
+				NodeSpec = ParseParamValue("RelatedToNode");
+				if (!String.IsNullOrEmpty(NodeSpec))
+				{
+					bRelatedToNode = true;
+				}
+			}
+			if (String.IsNullOrEmpty(NodeSpec) && CommanderSetup)
+			{
+				NodeSpec = ParseParamValue("SetupNode");
+				if (String.IsNullOrEmpty(NodeSpec))
+				{
+					NodeSpec = ParseParamValue("SetupRelatedToNode");
+					if (!String.IsNullOrEmpty(NodeSpec))
+					{
+						bRelatedToNode = true;
+					}
+				}
+			}
+			if (String.IsNullOrEmpty(NodeSpec))
+			{
+				NodeSpec = ParseParamValue("OnlyNode");
+				if (!String.IsNullOrEmpty(NodeSpec))
+				{
+					bOnlyNode = true;
+				}
+			}
+			if (!String.IsNullOrEmpty(NodeSpec))
+			{
+				List<string> Nodes = new List<string>(NodeSpec.Split('+'));
+				foreach (var NodeArg in Nodes)
+				{
+					var NodeName = NodeArg.Trim();
+					bool bFoundAnything = false;
+					if (!String.IsNullOrEmpty(NodeName))
+					{
+						foreach (var Node in GUBPNodes)
+						{
+							if (Node.Value.GetFullName().Equals(NodeArg, StringComparison.InvariantCultureIgnoreCase) ||
+								Node.Value.AgentSharingGroup.Equals(NodeArg, StringComparison.InvariantCultureIgnoreCase)
+								)
+							{
+								if (!NodesToDo.Contains(Node.Key))
+								{
+									NodesToDo.Add(Node.Key);
+								}
+								bFoundAnything = true;
+							}
+						}
+						if (!bFoundAnything)
+						{
+							throw new AutomationException("Could not find node named {0}", NodeName);
+						}
+					}
+				}
+			}
+		}
+		string GameSpec = ParseParamValue("Game");
+		if (!String.IsNullOrEmpty(GameSpec))
+		{
+			List<string> Games = new List<string>(GameSpec.Split('+'));
+			foreach (var GameArg in Games)
+			{
+				var GameName = GameArg.Trim();
+				if (!String.IsNullOrEmpty(GameName))
+				{
+					foreach (var GameProj in Branch.CodeProjects)
+					{
+						if (GameProj.GameName.Equals(GameName, StringComparison.InvariantCultureIgnoreCase))
+						{
+							if (GameProj.Options(UnrealTargetPlatform.Win64).bIsPromotable)
+							{
+								NodesToDo.Add(GameAggregatePromotableNode.StaticGetFullName(GameProj));
+							}
+							foreach (var Node in GUBPNodes)
+							{
+								if (Node.Value.GameNameIfAnyForTempStorage() == GameProj.GameName)
+								{
+									NodesToDo.Add(Node.Key);
+								}
+							}
+
+							GameName = null;
+						}
+					}
+					if (GameName != null)
+					{
+						foreach (var GameProj in Branch.NonCodeProjects)
+						{
+							if (GameProj.GameName.Equals(GameName, StringComparison.InvariantCultureIgnoreCase))
+							{
+								foreach (var Node in GUBPNodes)
+								{
+									if (Node.Value.GameNameIfAnyForTempStorage() == GameProj.GameName)
+									{
+										NodesToDo.Add(Node.Key);
+									}
+								}
+								GameName = null;
+							}
+						}
+					}
+					if (GameName != null)
+					{
+						throw new AutomationException("Could not find game named {0}", GameName);
+					}
+				}
+			}
+		}
+		if (NodesToDo.Count == 0)
+		{
+			LogVerbose("No nodes specified, adding all nodes");
+			foreach (var Node in GUBPNodes)
+			{
+				NodesToDo.Add(Node.Key);
+			}
+		}
+		else if (TimeIndex != 0)
+		{
+			LogVerbose("Check to make sure we didn't ask for nodes that will be culled by time index");
+			foreach (var NodeToDo in NodesToDo)
+			{
+				if (TimeIndex % (1 << GUBPNodes[NodeToDo].DependentCISFrequencyQuantumShift()) != 0)
+				{
+					throw new AutomationException("You asked specifically for node {0}, but it is culled by the time quantum: TimeIndex = {1}, DependentCISFrequencyQuantumShift = {2}.", NodeToDo, TimeIndex, GUBPNodes[NodeToDo].DependentCISFrequencyQuantumShift());
+				}
+			}
+		}
+
+		LogVerbose("Desired Nodes");
+		foreach (var NodeToDo in NodesToDo)
+		{
+			LogVerbose("  {0}", NodeToDo);
+		}
+		// if we are doing related to, then find things that depend on the selected nodes
+		if (bRelatedToNode)
+		{
+			bool bDoneWithDependencies = false;
+
+			while (!bDoneWithDependencies)
+			{
+				bDoneWithDependencies = true;
+				var Fringe = new HashSet<string>();
+				foreach (var NodeToDo in GUBPNodes)
+				{
+					if (!NodesToDo.Contains(NodeToDo.Key))
+					{
+						foreach (var Dep in GUBPNodes[NodeToDo.Key].FullNamesOfDependencies)
+						{
+							if (!GUBPNodes.ContainsKey(Dep))
+							{
+								throw new AutomationException("Node {0} is not in the graph. It is a dependency of {1}.", Dep, NodeToDo.Key);
+							}
+							if (NodesToDo.Contains(Dep))
+							{
+								Fringe.Add(NodeToDo.Key);
+								bDoneWithDependencies = false;
+							}
+						}
+						foreach (var Dep in GUBPNodes[NodeToDo.Key].FullNamesOfPseudosependencies)
+						{
+							if (!GUBPNodes.ContainsKey(Dep))
+							{
+								throw new AutomationException("Node {0} is not in the graph. It is a pseudodependency of {1}.", Dep, NodeToDo.Key);
+							}
+						}
+					}
+				}
+				NodesToDo.UnionWith(Fringe);
+			}
+		}
+
+		// find things that our nodes depend on
+		if (!bOnlyNode)
+		{
+			bool bDoneWithDependencies = false;
+
+			while (!bDoneWithDependencies)
+			{
+				bDoneWithDependencies = true;
+				var Fringe = new HashSet<string>();
+				foreach (var NodeToDo in NodesToDo)
+				{
+					if (BranchOptions.NodesToRemovePseudoDependencies.Count > 0)
+					{
+						var ListOfRetainedNodesToRemovePseudoDependencies = BranchOptions.NodesToRemovePseudoDependencies;
+						foreach (var NodeToRemovePseudoDependencies in BranchOptions.NodesToRemovePseudoDependencies)
+						{
+							if (NodeToDo == NodeToRemovePseudoDependencies)
+							{
+								RemoveAllPseudodependenciesFromNode(NodeToDo);
+								ListOfRetainedNodesToRemovePseudoDependencies.Remove(NodeToDo);
+								break;
+							}
+						}
+						BranchOptions.NodesToRemovePseudoDependencies = ListOfRetainedNodesToRemovePseudoDependencies;
+					}
+					foreach (var Dep in GUBPNodes[NodeToDo].FullNamesOfDependencies)
+					{
+						if (!GUBPNodes.ContainsKey(Dep))
+						{
+							throw new AutomationException("Node {0} is not in the graph. It is a dependency of {1}.", Dep, NodeToDo);
+						}
+						if (!NodesToDo.Contains(Dep))
+						{
+							Fringe.Add(Dep);
+							bDoneWithDependencies = false;
+						}
+					}
+					foreach (var Dep in GUBPNodes[NodeToDo].FullNamesOfPseudosependencies)
+					{
+						if (!GUBPNodes.ContainsKey(Dep))
+						{
+							throw new AutomationException("Node {0} is not in the graph. It is a pseudodependency of {1}.", Dep, NodeToDo);
+						}
+					}
+				}
+				NodesToDo.UnionWith(Fringe);
+			}
+		}
+		if (TimeIndex != 0)
+		{
+			LogVerbose("Culling based on time index");
+			var NewNodesToDo = new HashSet<string>();
+			foreach (var NodeToDo in NodesToDo)
+			{
+				if (TimeIndex % (1 << GUBPNodes[NodeToDo].DependentCISFrequencyQuantumShift()) == 0)
+				{
+					LogVerbose("  Keeping {0}", NodeToDo);
+					NewNodesToDo.Add(NodeToDo);
+				}
+				else
+				{
+					LogVerbose("  Rejecting {0}", NodeToDo);
+				}
+			}
+			NodesToDo = NewNodesToDo;
+		}
+		//Remove Plat if specified
+		if (WithoutLinux)
+		{
+			var NewNodesToDo = new HashSet<string>();
+			foreach (var NodeToDo in NodesToDo)
+			{
+				if (!GUBPNodes[NodeToDo].GetFullName().Contains("Linux"))
+				{
+					NewNodesToDo.Add(NodeToDo);
+				}
+				else
+				{
+					LogVerbose(" Rejecting {0} because -NoLinux was requested", NodeToDo);
+				}
+			}
+			NodesToDo = NewNodesToDo;
+		}
+		//find things that depend on our nodes and setup commander dictionary
+		if (!bOnlyNode)
+		{
+			foreach (var NodeToDo in NodesToDo)
+			{
+				if (!GUBPNodes[NodeToDo].IsAggregate() && !GUBPNodes[NodeToDo].IsTest())
+				{
+					List<string> ECDependencies = new List<string>();
+					ECDependencies = GetECDependencies(NodeToDo);
+					foreach (var Dep in ECDependencies)
+					{
+						if (!GUBPNodes.ContainsKey(Dep))
+						{
+							throw new AutomationException("Node {0} is not in the graph. It is a dependency of {1}.", Dep, NodeToDo);
+						}
+						if (!GUBPNodes[Dep].FullNamesOfDependedOn.Contains(NodeToDo))
+						{
+							GUBPNodes[Dep].FullNamesOfDependedOn.Add(NodeToDo);
+						}
+					}
+				}
+			}
+			foreach (var NodeToDo in NodesToDo)
+			{
+				var Deps = GUBPNodes[NodeToDo].FullNamesOfDependedOn;
+				string All = "";
+				foreach (var Dep in Deps)
+				{
+					if (All != "")
+					{
+						All += " ";
+					}
+					All += Dep;
+				}
+				FullNodeDependedOnBy.Add(NodeToDo, All);
+			}
+		}
+	}
+
 	private void GetFullNodeListAndDirectDependencies(Dictionary<string, string> GUBPNodesControllingTrigger, Dictionary<string, string> GUBPNodesControllingTriggerDotName, out Dictionary<string, string> FullNodeList, out Dictionary<string, string> FullNodeDirectDependencies)
 	{
 		FullNodeList = new Dictionary<string,string>();
