Move frequency barriers into a separate function.

diff --git a/Engine/Source/Programs/AutomationTool/GUBP/GUBP.cs b/Engine/Source/Programs/AutomationTool/GUBP/GUBP.cs
--- a/Engine/Source/Programs/AutomationTool/GUBP/GUBP.cs
+++ b/Engine/Source/Programs/AutomationTool/GUBP/GUBP.cs
@@ -1588,48 +1588,8 @@ public partial class GUBP : BuildCommand
                 }
             }
 
-			// Make sure that everything that's listed as a frequency barrier is completed with the given interval
-			Dictionary<string, int> FrequencyOverrides = new Dictionary<string,int>();
-			foreach (KeyValuePair<string, sbyte> Barrier in BranchOptions.FrequencyBarriers)
-			{
-				// All the nodes which are dependencies of the barrier node
-				HashSet<string> IncludedNodes = new HashSet<string> { Barrier.Key };
-
-				// Find all the nodes which are indirect dependencies of this node
-				List<string> SearchNodes = new List<string> { Barrier.Key };
-				for (int Idx = 0; Idx < SearchNodes.Count; Idx++)
-				{
-					GUBPNode Node = GUBPNodes[SearchNodes[Idx]];
-					foreach (string DependencyName in Node.FullNamesOfDependencies.Union(Node.FullNamesOfPseudosependencies))
-					{
-						if (!IncludedNodes.Contains(DependencyName))
-						{
-							IncludedNodes.Add(DependencyName);
-							SearchNodes.Add(DependencyName);
-						}
-					}
-				}
-
-				// Make sure that everything included in this list is before the cap, and everything not in the list is after it
-				foreach (KeyValuePair<string, GUBPNode> NodePair in GUBPNodes)
-				{
-					if (IncludedNodes.Contains(NodePair.Key))
-					{
-						int Frequency;
-						if(FrequencyOverrides.TryGetValue(NodePair.Key, out Frequency))
-						{
-							Frequency = Math.Min(Frequency, Barrier.Value);
-						}
-						else
-						{
-							Frequency = Barrier.Value;
-						}
-						FrequencyOverrides[NodePair.Key] = Frequency;
-					}
-				}
-			}
-
 			// Compute all the frequencies
+			Dictionary<string, int> FrequencyOverrides = ApplyFrequencyBarriers();
             foreach (var NodeToDo in GUBPNodes)
             {
                 ComputeDependentCISFrequencyQuantumShift(NodeToDo.Key, FrequencyOverrides);
@@ -1777,6 +1737,51 @@ public partial class GUBP : BuildCommand
 		ExecuteNodes(OrdereredToDo, bOnlyNode, bFakeEC, LocalOnly, bSaveSharedTempStorage, GUBPNodesCompleted, GUBPNodesHistory, CLString, FakeFail);
 	}
 
+	private Dictionary<string, int> ApplyFrequencyBarriers()
+	{
+		// Make sure that everything that's listed as a frequency barrier is completed with the given interval
+		Dictionary<string, int> FrequencyOverrides = new Dictionary<string, int>();
+		foreach (KeyValuePair<string, sbyte> Barrier in BranchOptions.FrequencyBarriers)
+		{
+			// All the nodes which are dependencies of the barrier node
+			HashSet<string> IncludedNodes = new HashSet<string> { Barrier.Key };
+
+			// Find all the nodes which are indirect dependencies of this node
+			List<string> SearchNodes = new List<string> { Barrier.Key };
+			for (int Idx = 0; Idx < SearchNodes.Count; Idx++)
+			{
+				GUBPNode Node = GUBPNodes[SearchNodes[Idx]];
+				foreach (string DependencyName in Node.FullNamesOfDependencies.Union(Node.FullNamesOfPseudosependencies))
+				{
+					if (!IncludedNodes.Contains(DependencyName))
+					{
+						IncludedNodes.Add(DependencyName);
+						SearchNodes.Add(DependencyName);
+					}
+				}
+			}
+
+			// Make sure that everything included in this list is before the cap, and everything not in the list is after it
+			foreach (KeyValuePair<string, GUBPNode> NodePair in GUBPNodes)
+			{
+				if (IncludedNodes.Contains(NodePair.Key))
+				{
+					int Frequency;
+					if (FrequencyOverrides.TryGetValue(NodePair.Key, out Frequency))
+					{
+						Frequency = Math.Min(Frequency, Barrier.Value);
+					}
+					else
+					{
+						Frequency = Barrier.Value;
+					}
+					FrequencyOverrides[NodePair.Key] = Frequency;
+				}
+			}
+		}
+		return FrequencyOverrides;
+	}
+
 	private void ParseNodesToDo(bool WithoutLinux, int TimeIndex, bool CommanderSetup, Dictionary<string, string> FullNodeDependedOnBy, out bool bOnlyNode, out HashSet<string> NodesToDo)
 	{
 		bOnlyNode = false;
