Use a wrapper class for GUBP node when building up the graph, so we can store cached data on the object itself (rather than having lots of separate dictionaries).

diff --git a/Engine/Source/Programs/AutomationTool/GUBP/GUBP.cs b/Engine/Source/Programs/AutomationTool/GUBP/GUBP.cs
--- a/Engine/Source/Programs/AutomationTool/GUBP/GUBP.cs
+++ b/Engine/Source/Programs/AutomationTool/GUBP/GUBP.cs
@@ -28,7 +28,12 @@ public partial class GUBP : BuildCommand
     static public string PreflightMangleSuffix = "";
     public GUBPBranchHacker.BranchOptions BranchOptions = null;	
 
-    Dictionary<string, GUBPNode> GUBPNodes;
+    Dictionary<string, NodeInfo> GUBPNodes;
+
+	class NodeInfo
+	{
+		public GUBPNode Node;
+	}
 
     class NodeHistory
     {
@@ -50,7 +55,7 @@ public partial class GUBP : BuildCommand
         {
             throw new AutomationException("Attempt to add a duplicate node {0}", Node.GetFullName());
         }
-        GUBPNodes.Add(Name, Node);
+        GUBPNodes.Add(Name, new NodeInfo{ Node = Node });
         return Name;
     }
 
@@ -61,14 +66,14 @@ public partial class GUBP : BuildCommand
 
     public GUBPNode FindNode(string Node)
     {
-        return GUBPNodes[Node];
+        return GUBPNodes[Node].Node;
     }
 
 	public GUBPNode TryFindNode(string Node)
 	{
-		GUBPNode Result;
+		NodeInfo Result;
 		GUBPNodes.TryGetValue(Node, out Result);
-		return Result;
+		return (Result == null)? null : Result.Node;
 	}
 
     public void RemovePseudodependencyFromNode(string Node, string Dep)
@@ -77,7 +82,7 @@ public partial class GUBP : BuildCommand
         {
             throw new AutomationException("Node {0} not found", Node);
         }
-        GUBPNodes[Node].RemovePseudodependency(Dep);
+        GUBPNodes[Node].Node.RemovePseudodependency(Dep);
     }
 
     public void RemoveAllPseudodependenciesFromNode(string Node)
@@ -86,15 +91,15 @@ public partial class GUBP : BuildCommand
         {
             throw new AutomationException("Node {0} not found", Node);
         }
-        GUBPNodes[Node].FullNamesOfPseudosependencies.Clear();
+        GUBPNodes[Node].Node.FullNamesOfPseudosependencies.Clear();
     }
 
     List<string> GetDependencies(string NodeToDo, bool bFlat = false, bool ECOnly = false)
     {
         var Result = new List<string>();
-        foreach (var Node in GUBPNodes[NodeToDo].FullNamesOfDependencies)
+        foreach (var Node in GUBPNodes[NodeToDo].Node.FullNamesOfDependencies)
         {
-            bool Usable = GUBPNodes[Node].RunInEC() || !ECOnly;
+            bool Usable = GUBPNodes[Node].Node.RunInEC() || !ECOnly;
             if (Usable)
             {
                 if (!Result.Contains(Node))
@@ -113,9 +118,9 @@ public partial class GUBP : BuildCommand
                 }
             }
         }
-        foreach (var Node in GUBPNodes[NodeToDo].FullNamesOfPseudosependencies)
+        foreach (var Node in GUBPNodes[NodeToDo].Node.FullNamesOfPseudosependencies)
         {
-            bool Usable = GUBPNodes[Node].RunInEC() || !ECOnly;
+            bool Usable = GUBPNodes[Node].Node.RunInEC() || !ECOnly;
             if (Usable)
             {
                 if (!Result.Contains(Node))
@@ -140,9 +145,9 @@ public partial class GUBP : BuildCommand
 	List<string> GetCompletedOnlyDependencies(string NodeToDo, bool bFlat = false, bool ECOnly = true)
 	{
 		var Result = new List<string>();
-		foreach (var Node in GUBPNodes[NodeToDo].CompletedDependencies)
+		foreach (var Node in GUBPNodes[NodeToDo].Node.CompletedDependencies)
 		{
-			bool Usable = GUBPNodes[Node].RunInEC() || !ECOnly;
+			bool Usable = GUBPNodes[Node].Node.RunInEC() || !ECOnly;
 			if(Usable)
 			{
 				if(!Result.Contains(Node))
@@ -180,13 +185,13 @@ public partial class GUBP : BuildCommand
             return GUBPNodesControllingTrigger[NodeToDo];
         }
         var Result = "";
-        foreach (var Node in GUBPNodes[NodeToDo].FullNamesOfDependencies)
+        foreach (var Node in GUBPNodes[NodeToDo].Node.FullNamesOfDependencies)
         {
             if (!GUBPNodes.ContainsKey(Node))
             {
                 throw new AutomationException("Dependency {0} in {1} not found.", Node, NodeToDo);
             }
-            bool IsTrigger = GUBPNodes[Node].TriggerNode();
+            bool IsTrigger = GUBPNodes[Node].Node.TriggerNode();
             if (IsTrigger)
             {
                 if (Node != Result && !string.IsNullOrEmpty(Result))
@@ -208,13 +213,13 @@ public partial class GUBP : BuildCommand
                 }
             }
         }
-        foreach (var Node in GUBPNodes[NodeToDo].FullNamesOfPseudosependencies)
+        foreach (var Node in GUBPNodes[NodeToDo].Node.FullNamesOfPseudosependencies)
         {
             if (!GUBPNodes.ContainsKey(Node))
             {
                 throw new AutomationException("Pseudodependency {0} in {1} not found.", Node, NodeToDo);
             }
-            bool IsTrigger = GUBPNodes[Node].TriggerNode();
+            bool IsTrigger = GUBPNodes[Node].Node.TriggerNode();
             if (IsTrigger)
             {
                 if (Node != Result && !string.IsNullOrEmpty(Result))
@@ -268,7 +273,7 @@ public partial class GUBP : BuildCommand
     public string CISFrequencyQuantumShiftString(string NodeToDo)
     {
         string FrequencyString = "";
-        int Quantum = GUBPNodes[NodeToDo].DependentCISFrequencyQuantumShift();
+        int Quantum = GUBPNodes[NodeToDo].Node.DependentCISFrequencyQuantumShift();
         if (Quantum > 0)
         {
 			int TimeQuantum = 20;
@@ -291,11 +296,11 @@ public partial class GUBP : BuildCommand
 
     public int ComputeDependentCISFrequencyQuantumShift(string NodeToDo, Dictionary<string, int> FrequencyOverrides)
     {
-        int Result = GUBPNodes[NodeToDo].ComputedDependentCISFrequencyQuantumShift;        
+        int Result = GUBPNodes[NodeToDo].Node.ComputedDependentCISFrequencyQuantumShift;        
         if (Result < 0)
         {
-            Result = GUBPNodes[NodeToDo].CISFrequencyQuantumShift(this);
-            Result = GetFrequencyForNode(this, GUBPNodes[NodeToDo].GetFullName(), Result);
+            Result = GUBPNodes[NodeToDo].Node.CISFrequencyQuantumShift(this);
+            Result = GetFrequencyForNode(this, GUBPNodes[NodeToDo].Node.GetFullName(), Result);
 
 			int FrequencyOverride;
 			if(FrequencyOverrides.TryGetValue(NodeToDo, out FrequencyOverride) && Result > FrequencyOverride)
@@ -303,15 +308,15 @@ public partial class GUBP : BuildCommand
 				Result = FrequencyOverride;
 			}
 
-            foreach (var Dep in GUBPNodes[NodeToDo].FullNamesOfDependencies)
+            foreach (var Dep in GUBPNodes[NodeToDo].Node.FullNamesOfDependencies)
             {
                 Result = Math.Max(ComputeDependentCISFrequencyQuantumShift(Dep, FrequencyOverrides), Result);
             }
-            foreach (var Dep in GUBPNodes[NodeToDo].FullNamesOfPseudosependencies)
+            foreach (var Dep in GUBPNodes[NodeToDo].Node.FullNamesOfPseudosependencies)
             {
                 Result = Math.Max(ComputeDependentCISFrequencyQuantumShift(Dep, FrequencyOverrides), Result);
             }
-			foreach (var Dep in GUBPNodes[NodeToDo].CompletedDependencies)
+			foreach (var Dep in GUBPNodes[NodeToDo].Node.CompletedDependencies)
 			{
 				Result = Math.Max(ComputeDependentCISFrequencyQuantumShift(Dep, FrequencyOverrides), Result);
 			}
@@ -319,7 +324,7 @@ public partial class GUBP : BuildCommand
             {
                 throw new AutomationException("Failed to compute shift.");
             }
-            GUBPNodes[NodeToDo].ComputedDependentCISFrequencyQuantumShift = Result;
+            GUBPNodes[NodeToDo].Node.ComputedDependentCISFrequencyQuantumShift = Result;
         }
         return Result;
     }
@@ -330,8 +335,8 @@ public partial class GUBP : BuildCommand
         {
             return GUBPNodesCompleted[NodeToDo];
         }
-        string NodeStoreName = StoreName + "-" + GUBPNodes[NodeToDo].GetFullName();
-        string GameNameIfAny = GUBPNodes[NodeToDo].GameNameIfAnyForTempStorage();
+        string NodeStoreName = StoreName + "-" + GUBPNodes[NodeToDo].Node.GetFullName();
+        string GameNameIfAny = GUBPNodes[NodeToDo].Node.GameNameIfAnyForTempStorage();
         bool Result;
         if (LocalOnly)
         {
@@ -347,11 +352,11 @@ public partial class GUBP : BuildCommand
         }
         if (Result)
         {
-            LogVerbose("***** GUBP Trigger Node was already triggered {0} -> {1} : {2}", GUBPNodes[NodeToDo].GetFullName(), GameNameIfAny, NodeStoreName);
+            LogVerbose("***** GUBP Trigger Node was already triggered {0} -> {1} : {2}", GUBPNodes[NodeToDo].Node.GetFullName(), GameNameIfAny, NodeStoreName);
         }
         else
         {
-            LogVerbose("***** GUBP Trigger Node was NOT yet triggered {0} -> {1} : {2}", GUBPNodes[NodeToDo].GetFullName(), GameNameIfAny, NodeStoreName);
+            LogVerbose("***** GUBP Trigger Node was NOT yet triggered {0} -> {1} : {2}", GUBPNodes[NodeToDo].Node.GetFullName(), GameNameIfAny, NodeStoreName);
         }
         GUBPNodesCompleted.Add(NodeToDo, Result);
         return Result;
@@ -508,7 +513,7 @@ public partial class GUBP : BuildCommand
         string LastAgentGroup = "";
         foreach (var NodeToDo in Nodes)
         {
-            if (ECOnly && !GUBPNodes[NodeToDo].RunInEC())
+            if (ECOnly && !GUBPNodes[NodeToDo].Node.RunInEC())
             {
                 continue;
             }
@@ -542,14 +547,14 @@ public partial class GUBP : BuildCommand
                     }
                 }
             }
-            if (GUBPNodes[NodeToDo].AgentSharingGroup != LastAgentGroup && GUBPNodes[NodeToDo].AgentSharingGroup != "")
+            if (GUBPNodes[NodeToDo].Node.AgentSharingGroup != LastAgentGroup && GUBPNodes[NodeToDo].Node.AgentSharingGroup != "")
             {
-                Log("    Agent Group: {0}", GUBPNodes[NodeToDo].AgentSharingGroup);
+                Log("    Agent Group: {0}", GUBPNodes[NodeToDo].Node.AgentSharingGroup);
             }
-            LastAgentGroup = GUBPNodes[NodeToDo].AgentSharingGroup;
+            LastAgentGroup = GUBPNodes[NodeToDo].Node.AgentSharingGroup;
 
-            string Agent = GUBPNodes[NodeToDo].ECAgentString();
-			if(ParseParamValue("AgentOverride") != "" && !GUBPNodes[NodeToDo].GetFullName().Contains("Mac"))
+            string Agent = GUBPNodes[NodeToDo].Node.ECAgentString();
+			if(ParseParamValue("AgentOverride") != "" && !GUBPNodes[NodeToDo].Node.GetFullName().Contains("Mac"))
 			{
 				Agent = ParseParamValue("AgentOverride");
 			}
@@ -557,7 +562,7 @@ public partial class GUBP : BuildCommand
             {
                 Agent = "[" + Agent + "]";
             }
-            string MemoryReq = "[" + GUBPNodes[NodeToDo].AgentMemoryRequirement(bp).ToString() + "]";            
+            string MemoryReq = "[" + GUBPNodes[NodeToDo].Node.AgentMemoryRequirement(bp).ToString() + "]";            
             if(MemoryReq == "[0]")
             {
                 MemoryReq = "";                
@@ -569,12 +574,12 @@ public partial class GUBP : BuildCommand
                 NodeToDo,
                 FrequencyString,
                 NodeIsAlreadyComplete(GUBPNodesCompleted, NodeToDo, LocalOnly) ? " - (Completed)" : "",
-                GUBPNodes[NodeToDo].TriggerNode() ? " - (TriggerNode)" : "",
-                GUBPNodes[NodeToDo].IsSticky() ? " - (Sticky)" : "",				
+                GUBPNodes[NodeToDo].Node.TriggerNode() ? " - (TriggerNode)" : "",
+                GUBPNodes[NodeToDo].Node.IsSticky() ? " - (Sticky)" : "",				
                 Agent,
                 MemoryReq,
                 EMails,
-                ECProc ? GUBPNodes[NodeToDo].ECProcedure() : ""
+                ECProc ? GUBPNodes[NodeToDo].Node.ECProcedure() : ""
                 );
             if (bShowHistory && GUBPNodesHistory.ContainsKey(NodeToDo))
             {
@@ -594,15 +599,15 @@ public partial class GUBP : BuildCommand
             }
             if (bShowDependencies)
             {
-                foreach (var Dep in GUBPNodes[NodeToDo].FullNamesOfDependencies)
+                foreach (var Dep in GUBPNodes[NodeToDo].Node.FullNamesOfDependencies)
                 {
                     Log("            dep> {0}", Dep);
                 }
-                foreach (var Dep in GUBPNodes[NodeToDo].FullNamesOfPseudosependencies)
+                foreach (var Dep in GUBPNodes[NodeToDo].Node.FullNamesOfPseudosependencies)
                 {
                     Log("           pdep> {0}", Dep);
                 }
-				foreach (var Dep in GUBPNodes[NodeToDo].CompletedDependencies)
+				foreach (var Dep in GUBPNodes[NodeToDo].Node.CompletedDependencies)
 				{
 					Log("			cdep>{0}", Dep);
 				}
@@ -621,14 +626,14 @@ public partial class GUBP : BuildCommand
 
 			if(bShowDependednOn)
 			{
-				foreach (var Dep in GUBPNodes[NodeToDo].FullNamesOfDependedOn)
+				foreach (var Dep in GUBPNodes[NodeToDo].Node.FullNamesOfDependedOn)
 				{
 					Log("            depOn> {0}", Dep);
 				}
 			}
 			if (bShowDependentPromotions)
 			{
-				foreach (var Dep in GUBPNodes[NodeToDo].DependentPromotions)
+				foreach (var Dep in GUBPNodes[NodeToDo].Node.DependentPromotions)
 				{
 					Log("            depPro> {0}", Dep);
 				}
@@ -662,7 +667,7 @@ public partial class GUBP : BuildCommand
             var Node = Nodes[NodeIndex];
             GraphNode NodeGraphNode = NodeToGraphNodeMap[Node];
 
-            foreach (var Dep in GUBPNodes[Node].FullNamesOfDependencies)
+            foreach (var Dep in GUBPNodes[Node].Node.FullNamesOfDependencies)
             {
                 GraphNode PrerequisiteFileGraphNode;
                 if (NodeToGraphNodeMap.TryGetValue(Dep, out PrerequisiteFileGraphNode))
@@ -680,7 +685,7 @@ public partial class GUBP : BuildCommand
                 }
 
             }
-            foreach (var Dep in GUBPNodes[Node].FullNamesOfPseudosependencies)
+            foreach (var Dep in GUBPNodes[Node].Node.FullNamesOfPseudosependencies)
             {
                 GraphNode PrerequisiteFileGraphNode;
                 if (NodeToGraphNodeMap.TryGetValue(Dep, out PrerequisiteFileGraphNode))
@@ -792,7 +797,7 @@ public partial class GUBP : BuildCommand
             var AgentGroupChains = new Dictionary<string, List<string>>();
             foreach (var NodeToDo in NodesToDo)
             {
-                string MyAgentGroup = GUBPNodes[NodeToDo].AgentSharingGroup;
+                string MyAgentGroup = GUBPNodes[NodeToDo].Node.AgentSharingGroup;
                 if (MyAgentGroup != "")
                 {
                     if (!AgentGroupChains.ContainsKey(MyAgentGroup))
@@ -827,25 +832,25 @@ public partial class GUBP : BuildCommand
                 bool bReady = true;
 				bool bPseudoReady = true;
 				bool bCompReady = true;
-                if (!SubSort && GUBPNodes[NodeToDo].AgentSharingGroup != "")
+                if (!SubSort && GUBPNodes[NodeToDo].Node.AgentSharingGroup != "")
                 {
-                    if (ExaminedAgentGroups.Contains(GUBPNodes[NodeToDo].AgentSharingGroup))
+                    if (ExaminedAgentGroups.Contains(GUBPNodes[NodeToDo].Node.AgentSharingGroup))
                     {
-                        bReady = !NonReadyAgentGroups.Contains(GUBPNodes[NodeToDo].AgentSharingGroup);
-                        bPseudoReady = !NonPeudoReadyAgentGroups.Contains(GUBPNodes[NodeToDo].AgentSharingGroup); //this might not be accurate if bReady==false
+                        bReady = !NonReadyAgentGroups.Contains(GUBPNodes[NodeToDo].Node.AgentSharingGroup);
+                        bPseudoReady = !NonPeudoReadyAgentGroups.Contains(GUBPNodes[NodeToDo].Node.AgentSharingGroup); //this might not be accurate if bReady==false
                     }
                     else
                     {
-                        ExaminedAgentGroups.Add(GUBPNodes[NodeToDo].AgentSharingGroup);
-                        foreach (var ChainNode in SortedAgentGroupChains[GUBPNodes[NodeToDo].AgentSharingGroup])
+                        ExaminedAgentGroups.Add(GUBPNodes[NodeToDo].Node.AgentSharingGroup);
+                        foreach (var ChainNode in SortedAgentGroupChains[GUBPNodes[NodeToDo].Node.AgentSharingGroup])
                         {
-                            foreach (var Dep in GUBPNodes[ChainNode].FullNamesOfDependencies)
+                            foreach (var Dep in GUBPNodes[ChainNode].Node.FullNamesOfDependencies)
                             {
                                 if (!GUBPNodes.ContainsKey(Dep))
                                 {
                                     throw new AutomationException("Dependency {0} node found.", Dep);
                                 }
-                                if (!SortedAgentGroupChains[GUBPNodes[NodeToDo].AgentSharingGroup].Contains(Dep) && NodesToDo.Contains(Dep))
+                                if (!SortedAgentGroupChains[GUBPNodes[NodeToDo].Node.AgentSharingGroup].Contains(Dep) && NodesToDo.Contains(Dep))
                                 {
                                     bReady = false;
                                     break;
@@ -853,19 +858,19 @@ public partial class GUBP : BuildCommand
                             }
                             if (!bReady)
                             {
-                                NonReadyAgentGroups.Add(GUBPNodes[NodeToDo].AgentSharingGroup);
+                                NonReadyAgentGroups.Add(GUBPNodes[NodeToDo].Node.AgentSharingGroup);
                                 break;
                             }
-                            foreach (var Dep in GUBPNodes[ChainNode].FullNamesOfPseudosependencies)
+                            foreach (var Dep in GUBPNodes[ChainNode].Node.FullNamesOfPseudosependencies)
                             {
                                 if (!GUBPNodes.ContainsKey(Dep))
                                 {
                                     throw new AutomationException("Pseudodependency {0} node found.", Dep);
                                 }
-                                if (!SortedAgentGroupChains[GUBPNodes[NodeToDo].AgentSharingGroup].Contains(Dep) && NodesToDo.Contains(Dep))
+                                if (!SortedAgentGroupChains[GUBPNodes[NodeToDo].Node.AgentSharingGroup].Contains(Dep) && NodesToDo.Contains(Dep))
                                 {
                                     bPseudoReady = false;
-                                    NonPeudoReadyAgentGroups.Add(GUBPNodes[NodeToDo].AgentSharingGroup);
+                                    NonPeudoReadyAgentGroups.Add(GUBPNodes[NodeToDo].Node.AgentSharingGroup);
                                     break;
                                 }
                             }
@@ -874,7 +879,7 @@ public partial class GUBP : BuildCommand
                 }
                 else
                 {
-                    foreach (var Dep in GUBPNodes[NodeToDo].FullNamesOfDependencies)
+                    foreach (var Dep in GUBPNodes[NodeToDo].Node.FullNamesOfDependencies)
                     {
                         if (!GUBPNodes.ContainsKey(Dep))
                         {
@@ -886,7 +891,7 @@ public partial class GUBP : BuildCommand
                             break;
                         }
                     }
-                    foreach (var Dep in GUBPNodes[NodeToDo].FullNamesOfPseudosependencies)
+                    foreach (var Dep in GUBPNodes[NodeToDo].Node.FullNamesOfPseudosependencies)
                     {
                         if (!GUBPNodes.ContainsKey(Dep))
                         {
@@ -898,7 +903,7 @@ public partial class GUBP : BuildCommand
                             break;
                         }
                     }
-					foreach (var Dep in GUBPNodes[NodeToDo].CompletedDependencies)
+					foreach (var Dep in GUBPNodes[NodeToDo].Node.CompletedDependencies)
 					{
 						if (!GUBPNodes.ContainsKey(Dep))
 						{
@@ -911,7 +916,7 @@ public partial class GUBP : BuildCommand
 						}
 					}
                 }
-                var Priority = GUBPNodes[NodeToDo].Priority();
+                var Priority = GUBPNodes[NodeToDo].Node.Priority();
 
                 if (bReady && BestNode != "")
                 {
@@ -921,11 +926,11 @@ public partial class GUBP : BuildCommand
                     }
                     else if (String.Compare(GetControllingTriggerDotName(BestNode, GUBPNodesControllingTriggerDotName, GUBPNodesControllingTrigger), GetControllingTriggerDotName(NodeToDo, GUBPNodesControllingTriggerDotName, GUBPNodesControllingTrigger)) == 0) //sorted by controlling trigger
                     {
-                        if (GUBPNodes[BestNode].IsSticky() && !GUBPNodes[NodeToDo].IsSticky()) //sticky nodes first
+                        if (GUBPNodes[BestNode].Node.IsSticky() && !GUBPNodes[NodeToDo].Node.IsSticky()) //sticky nodes first
                         {
                             bReady = false;
                         }
-                        else if (GUBPNodes[BestNode].IsSticky() == GUBPNodes[NodeToDo].IsSticky())
+                        else if (GUBPNodes[BestNode].Node.IsSticky() == GUBPNodes[NodeToDo].Node.IsSticky())
                         {
                             if (BestPseudoReady && !bPseudoReady)
                             {
@@ -933,8 +938,8 @@ public partial class GUBP : BuildCommand
                             }
                             else if (BestPseudoReady == bPseudoReady)
                             {
-                                bool IamLateTrigger = !DoNotConsiderCompletion && GUBPNodes[NodeToDo].TriggerNode() && NodeToDo != ExplicitTrigger && !NodeIsAlreadyComplete(GUBPNodesCompleted, NodeToDo, LocalOnly);
-                                bool BestIsLateTrigger = !DoNotConsiderCompletion && GUBPNodes[BestNode].TriggerNode() && BestNode != ExplicitTrigger && !NodeIsAlreadyComplete(GUBPNodesCompleted, BestNode, LocalOnly);
+                                bool IamLateTrigger = !DoNotConsiderCompletion && GUBPNodes[NodeToDo].Node.TriggerNode() && NodeToDo != ExplicitTrigger && !NodeIsAlreadyComplete(GUBPNodesCompleted, NodeToDo, LocalOnly);
+                                bool BestIsLateTrigger = !DoNotConsiderCompletion && GUBPNodes[BestNode].Node.TriggerNode() && BestNode != ExplicitTrigger && !NodeIsAlreadyComplete(GUBPNodesCompleted, BestNode, LocalOnly);
                                 if (BestIsLateTrigger && !IamLateTrigger)
                                 {
                                     bReady = false;
@@ -972,9 +977,9 @@ public partial class GUBP : BuildCommand
             }
             if (bProgressMade)
             {
-                if (!SubSort && GUBPNodes[BestNode].AgentSharingGroup != "")
+                if (!SubSort && GUBPNodes[BestNode].Node.AgentSharingGroup != "")
                 {
-                    foreach (var ChainNode in SortedAgentGroupChains[GUBPNodes[BestNode].AgentSharingGroup])
+                    foreach (var ChainNode in SortedAgentGroupChains[GUBPNodes[BestNode].Node.AgentSharingGroup])
                     {
                         OrdereredToDo.Add(ChainNode);
                         NodesToDo.Remove(ChainNode);
@@ -993,27 +998,27 @@ public partial class GUBP : BuildCommand
                 foreach (var NodeToDo in NodesToDo)
                 {
                     string Deps = "";
-                    if (!SubSort && GUBPNodes[NodeToDo].AgentSharingGroup != "")
+                    if (!SubSort && GUBPNodes[NodeToDo].Node.AgentSharingGroup != "")
                     {
-                        foreach (var ChainNode in SortedAgentGroupChains[GUBPNodes[NodeToDo].AgentSharingGroup])
+                        foreach (var ChainNode in SortedAgentGroupChains[GUBPNodes[NodeToDo].Node.AgentSharingGroup])
                         {
-                            foreach (var Dep in GUBPNodes[ChainNode].FullNamesOfDependencies)
+                            foreach (var Dep in GUBPNodes[ChainNode].Node.FullNamesOfDependencies)
                             {
-                                if (!SortedAgentGroupChains[GUBPNodes[NodeToDo].AgentSharingGroup].Contains(Dep) && NodesToDo.Contains(Dep))
+                                if (!SortedAgentGroupChains[GUBPNodes[NodeToDo].Node.AgentSharingGroup].Contains(Dep) && NodesToDo.Contains(Dep))
                                 {
-                                    Deps = Deps + Dep + "[" + ChainNode + "->" + GUBPNodes[NodeToDo].AgentSharingGroup + "]" + " ";
+                                    Deps = Deps + Dep + "[" + ChainNode + "->" + GUBPNodes[NodeToDo].Node.AgentSharingGroup + "]" + " ";
                                 }
                             }
                         }
                     }
-                    foreach (var Dep in GUBPNodes[NodeToDo].FullNamesOfDependencies)
+                    foreach (var Dep in GUBPNodes[NodeToDo].Node.FullNamesOfDependencies)
                     {
                         if (NodesToDo.Contains(Dep))
                         {
                             Deps = Deps + Dep + " ";
                         }
                     }
-                    foreach (var Dep in GUBPNodes[NodeToDo].FullNamesOfPseudosependencies)
+                    foreach (var Dep in GUBPNodes[NodeToDo].Node.FullNamesOfPseudosependencies)
                     {
                         if (NodesToDo.Contains(Dep))
                         {
@@ -1036,9 +1041,9 @@ public partial class GUBP : BuildCommand
 
     string GetJobStepPath(string Dep)
     {
-        if (Dep != "Noop" && GUBPNodes[Dep].AgentSharingGroup != "")
+        if (Dep != "Noop" && GUBPNodes[Dep].Node.AgentSharingGroup != "")
         {
-            return "jobSteps[" + GUBPNodes[Dep].AgentSharingGroup + "]/jobSteps[" + Dep + "]";
+            return "jobSteps[" + GUBPNodes[Dep].Node.AgentSharingGroup + "]/jobSteps[" + Dep + "]";
         }
         return "jobSteps[" + Dep + "]";
     }
@@ -1049,10 +1054,10 @@ public partial class GUBP : BuildCommand
 
     void UpdateNodeHistory(Dictionary<string, NodeHistory> GUBPNodesHistory, string Node, string CLString)
     {
-        if (GUBPNodes[Node].RunInEC() && !GUBPNodes[Node].TriggerNode() && CLString != "")
+        if (GUBPNodes[Node].Node.RunInEC() && !GUBPNodes[Node].Node.TriggerNode() && CLString != "")
         {
-            string GameNameIfAny = GUBPNodes[Node].GameNameIfAnyForTempStorage();
-            string NodeStoreWildCard = StoreName.Replace(CLString, "*") + "-" + GUBPNodes[Node].GetFullName();
+            string GameNameIfAny = GUBPNodes[Node].Node.GameNameIfAnyForTempStorage();
+            string NodeStoreWildCard = StoreName.Replace(CLString, "*") + "-" + GUBPNodes[Node].Node.GetFullName();
             var History = new NodeHistory();
 
             History.AllStarted = ConvertCLToIntList(TempStorage.FindTempStorageManifests(CmdEnv, NodeStoreWildCard + StartedTempStorageSuffix, false, true, GameNameIfAny));
@@ -1173,7 +1178,7 @@ public partial class GUBP : BuildCommand
         {
             var AdditionalEmails = "";
 			string Causers = "";
-            if (ParseParam("CIS") && !GUBPNodes[NodeToDo].SendSuccessEmail() && !GUBPNodes[NodeToDo].TriggerNode())
+            if (ParseParam("CIS") && !GUBPNodes[NodeToDo].Node.SendSuccessEmail() && !GUBPNodes[NodeToDo].Node.TriggerNode())
             {
 				Causers = FailCauserEMails;
            }
@@ -1185,7 +1190,7 @@ public partial class GUBP : BuildCommand
 			EMails = GetEMailListForNode(this, NodeToDo, AdditionalEmails, Causers);
 			RunECTool(String.Format("setProperty \"/myWorkflow/FailEmails/{0}\" \"{1}\"", NodeToDo, EMails));            
         }
-		if (GUBPNodes[NodeToDo].SendSuccessEmail() || SendSuccessForGreenAfterRed)
+		if (GUBPNodes[NodeToDo].Node.SendSuccessEmail() || SendSuccessForGreenAfterRed)
 		{
 			RunECTool(String.Format("setProperty \"/myWorkflow/SendSuccessEmail/{0}\" \"{1}\"", NodeToDo, "1"));
 		}
@@ -1223,8 +1228,8 @@ public partial class GUBP : BuildCommand
         var History = GUBPNodesHistory[NodeToDo];
         int Result = P4Env.Changelist;
 
-        string GameNameIfAny = GUBPNodes[NodeToDo].GameNameIfAnyForTempStorage();
-        string NodeStore = StoreName + "-" + GUBPNodes[NodeToDo].GetFullName() + FailedTempStorageSuffix;
+        string GameNameIfAny = GUBPNodes[NodeToDo].Node.GameNameIfAnyForTempStorage();
+        string NodeStore = StoreName + "-" + GUBPNodes[NodeToDo].Node.GetFullName() + FailedTempStorageSuffix;
 
         var BackwardsFails = new List<int>(History.AllFailed);
         BackwardsFails.Add(P4Env.Changelist);
@@ -1294,14 +1299,14 @@ public partial class GUBP : BuildCommand
 	
 		if (!OnlyLateUpdates)
 		{
-			string AgentReq = GUBPNodes[NodeToDo].ECAgentString();
-			if(ParseParamValue("AgentOverride") != "" && !GUBPNodes[NodeToDo].GetFullName().Contains("OnMac"))
+			string AgentReq = GUBPNodes[NodeToDo].Node.ECAgentString();
+			if(ParseParamValue("AgentOverride") != "" && !GUBPNodes[NodeToDo].Node.GetFullName().Contains("OnMac"))
 			{
 				AgentReq = ParseParamValue("AgentOverride");
 			}
 			ECProps.Add(string.Format("AgentRequirementString/{0}={1}", NodeToDo, AgentReq));
-			ECProps.Add(string.Format("RequiredMemory/{0}={1}", NodeToDo, GUBPNodes[NodeToDo].AgentMemoryRequirement(this)));
-			ECProps.Add(string.Format("Timeouts/{0}={1}", NodeToDo, GUBPNodes[NodeToDo].TimeoutInMinutes()));
+			ECProps.Add(string.Format("RequiredMemory/{0}={1}", NodeToDo, GUBPNodes[NodeToDo].Node.AgentMemoryRequirement(this)));
+			ECProps.Add(string.Format("Timeouts/{0}={1}", NodeToDo, GUBPNodes[NodeToDo].Node.TimeoutInMinutes()));
 			ECProps.Add(string.Format("JobStepPath/{0}={1}", NodeToDo, GetJobStepPath(NodeToDo)));
 		}
 		
@@ -1545,7 +1550,7 @@ public partial class GUBP : BuildCommand
 		LogVerbose("************************* TimeIndex:           		    {0}", TimeIndex);        
 
 
-        GUBPNodes = new Dictionary<string, GUBPNode>();        
+        GUBPNodes = new Dictionary<string, NodeInfo>();        
         Branch = new BranchInfo(HostPlatforms);        
 
 		AddNodesForBranch(TimeIndex, bNoAutomatedTesting);
@@ -1566,9 +1571,9 @@ public partial class GUBP : BuildCommand
         {
             foreach (var NodeToDo in GUBPNodes)
             {
-                if (GUBPNodes[NodeToDo.Key].IsSeparatePromotable())
+                if (GUBPNodes[NodeToDo.Key].Node.IsSeparatePromotable())
                 {
-                    SeparatePromotables.Add(GUBPNodes[NodeToDo.Key].GetFullName());
+                    SeparatePromotables.Add(GUBPNodes[NodeToDo.Key].Node.GetFullName());
                     List<string> Dependencies = new List<string>();
                     Dependencies = GetECDependencies(NodeToDo.Key);
                     foreach (var Dep in Dependencies)
@@ -1577,11 +1582,11 @@ public partial class GUBP : BuildCommand
                         {
                             throw new AutomationException("Node {0} is not in the graph.  It is a dependency of {1}.", Dep, NodeToDo);
                         }
-                        if (!GUBPNodes[Dep].IsPromotableAggregate())
+                        if (!GUBPNodes[Dep].Node.IsPromotableAggregate())
                         {
-                            if (!GUBPNodes[Dep].DependentPromotions.Contains(NodeToDo.Key))
+                            if (!GUBPNodes[Dep].Node.DependentPromotions.Contains(NodeToDo.Key))
                             {
-                                GUBPNodes[Dep].DependentPromotions.Add(NodeToDo.Key);
+                                GUBPNodes[Dep].Node.DependentPromotions.Add(NodeToDo.Key);
                             }
                         }
                     }
@@ -1597,7 +1602,7 @@ public partial class GUBP : BuildCommand
 
             foreach (var NodeToDo in GUBPNodes)
             {
-                var Deps = GUBPNodes[NodeToDo.Key].DependentPromotions;
+                var Deps = GUBPNodes[NodeToDo.Key].Node.DependentPromotions;
                 string All = "";
                 foreach (var Dep in Deps)
                 {
@@ -1628,11 +1633,11 @@ public partial class GUBP : BuildCommand
                 bool bFoundIt = false;
                 foreach (var Node in GUBPNodes)
                 {
-                    if (Node.Value.GetFullName().Equals(ExplicitTrigger, StringComparison.InvariantCultureIgnoreCase))
+                    if (Node.Value.Node.GetFullName().Equals(ExplicitTrigger, StringComparison.InvariantCultureIgnoreCase))
                     {
-                        if (Node.Value.TriggerNode() && Node.Value.RunInEC())
+                        if (Node.Value.Node.TriggerNode() && Node.Value.Node.RunInEC())
                         {
-                            Node.Value.SetAsExplicitTrigger();
+                            Node.Value.Node.SetAsExplicitTrigger();
                             bFoundIt = true;
                             break;
                         }
@@ -1649,9 +1654,9 @@ public partial class GUBP : BuildCommand
                 {
                     foreach (var Node in GUBPNodes)
                     {
-                        if (Node.Value.TriggerNode() && Node.Value.RunInEC())
+                        if (Node.Value.Node.TriggerNode() && Node.Value.Node.RunInEC())
                         {
-                            Node.Value.SetAsExplicitTrigger();
+                            Node.Value.Node.SetAsExplicitTrigger();
                         }
                     }
                 }
@@ -1664,7 +1669,7 @@ public partial class GUBP : BuildCommand
 			foreach (var NodeToDo in NodesToDo)
 			{
 				var TriggerDot = GetControllingTriggerDotName(NodeToDo, GUBPNodesControllingTriggerDotName, GUBPNodesControllingTrigger);
-				if (TriggerDot == "" && !GUBPNodes[NodeToDo].TriggerNode())
+				if (TriggerDot == "" && !GUBPNodes[NodeToDo].Node.TriggerNode())
 				{
 					LogVerbose("  Keeping {0}", NodeToDo);
 					NewNodesToDo.Add(NodeToDo);
@@ -1750,7 +1755,7 @@ public partial class GUBP : BuildCommand
 			List<string> SearchNodes = new List<string> { Barrier.Key };
 			for (int Idx = 0; Idx < SearchNodes.Count; Idx++)
 			{
-				GUBPNode Node = GUBPNodes[SearchNodes[Idx]];
+				GUBPNode Node = GUBPNodes[SearchNodes[Idx]].Node;
 				foreach (string DependencyName in Node.FullNamesOfDependencies.Union(Node.FullNamesOfPseudosependencies))
 				{
 					if (!IncludedNodes.Contains(DependencyName))
@@ -1762,7 +1767,7 @@ public partial class GUBP : BuildCommand
 			}
 
 			// Make sure that everything included in this list is before the cap, and everything not in the list is after it
-			foreach (KeyValuePair<string, GUBPNode> NodePair in GUBPNodes)
+			foreach (KeyValuePair<string, NodeInfo> NodePair in GUBPNodes)
 			{
 				if (IncludedNodes.Contains(NodePair.Key))
 				{
@@ -1829,8 +1834,8 @@ public partial class GUBP : BuildCommand
 					{
 						foreach (var Node in GUBPNodes)
 						{
-							if (Node.Value.GetFullName().Equals(NodeArg, StringComparison.InvariantCultureIgnoreCase) ||
-								Node.Value.AgentSharingGroup.Equals(NodeArg, StringComparison.InvariantCultureIgnoreCase)
+							if (Node.Value.Node.GetFullName().Equals(NodeArg, StringComparison.InvariantCultureIgnoreCase) ||
+								Node.Value.Node.AgentSharingGroup.Equals(NodeArg, StringComparison.InvariantCultureIgnoreCase)
 								)
 							{
 								if (!NodesToDo.Contains(Node.Key))
@@ -1867,7 +1872,7 @@ public partial class GUBP : BuildCommand
 							}
 							foreach (var Node in GUBPNodes)
 							{
-								if (Node.Value.GameNameIfAnyForTempStorage() == GameProj.GameName)
+								if (Node.Value.Node.GameNameIfAnyForTempStorage() == GameProj.GameName)
 								{
 									NodesToDo.Add(Node.Key);
 								}
@@ -1884,7 +1889,7 @@ public partial class GUBP : BuildCommand
 							{
 								foreach (var Node in GUBPNodes)
 								{
-									if (Node.Value.GameNameIfAnyForTempStorage() == GameProj.GameName)
+									if (Node.Value.Node.GameNameIfAnyForTempStorage() == GameProj.GameName)
 									{
 										NodesToDo.Add(Node.Key);
 									}
@@ -1913,9 +1918,9 @@ public partial class GUBP : BuildCommand
 			LogVerbose("Check to make sure we didn't ask for nodes that will be culled by time index");
 			foreach (var NodeToDo in NodesToDo)
 			{
-				if (TimeIndex % (1 << GUBPNodes[NodeToDo].DependentCISFrequencyQuantumShift()) != 0)
+				if (TimeIndex % (1 << GUBPNodes[NodeToDo].Node.DependentCISFrequencyQuantumShift()) != 0)
 				{
-					throw new AutomationException("You asked specifically for node {0}, but it is culled by the time quantum: TimeIndex = {1}, DependentCISFrequencyQuantumShift = {2}.", NodeToDo, TimeIndex, GUBPNodes[NodeToDo].DependentCISFrequencyQuantumShift());
+					throw new AutomationException("You asked specifically for node {0}, but it is culled by the time quantum: TimeIndex = {1}, DependentCISFrequencyQuantumShift = {2}.", NodeToDo, TimeIndex, GUBPNodes[NodeToDo].Node.DependentCISFrequencyQuantumShift());
 				}
 			}
 		}
@@ -1938,7 +1943,7 @@ public partial class GUBP : BuildCommand
 				{
 					if (!NodesToDo.Contains(NodeToDo.Key))
 					{
-						foreach (var Dep in GUBPNodes[NodeToDo.Key].FullNamesOfDependencies)
+						foreach (var Dep in GUBPNodes[NodeToDo.Key].Node.FullNamesOfDependencies)
 						{
 							if (!GUBPNodes.ContainsKey(Dep))
 							{
@@ -1950,7 +1955,7 @@ public partial class GUBP : BuildCommand
 								bDoneWithDependencies = false;
 							}
 						}
-						foreach (var Dep in GUBPNodes[NodeToDo.Key].FullNamesOfPseudosependencies)
+						foreach (var Dep in GUBPNodes[NodeToDo.Key].Node.FullNamesOfPseudosependencies)
 						{
 							if (!GUBPNodes.ContainsKey(Dep))
 							{
@@ -1988,7 +1993,7 @@ public partial class GUBP : BuildCommand
 						}
 						BranchOptions.NodesToRemovePseudoDependencies = ListOfRetainedNodesToRemovePseudoDependencies;
 					}
-					foreach (var Dep in GUBPNodes[NodeToDo].FullNamesOfDependencies)
+					foreach (var Dep in GUBPNodes[NodeToDo].Node.FullNamesOfDependencies)
 					{
 						if (!GUBPNodes.ContainsKey(Dep))
 						{
@@ -2000,7 +2005,7 @@ public partial class GUBP : BuildCommand
 							bDoneWithDependencies = false;
 						}
 					}
-					foreach (var Dep in GUBPNodes[NodeToDo].FullNamesOfPseudosependencies)
+					foreach (var Dep in GUBPNodes[NodeToDo].Node.FullNamesOfPseudosependencies)
 					{
 						if (!GUBPNodes.ContainsKey(Dep))
 						{
@@ -2017,7 +2022,7 @@ public partial class GUBP : BuildCommand
 			var NewNodesToDo = new HashSet<string>();
 			foreach (var NodeToDo in NodesToDo)
 			{
-				if (TimeIndex % (1 << GUBPNodes[NodeToDo].DependentCISFrequencyQuantumShift()) == 0)
+				if (TimeIndex % (1 << GUBPNodes[NodeToDo].Node.DependentCISFrequencyQuantumShift()) == 0)
 				{
 					LogVerbose("  Keeping {0}", NodeToDo);
 					NewNodesToDo.Add(NodeToDo);
@@ -2035,7 +2040,7 @@ public partial class GUBP : BuildCommand
 			var NewNodesToDo = new HashSet<string>();
 			foreach (var NodeToDo in NodesToDo)
 			{
-				if (!GUBPNodes[NodeToDo].GetFullName().Contains("Linux"))
+				if (!GUBPNodes[NodeToDo].Node.GetFullName().Contains("Linux"))
 				{
 					NewNodesToDo.Add(NodeToDo);
 				}
@@ -2051,7 +2056,7 @@ public partial class GUBP : BuildCommand
 		{
 			foreach (var NodeToDo in NodesToDo)
 			{
-				if (!GUBPNodes[NodeToDo].IsAggregate() && !GUBPNodes[NodeToDo].IsTest())
+				if (!GUBPNodes[NodeToDo].Node.IsAggregate() && !GUBPNodes[NodeToDo].Node.IsTest())
 				{
 					List<string> ECDependencies = new List<string>();
 					ECDependencies = GetECDependencies(NodeToDo);
@@ -2061,16 +2066,16 @@ public partial class GUBP : BuildCommand
 						{
 							throw new AutomationException("Node {0} is not in the graph. It is a dependency of {1}.", Dep, NodeToDo);
 						}
-						if (!GUBPNodes[Dep].FullNamesOfDependedOn.Contains(NodeToDo))
+						if (!GUBPNodes[Dep].Node.FullNamesOfDependedOn.Contains(NodeToDo))
 						{
-							GUBPNodes[Dep].FullNamesOfDependedOn.Add(NodeToDo);
+							GUBPNodes[Dep].Node.FullNamesOfDependedOn.Add(NodeToDo);
 						}
 					}
 				}
 			}
 			foreach (var NodeToDo in NodesToDo)
 			{
-				var Deps = GUBPNodes[NodeToDo].FullNamesOfDependedOn;
+				var Deps = GUBPNodes[NodeToDo].Node.FullNamesOfDependedOn;
 				string All = "";
 				foreach (var Dep in Deps)
 				{
@@ -2104,7 +2109,7 @@ public partial class GUBP : BuildCommand
 			{
 				Note = "always";
 			}
-			if (GUBPNodes[Node].RunInEC())
+			if (GUBPNodes[Node].Node.RunInEC())
 			{
 				var Deps = GetECDependencies(Node);
 				string All = "";
@@ -2133,7 +2138,7 @@ public partial class GUBP : BuildCommand
 	{
 		foreach (var NodeToDo in GUBPNodes)
 		{
-			foreach (var Dep in GUBPNodes[NodeToDo.Key].FullNamesOfDependencies)
+			foreach (var Dep in GUBPNodes[NodeToDo.Key].Node.FullNamesOfDependencies)
 			{
 				if (!GUBPNodes.ContainsKey(Dep))
 				{
@@ -2144,7 +2149,7 @@ public partial class GUBP : BuildCommand
 					throw new AutomationException("Node {0} has a self arc.", NodeToDo.Key);
 				}
 			}
-			foreach (var Dep in GUBPNodes[NodeToDo.Key].FullNamesOfPseudosependencies)
+			foreach (var Dep in GUBPNodes[NodeToDo.Key].Node.FullNamesOfPseudosependencies)
 			{
 				if (!GUBPNodes.ContainsKey(Dep))
 				{
@@ -2284,7 +2289,7 @@ public partial class GUBP : BuildCommand
 		{
 			foreach (var NodeToDo in OrdereredToDo)
 			{
-				if (GUBPNodes[NodeToDo].TriggerNode() && !NodeIsAlreadyComplete(GUBPNodesCompleted, NodeToDo, LocalOnly))
+				if (GUBPNodes[NodeToDo].Node.TriggerNode() && !NodeIsAlreadyComplete(GUBPNodesCompleted, NodeToDo, LocalOnly))
 				{
 					if (String.IsNullOrEmpty(ExplicitTrigger) || ExplicitTrigger != NodeToDo)
 					{
@@ -2300,16 +2305,16 @@ public partial class GUBP : BuildCommand
 	{
 		foreach (var NodeToDo in OrdereredToDo)
 		{
-			if (GUBPNodes[NodeToDo].TriggerNode() && (GUBPNodes[NodeToDo].IsSticky() || NodeIsAlreadyComplete(GUBPNodesCompleted, NodeToDo, LocalOnly))) // these sticky triggers are ok, everything is already completed anyway
+			if (GUBPNodes[NodeToDo].Node.TriggerNode() && (GUBPNodes[NodeToDo].Node.IsSticky() || NodeIsAlreadyComplete(GUBPNodesCompleted, NodeToDo, LocalOnly))) // these sticky triggers are ok, everything is already completed anyway
 			{
 				continue;
 			}
-			if (GUBPNodes[NodeToDo].IsTest())
+			if (GUBPNodes[NodeToDo].Node.IsTest())
 			{
 				bHasTests = true;
 			}
 			int MyIndex = OrdereredToDo.IndexOf(NodeToDo);
-			foreach (var Dep in GUBPNodes[NodeToDo].FullNamesOfDependencies)
+			foreach (var Dep in GUBPNodes[NodeToDo].Node.FullNamesOfDependencies)
 			{
 				int DepIndex = OrdereredToDo.IndexOf(Dep);
 				if (DepIndex >= MyIndex)
@@ -2317,7 +2322,7 @@ public partial class GUBP : BuildCommand
 					throw new AutomationException("Topological sort error, node {0} has a dependency of {1} which sorted after it.", NodeToDo, Dep);
 				}
 			}
-			foreach (var Dep in GUBPNodes[NodeToDo].FullNamesOfPseudosependencies)
+			foreach (var Dep in GUBPNodes[NodeToDo].Node.FullNamesOfPseudosependencies)
 			{
 				int DepIndex = OrdereredToDo.IndexOf(Dep);
 				if (DepIndex >= MyIndex)
@@ -2389,7 +2394,7 @@ public partial class GUBP : BuildCommand
 						continue; // this wasn't on the chain related to the trigger we are triggering, so it is not relevant
 					}
 				}
-				if (bPreflightBuild && !bSkipTriggers && GUBPNodes[NodeToDo].TriggerNode())
+				if (bPreflightBuild && !bSkipTriggers && GUBPNodes[NodeToDo].Node.TriggerNode())
 				{
 					// in preflight builds, we are either skipping triggers (and running things downstream) or we just stop at triggers and don't make them available for triggering.
 					continue;
@@ -2410,7 +2415,7 @@ public partial class GUBP : BuildCommand
 		{
 			foreach (var NodeToDo in FilteredOrdereredToDo)
 			{
-				if (GUBPNodes[NodeToDo].RunInEC() && !NodeIsAlreadyComplete(GUBPNodesCompleted, NodeToDo, LocalOnly) && NodeToDo != ExplicitTrigger && !NodeDependsOn(ExplicitTrigger, NodeToDo)) // if something is already finished, we don't put it into EC
+				if (GUBPNodes[NodeToDo].Node.RunInEC() && !NodeIsAlreadyComplete(GUBPNodesCompleted, NodeToDo, LocalOnly) && NodeToDo != ExplicitTrigger && !NodeDependsOn(ExplicitTrigger, NodeToDo)) // if something is already finished, we don't put it into EC
 				{
 					throw new AutomationException("We are being asked to process node {0}, however, this is an explicit trigger {1}, so everything before it should already be handled. It seems likely that you waited too long to run the trigger. You will have to do a new build from scratch.", NodeToDo, ExplicitTrigger);
 				}
@@ -2430,10 +2435,10 @@ public partial class GUBP : BuildCommand
 		// sticky nodes are ones that we run on the main agent. We run then first and they must not be intermixed with parallel jobs
 		foreach (var NodeToDo in OrdereredToDo)
 		{
-			if (GUBPNodes[NodeToDo].RunInEC() && !NodeIsAlreadyComplete(GUBPNodesCompleted, NodeToDo, LocalOnly)) // if something is already finished, we don't put it into EC
+			if (GUBPNodes[NodeToDo].Node.RunInEC() && !NodeIsAlreadyComplete(GUBPNodesCompleted, NodeToDo, LocalOnly)) // if something is already finished, we don't put it into EC
 			{
 				bHaveECNodes = true;
-				if (GUBPNodes[NodeToDo].IsSticky())
+				if (GUBPNodes[NodeToDo].Node.IsSticky())
 				{
 					LastSticky = NodeToDo;
 					if (HitNonSticky && !bSkipTriggers)
@@ -2466,9 +2471,9 @@ public partial class GUBP : BuildCommand
 		var StickyChain = new List<string>();
 		foreach (var NodeToDo in OrdereredToDo)
 		{
-			if (GUBPNodes[NodeToDo].RunInEC() && !NodeIsAlreadyComplete(GUBPNodesCompleted, NodeToDo, LocalOnly)) // if something is already finished, we don't put it into EC  
+			if (GUBPNodes[NodeToDo].Node.RunInEC() && !NodeIsAlreadyComplete(GUBPNodesCompleted, NodeToDo, LocalOnly)) // if something is already finished, we don't put it into EC  
 			{
-				string MyAgentGroup = GUBPNodes[NodeToDo].AgentSharingGroup;
+				string MyAgentGroup = GUBPNodes[NodeToDo].Node.AgentSharingGroup;
 				if (MyAgentGroup != "")
 				{
 					if (!AgentGroupChains.ContainsKey(MyAgentGroup))
@@ -2481,7 +2486,7 @@ public partial class GUBP : BuildCommand
 					}
 				}
 			}
-			if (GUBPNodes[NodeToDo].IsSticky())
+			if (GUBPNodes[NodeToDo].Node.IsSticky())
 			{
 				if (!StickyChain.Contains(NodeToDo))
 				{
@@ -2491,30 +2496,30 @@ public partial class GUBP : BuildCommand
 		}
 		foreach (var NodeToDo in OrdereredToDo)
 		{
-			if (GUBPNodes[NodeToDo].RunInEC() && !NodeIsAlreadyComplete(GUBPNodesCompleted, NodeToDo, LocalOnly)) // if something is already finished, we don't put it into EC  
+			if (GUBPNodes[NodeToDo].Node.RunInEC() && !NodeIsAlreadyComplete(GUBPNodesCompleted, NodeToDo, LocalOnly)) // if something is already finished, we don't put it into EC  
 			{
 				string EMails;
 				var NodeProps = GetECPropsForNode(NodeToDo, CLString, out EMails);
 				ECProps.AddRange(NodeProps);
 
-				bool Sticky = GUBPNodes[NodeToDo].IsSticky();
+				bool Sticky = GUBPNodes[NodeToDo].Node.IsSticky();
 				bool DoParallel = !Sticky;
-				if (GUBPNodes[NodeToDo].ECProcedure() == "GUBP_UAT_Node_Parallel_AgentShare_Editor")
+				if (GUBPNodes[NodeToDo].Node.ECProcedure() == "GUBP_UAT_Node_Parallel_AgentShare_Editor")
 				{
 					DoParallel = true;
 				}
-				if (Sticky && GUBPNodes[NodeToDo].ECAgentString() != "")
+				if (Sticky && GUBPNodes[NodeToDo].Node.ECAgentString() != "")
 				{
 					throw new AutomationException("Node {1} is sticky but has agent requirements.", NodeToDo);
 				}
-				string Procedure = GUBPNodes[NodeToDo].ECProcedure();
-				if (GUBPNodes[NodeToDo].IsSticky() && NodeToDo == LastSticky)
+				string Procedure = GUBPNodes[NodeToDo].Node.ECProcedure();
+				if (GUBPNodes[NodeToDo].Node.IsSticky() && NodeToDo == LastSticky)
 				{
 					Procedure = Procedure + "_Release";
 				}
 				string Args = String.Format("{0}, subprocedure => '{1}', parallel => '{2}', jobStepName => '{3}', actualParameter => [{{actualParameterName => 'NodeName', value =>'{4}'}}",
 					BaseArgs, Procedure, DoParallel ? 1 : 0, NodeToDo, NodeToDo);
-				string ProcedureParams = GUBPNodes[NodeToDo].ECProcedureParams();
+				string ProcedureParams = GUBPNodes[NodeToDo].Node.ECProcedureParams();
 				if (!String.IsNullOrEmpty(ProcedureParams))
 				{
 					Args = Args + ProcedureParams;
@@ -2530,7 +2535,7 @@ public partial class GUBP : BuildCommand
 				var UncompletedEcDeps = new List<string>();
 				var UncompletedCompletedDeps = new List<string>();
 				var PreConditionUncompletedEcDeps = new List<string>();
-				string MyAgentGroup = GUBPNodes[NodeToDo].AgentSharingGroup;
+				string MyAgentGroup = GUBPNodes[NodeToDo].Node.AgentSharingGroup;
 				bool bDoNestedJobstep = false;
 				bool bDoFirstNestedJobstep = false;
 
@@ -2543,7 +2548,7 @@ public partial class GUBP : BuildCommand
 					var EcDeps = GetECDependencies(NodeToDo);
 					foreach (var Dep in EcDeps)
 					{
-						if (GUBPNodes[Dep].RunInEC() && !NodeIsAlreadyComplete(GUBPNodesCompleted, Dep, LocalOnly) && OrdereredToDo.Contains(Dep)) // if something is already finished, we don't put it into EC
+						if (GUBPNodes[Dep].Node.RunInEC() && !NodeIsAlreadyComplete(GUBPNodesCompleted, Dep, LocalOnly) && OrdereredToDo.Contains(Dep)) // if something is already finished, we don't put it into EC
 						{
 							if (OrdereredToDo.IndexOf(Dep) > OrdereredToDo.IndexOf(NodeToDo))
 							{
@@ -2561,7 +2566,7 @@ public partial class GUBP : BuildCommand
 				var CompletedDeps = GetCompletedOnlyDependencies(NodeToDo);
 				foreach (var Dep in CompletedDeps)
 				{
-					if (GUBPNodes[Dep].RunInEC() && !NodeIsAlreadyComplete(GUBPNodesCompleted, Dep, LocalOnly) && OrdereredToDo.Contains(Dep)) // if something is already finished, we don't put it into EC
+					if (GUBPNodes[Dep].Node.RunInEC() && !NodeIsAlreadyComplete(GUBPNodesCompleted, Dep, LocalOnly) && OrdereredToDo.Contains(Dep)) // if something is already finished, we don't put it into EC
 					{
 						if (OrdereredToDo.IndexOf(Dep) > OrdereredToDo.IndexOf(NodeToDo))
 						{
@@ -2592,7 +2597,7 @@ public partial class GUBP : BuildCommand
 							var EcDeps = GetECDependencies(Chain);
 							foreach (var Dep in EcDeps)
 							{
-								if (GUBPNodes[Dep].RunInEC() && !NodeIsAlreadyComplete(GUBPNodesCompleted, Dep, LocalOnly) && OrdereredToDo.Contains(Dep)) // if something is already finished, we don't put it into EC
+								if (GUBPNodes[Dep].Node.RunInEC() && !NodeIsAlreadyComplete(GUBPNodesCompleted, Dep, LocalOnly) && OrdereredToDo.Contains(Dep)) // if something is already finished, we don't put it into EC
 								{
 									if (OrdereredToDo.IndexOf(Dep) > OrdereredToDo.IndexOf(Chain))
 									{
@@ -2607,7 +2612,7 @@ public partial class GUBP : BuildCommand
 						}
 					}
 				}
-				if (GUBPNodes[NodeToDo].IsSticky())
+				if (GUBPNodes[NodeToDo].Node.IsSticky())
 				{
 					var MyChain = StickyChain;
 					int MyIndex = MyChain.IndexOf(NodeToDo);
@@ -2623,7 +2628,7 @@ public partial class GUBP : BuildCommand
 						var EcDeps = GetECDependencies(NodeToDo);
 						foreach (var Dep in EcDeps)
 						{
-							if (GUBPNodes[Dep].RunInEC() && !NodeIsAlreadyComplete(GUBPNodesCompleted, Dep, LocalOnly) && OrdereredToDo.Contains(Dep)) // if something is already finished, we don't put it into EC
+							if (GUBPNodes[Dep].Node.RunInEC() && !NodeIsAlreadyComplete(GUBPNodesCompleted, Dep, LocalOnly) && OrdereredToDo.Contains(Dep)) // if something is already finished, we don't put it into EC
 							{
 								if (OrdereredToDo.IndexOf(Dep) > OrdereredToDo.IndexOf(NodeToDo))
 								{
@@ -2664,7 +2669,7 @@ public partial class GUBP : BuildCommand
 					int Index = 0;
 					foreach (var Dep in CompletedDeps)
 					{
-						if (GUBPNodes[Dep].RunInEC())
+						if (GUBPNodes[Dep].Node.RunInEC())
 						{
 							PreCondition = PreCondition + "getProperty('" + GetJobStep(PreconditionParentPath, Dep) + "/status\') == \'completed\'";
 							Index++;
@@ -2711,7 +2716,7 @@ public partial class GUBP : BuildCommand
 						}
 						{
 							string NestArgs = String.Format("$batch->createJobStep({{parentPath => '{0}/jobSteps[{1}]', jobStepName => '{2}_GetPool', subprocedure => 'GUBP{3}_AgentShare_GetPool', parallel => '1', actualParameter => [{{actualParameterName => 'AgentSharingGroup', value => '{4}'}}, {{actualParameterName => 'NodeName', value => '{5}'}}]",
-								ParentPath, MyAgentGroup, MyAgentGroup, GUBPNodes[NodeToDo].ECProcedureInfix(), MyAgentGroup, NodeToDo);
+								ParentPath, MyAgentGroup, MyAgentGroup, GUBPNodes[NodeToDo].Node.ECProcedureInfix(), MyAgentGroup, NodeToDo);
 							if (!String.IsNullOrEmpty(PreCondition))
 							{
 								NestArgs = NestArgs + ", precondition => " + PreCondition;
@@ -2721,7 +2726,7 @@ public partial class GUBP : BuildCommand
 						}
 						{
 							string NestArgs = String.Format("$batch->createJobStep({{parentPath => '{0}/jobSteps[{1}]', jobStepName => '{2}_GetAgent', subprocedure => 'GUBP{3}_AgentShare_GetAgent', parallel => '1', exclusiveMode => 'call', resourceName => '{4}', actualParameter => [{{actualParameterName => 'AgentSharingGroup', value => '{5}'}}, {{actualParameterName => 'NodeName', value=> '{6}'}}]",
-								ParentPath, MyAgentGroup, MyAgentGroup, GUBPNodes[NodeToDo].ECProcedureInfix(),
+								ParentPath, MyAgentGroup, MyAgentGroup, GUBPNodes[NodeToDo].Node.ECProcedureInfix(),
 								String.Format("$[/myJob/jobSteps[{0}]/ResourcePool]", MyAgentGroup),
 								MyAgentGroup, NodeToDo);
 							{
@@ -2766,7 +2771,7 @@ public partial class GUBP : BuildCommand
 				StepList.Add(Args);
 				if (bFakeEC &&
 					!UnfinishedTriggers.Contains(NodeToDo) &&
-					(GUBPNodes[NodeToDo].ECProcedure().StartsWith("GUBP_UAT_Node") || GUBPNodes[NodeToDo].ECProcedure().StartsWith("GUBP_Mac_UAT_Node")) // other things we really can't test
+					(GUBPNodes[NodeToDo].Node.ECProcedure().StartsWith("GUBP_UAT_Node") || GUBPNodes[NodeToDo].Node.ECProcedure().StartsWith("GUBP_Mac_UAT_Node")) // other things we really can't test
 					) // unfinished triggers are never run directly by EC, rather it does another job setup
 				{
 					string Arg = String.Format("gubp -Node={0} -FakeEC {1} {2} {3} {4} {5}",
@@ -2838,17 +2843,17 @@ public partial class GUBP : BuildCommand
         var BuildProductToNodeMap = new Dictionary<string, string>();
 		foreach (var NodeToDo in OrdereredToDo)
         {
-            if (GUBPNodes[NodeToDo].BuildProducts != null || GUBPNodes[NodeToDo].AllDependencyBuildProducts != null)
+            if (GUBPNodes[NodeToDo].Node.BuildProducts != null || GUBPNodes[NodeToDo].Node.AllDependencyBuildProducts != null)
             {
                 throw new AutomationException("topological sort error");
             }
 
-            GUBPNodes[NodeToDo].AllDependencyBuildProducts = new List<string>();
-            GUBPNodes[NodeToDo].AllDependencies = new List<string>();
-            foreach (var Dep in GUBPNodes[NodeToDo].FullNamesOfDependencies)
+            GUBPNodes[NodeToDo].Node.AllDependencyBuildProducts = new List<string>();
+            GUBPNodes[NodeToDo].Node.AllDependencies = new List<string>();
+            foreach (var Dep in GUBPNodes[NodeToDo].Node.FullNamesOfDependencies)
             {
-                GUBPNodes[NodeToDo].AddAllDependent(Dep);
-                if (GUBPNodes[Dep].AllDependencies == null)
+                GUBPNodes[NodeToDo].Node.AddAllDependent(Dep);
+                if (GUBPNodes[Dep].Node.AllDependencies == null)
                 {
                     if (!bOnlyNode)
                     {
@@ -2857,13 +2862,13 @@ public partial class GUBP : BuildCommand
                 }
                 else
 				{
-					foreach (var DepDep in GUBPNodes[Dep].AllDependencies)
+					foreach (var DepDep in GUBPNodes[Dep].Node.AllDependencies)
 					{
-						GUBPNodes[NodeToDo].AddAllDependent(DepDep);
+						GUBPNodes[NodeToDo].Node.AddAllDependent(DepDep);
 					}
 				}
 				
-                if (GUBPNodes[Dep].BuildProducts == null)
+                if (GUBPNodes[Dep].Node.BuildProducts == null)
                 {
                     if (!bOnlyNode)
                     {
@@ -2872,25 +2877,25 @@ public partial class GUBP : BuildCommand
                 }
                 else
                 {
-                    foreach (var Prod in GUBPNodes[Dep].BuildProducts)
+                    foreach (var Prod in GUBPNodes[Dep].Node.BuildProducts)
                     {
-                        GUBPNodes[NodeToDo].AddDependentBuildProduct(Prod);
+                        GUBPNodes[NodeToDo].Node.AddDependentBuildProduct(Prod);
                     }
-                    if (GUBPNodes[Dep].AllDependencyBuildProducts == null)
+                    if (GUBPNodes[Dep].Node.AllDependencyBuildProducts == null)
                     {
                         throw new AutomationException("Node {0} was not processed yet2?  Processing {1}", Dep, NodeToDo);
                     }
-                    foreach (var Prod in GUBPNodes[Dep].AllDependencyBuildProducts)
+                    foreach (var Prod in GUBPNodes[Dep].Node.AllDependencyBuildProducts)
                     {
-                        GUBPNodes[NodeToDo].AddDependentBuildProduct(Prod);
+                        GUBPNodes[NodeToDo].Node.AddDependentBuildProduct(Prod);
                     }
                 }
             }
 
-            string NodeStoreName = StoreName + "-" + GUBPNodes[NodeToDo].GetFullName();
+            string NodeStoreName = StoreName + "-" + GUBPNodes[NodeToDo].Node.GetFullName();
             
-            string GameNameIfAny = GUBPNodes[NodeToDo].GameNameIfAnyForTempStorage();
-            string StorageRootIfAny = GUBPNodes[NodeToDo].RootIfAnyForTempStorage();
+            string GameNameIfAny = GUBPNodes[NodeToDo].Node.GameNameIfAnyForTempStorage();
+            string StorageRootIfAny = GUBPNodes[NodeToDo].Node.RootIfAnyForTempStorage();
 					
             if (bFake)
             {
@@ -2899,31 +2904,31 @@ public partial class GUBP : BuildCommand
 
             // this is kinda complicated
             bool SaveSuccessRecords = (IsBuildMachine || bFakeEC) && // no real reason to make these locally except for fakeEC tests
-                (!GUBPNodes[NodeToDo].TriggerNode() || GUBPNodes[NodeToDo].IsSticky()) // trigger nodes are run twice, one to start the new workflow and once when it is actually triggered, we will save reconds for the latter
-                && (GUBPNodes[NodeToDo].RunInEC() || !GUBPNodes[NodeToDo].IsAggregate()); //aggregates not in EC can be "run" multiple times, so we can't track those
+                (!GUBPNodes[NodeToDo].Node.TriggerNode() || GUBPNodes[NodeToDo].Node.IsSticky()) // trigger nodes are run twice, one to start the new workflow and once when it is actually triggered, we will save reconds for the latter
+                && (GUBPNodes[NodeToDo].Node.RunInEC() || !GUBPNodes[NodeToDo].Node.IsAggregate()); //aggregates not in EC can be "run" multiple times, so we can't track those
 
-            Log("***** Running GUBP Node {0} -> {1} : {2}", GUBPNodes[NodeToDo].GetFullName(), GameNameIfAny, NodeStoreName);
+            Log("***** Running GUBP Node {0} -> {1} : {2}", GUBPNodes[NodeToDo].Node.GetFullName(), GameNameIfAny, NodeStoreName);
             if (NodeIsAlreadyComplete(GUBPNodesCompleted, NodeToDo, LocalOnly))
             {
                 if (NodeToDo == VersionFilesNode.StaticGetFullName() && !IsBuildMachine)
                 {
-                    Log("***** NOT ****** Retrieving GUBP Node {0} from {1}; it is the version files.", GUBPNodes[NodeToDo].GetFullName(), NodeStoreName);
-                    GUBPNodes[NodeToDo].BuildProducts = new List<string>();
+                    Log("***** NOT ****** Retrieving GUBP Node {0} from {1}; it is the version files.", GUBPNodes[NodeToDo].Node.GetFullName(), NodeStoreName);
+                    GUBPNodes[NodeToDo].Node.BuildProducts = new List<string>();
 
                 }
                 else
                 {
-                    Log("***** Retrieving GUBP Node {0} from {1}", GUBPNodes[NodeToDo].GetFullName(), NodeStoreName);
+                    Log("***** Retrieving GUBP Node {0} from {1}", GUBPNodes[NodeToDo].Node.GetFullName(), NodeStoreName);
                     bool WasLocal;
 					try
 					{
-						GUBPNodes[NodeToDo].BuildProducts = TempStorage.RetrieveFromTempStorage(CmdEnv, NodeStoreName, out WasLocal, GameNameIfAny, StorageRootIfAny);
+						GUBPNodes[NodeToDo].Node.BuildProducts = TempStorage.RetrieveFromTempStorage(CmdEnv, NodeStoreName, out WasLocal, GameNameIfAny, StorageRootIfAny);
 					}
 					catch
 					{
 						if(GameNameIfAny != "")
 						{
-							GUBPNodes[NodeToDo].BuildProducts = TempStorage.RetrieveFromTempStorage(CmdEnv, NodeStoreName, out WasLocal, "", StorageRootIfAny);
+							GUBPNodes[NodeToDo].Node.BuildProducts = TempStorage.RetrieveFromTempStorage(CmdEnv, NodeStoreName, out WasLocal, "", StorageRootIfAny);
 						}
 						else
 						{
@@ -2932,7 +2937,7 @@ public partial class GUBP : BuildCommand
 					}
                     if (!WasLocal)
                     {
-                        GUBPNodes[NodeToDo].PostLoadFromSharedTempStorage(this);
+                        GUBPNodes[NodeToDo].Node.PostLoadFromSharedTempStorage(this);
                     }
                 }
             }
@@ -2952,14 +2957,14 @@ public partial class GUBP : BuildCommand
                     if (bFake)
                     {
                         Log("***** FAKE!! Building GUBP Node {0} for {1}", NodeToDo, NodeStoreName);
-                        GUBPNodes[NodeToDo].DoFakeBuild(this);
+                        GUBPNodes[NodeToDo].Node.DoFakeBuild(this);
                     }
                     else
                     {                        
 						Log("***** Building GUBP Node {0} for {1}", NodeToDo, NodeStoreName);
 						var StartTime = DateTime.UtcNow;
                         var StartBuild = DateTime.Now.ToString();
-                        GUBPNodes[NodeToDo].DoBuild(this);
+                        GUBPNodes[NodeToDo].Node.DoBuild(this);
                         var FinishBuild = DateTime.Now.ToString();
                         if (IsBuildMachine)
                         {
@@ -2968,12 +2973,12 @@ public partial class GUBP : BuildCommand
 						BuildDuration = (DateTime.UtcNow - StartTime).TotalMilliseconds / 1000;
 						
                     }
-                    if (!GUBPNodes[NodeToDo].IsAggregate())
+                    if (!GUBPNodes[NodeToDo].Node.IsAggregate())
                     {
 						var StoreDuration = 0.0;
 						var StartTime = DateTime.UtcNow;
                         var StartStore = DateTime.Now.ToString();
-                        TempStorage.StoreToTempStorage(CmdEnv, NodeStoreName, GUBPNodes[NodeToDo].BuildProducts, !bSaveSharedTempStorage, GameNameIfAny, StorageRootIfAny);
+                        TempStorage.StoreToTempStorage(CmdEnv, NodeStoreName, GUBPNodes[NodeToDo].Node.BuildProducts, !bSaveSharedTempStorage, GameNameIfAny, StorageRootIfAny);
 						StoreDuration = (DateTime.UtcNow - StartTime).TotalMilliseconds / 1000;
                         var FinishStore = DateTime.Now.ToString();                        
 						Log("Took {0} seconds to store build products", StoreDuration);
@@ -2984,7 +2989,7 @@ public partial class GUBP : BuildCommand
 						}
                         if (ParseParam("StompCheck"))
                         {
-                            foreach (var Dep in GUBPNodes[NodeToDo].AllDependencies)
+                            foreach (var Dep in GUBPNodes[NodeToDo].Node.AllDependencies)
                             {
                                 try
                                 {
@@ -3102,7 +3107,7 @@ public partial class GUBP : BuildCommand
 					UpdateECBuildTime(NodeToDo, BuildDuration);
                 }
             }
-            foreach (var Product in GUBPNodes[NodeToDo].BuildProducts)
+            foreach (var Product in GUBPNodes[NodeToDo].Node.BuildProducts)
             {
                 if (BuildProductToNodeMap.ContainsKey(Product))
                 {
@@ -3118,17 +3123,17 @@ public partial class GUBP : BuildCommand
 	/// Sorts a list of nodes to display in EC. The default order is based on execution order and agent groups, whereas this function arranges nodes by
 	/// frequency then execution order, while trying to group nodes on parallel paths (eg. Mac/Windows editor nodes) together.
 	/// </summary>
-	static Dictionary<string, int> GetDisplayOrder(List<string> NodeNames, Dictionary<string, string> InitialNodeDependencyNames, Dictionary<string, GUBPNode> GUBPNodes)
+	static Dictionary<string, int> GetDisplayOrder(List<string> NodeNames, Dictionary<string, string> InitialNodeDependencyNames, Dictionary<string, NodeInfo> GUBPNodes)
 	{
 		// Split the nodes into separate lists for each frequency
 		SortedDictionary<int, List<string>> NodesByFrequency = new SortedDictionary<int,List<string>>();
 		foreach(string NodeName in NodeNames)
 		{
 			List<string> NodesByThisFrequency;
-			if(!NodesByFrequency.TryGetValue(GUBPNodes[NodeName].DependentCISFrequencyQuantumShift(), out NodesByThisFrequency))
+			if(!NodesByFrequency.TryGetValue(GUBPNodes[NodeName].Node.DependentCISFrequencyQuantumShift(), out NodesByThisFrequency))
 			{
 				NodesByThisFrequency = new List<string>();
-				NodesByFrequency.Add(GUBPNodes[NodeName].DependentCISFrequencyQuantumShift(), NodesByThisFrequency);
+				NodesByFrequency.Add(GUBPNodes[NodeName].Node.DependentCISFrequencyQuantumShift(), NodesByThisFrequency);
 			}
 			NodesByThisFrequency.Add(NodeName);
 		}
@@ -3142,7 +3147,7 @@ public partial class GUBP : BuildCommand
 			Dictionary<string, string> DisplayGroups = new Dictionary<string,string>();
 			foreach(string NodeName in NodesByThisFrequency)
 			{
-				string GroupName = GUBPNodes[NodeName].GetDisplayGroupName();
+				string GroupName = GUBPNodes[NodeName].Node.GetDisplayGroupName();
 				if(!DisplayGroups.ContainsKey(GroupName))
 				{
 					DisplayGroups.Add(GroupName, NodeName);
@@ -3171,7 +3176,7 @@ public partial class GUBP : BuildCommand
 			// Add nodes for each frequency into the master list, trying to match up different groups along the way
 			foreach(string FirstNodeName in NodesByThisFrequency)
 			{
-				string[] GroupNodeNames = DisplayGroups[GUBPNodes[FirstNodeName].GetDisplayGroupName()].Split(' ');
+				string[] GroupNodeNames = DisplayGroups[GUBPNodes[FirstNodeName].Node.GetDisplayGroupName()].Split(' ');
 				foreach(string GroupNodeName in GroupNodeNames)
 				{
 					AddNodeAndDependencies(GroupNodeName, NodeDependencyNames, VisitedNodes, SortedNodes);
diff --git a/Engine/Source/Programs/AutomationTool/GUBP/LegacyBranchSetup.cs b/Engine/Source/Programs/AutomationTool/GUBP/LegacyBranchSetup.cs
--- a/Engine/Source/Programs/AutomationTool/GUBP/LegacyBranchSetup.cs
+++ b/Engine/Source/Programs/AutomationTool/GUBP/LegacyBranchSetup.cs
@@ -669,16 +669,16 @@ partial class GUBP
 														{
 															BuildAgentSharingGroup = "";
 														}
-														GUBPNodes[CookNode.StaticGetFullName(HostPlatform, NonCodeProject, CookedPlatform)].AgentSharingGroup = BuildAgentSharingGroup;
-														GUBPNodes[GamePlatformCookedAndCompiledNode.StaticGetFullName(HostPlatform, NonCodeProject, Plat)].AgentSharingGroup = BuildAgentSharingGroup;
-														GUBPNodes[NodeName].AgentSharingGroup = BuildAgentSharingGroup;
+														GUBPNodes[CookNode.StaticGetFullName(HostPlatform, NonCodeProject, CookedPlatform)].Node.AgentSharingGroup = BuildAgentSharingGroup;
+														GUBPNodes[GamePlatformCookedAndCompiledNode.StaticGetFullName(HostPlatform, NonCodeProject, Plat)].Node.AgentSharingGroup = BuildAgentSharingGroup;
+														GUBPNodes[NodeName].Node.AgentSharingGroup = BuildAgentSharingGroup;
 													}
 													else
 													{
 														//GUBPNodes[NodeName].AgentSharingGroup = FormalAgentSharingGroup;
 															if (Plat == UnrealTargetPlatform.XboxOne)
 														{
-															GUBPNodes[NodeName].AgentSharingGroup = "";
+															GUBPNodes[NodeName].Node.AgentSharingGroup = "";
 														}
 													}
 													if (PlatPair.bTest)
@@ -911,16 +911,16 @@ partial class GUBP
 												{
 													BuildAgentSharingGroup = "";
 												}
-												GUBPNodes[CookNode.StaticGetFullName(HostPlatform, CodeProj, CookedPlatform)].AgentSharingGroup = BuildAgentSharingGroup;
-												GUBPNodes[GamePlatformCookedAndCompiledNode.StaticGetFullName(HostPlatform, CodeProj, Plat)].AgentSharingGroup = BuildAgentSharingGroup;
-												GUBPNodes[FormalNodeName].AgentSharingGroup = BuildAgentSharingGroup;
+												GUBPNodes[CookNode.StaticGetFullName(HostPlatform, CodeProj, CookedPlatform)].Node.AgentSharingGroup = BuildAgentSharingGroup;
+												GUBPNodes[GamePlatformCookedAndCompiledNode.StaticGetFullName(HostPlatform, CodeProj, Plat)].Node.AgentSharingGroup = BuildAgentSharingGroup;
+												GUBPNodes[FormalNodeName].Node.AgentSharingGroup = BuildAgentSharingGroup;
 											}
 											else
 											{
 												//GUBPNodes[FormalNodeName].AgentSharingGroup = FormalAgentSharingGroup;
 												if (Plat == UnrealTargetPlatform.XboxOne)
 												{
-													GUBPNodes[FormalNodeName].AgentSharingGroup = "";
+													GUBPNodes[FormalNodeName].Node.AgentSharingGroup = "";
 												}
 											}
 											if (Config.bTest)
diff --git a/Engine/Source/Programs/AutomationTool/GUBP/LegacyNodes.cs b/Engine/Source/Programs/AutomationTool/GUBP/LegacyNodes.cs
--- a/Engine/Source/Programs/AutomationTool/GUBP/LegacyNodes.cs
+++ b/Engine/Source/Programs/AutomationTool/GUBP/LegacyNodes.cs
@@ -3506,7 +3506,7 @@ partial class GUBP
                 var StartTime = DateTime.UtcNow;
                 foreach (var NodeToDo in bp.GUBPNodes)
                 {
-                    TempStorage.CleanSharedTempStorageDirectory(NodeToDo.Value.GameNameIfAnyForTempStorage());
+                    TempStorage.CleanSharedTempStorageDirectory(NodeToDo.Value.Node.GameNameIfAnyForTempStorage());
                 }
                 var BuildDuration = (DateTime.UtcNow - StartTime).TotalMilliseconds;
                 Log("Took {0}s to clear temp storage of old files.", BuildDuration / 1000);
