Replace vars with explicit types.

diff --git a/Engine/Source/Programs/AutomationTool/GUBP/GUBP.cs b/Engine/Source/Programs/AutomationTool/GUBP/GUBP.cs
--- a/Engine/Source/Programs/AutomationTool/GUBP/GUBP.cs
+++ b/Engine/Source/Programs/AutomationTool/GUBP/GUBP.cs
@@ -96,8 +96,8 @@ public partial class GUBP : BuildCommand
 
     List<string> GetDependencies(string NodeToDo, bool bFlat = false, bool ECOnly = false)
     {
-        var Result = new List<string>();
-        foreach (var Node in GUBPNodes[NodeToDo].Node.FullNamesOfDependencies)
+        List<string> Result = new List<string>();
+        foreach (string Node in GUBPNodes[NodeToDo].Node.FullNamesOfDependencies)
         {
             bool Usable = GUBPNodes[Node].Node.RunInEC() || !ECOnly;
             if (Usable)
@@ -109,7 +109,7 @@ public partial class GUBP : BuildCommand
             }
             if (bFlat || !Usable)
             {
-                foreach (var RNode in GetDependencies(Node, bFlat, ECOnly))
+                foreach (string RNode in GetDependencies(Node, bFlat, ECOnly))
                 {
                     if (!Result.Contains(RNode))
                     {
@@ -118,7 +118,7 @@ public partial class GUBP : BuildCommand
                 }
             }
         }
-        foreach (var Node in GUBPNodes[NodeToDo].Node.FullNamesOfPseudosependencies)
+        foreach (string Node in GUBPNodes[NodeToDo].Node.FullNamesOfPseudosependencies)
         {
             bool Usable = GUBPNodes[Node].Node.RunInEC() || !ECOnly;
             if (Usable)
@@ -130,7 +130,7 @@ public partial class GUBP : BuildCommand
             }
             if (bFlat || !Usable)
             {
-                foreach (var RNode in GetDependencies(Node, bFlat, ECOnly))
+                foreach (string RNode in GetDependencies(Node, bFlat, ECOnly))
                 {
                     if (!Result.Contains(RNode))
                     {
@@ -144,8 +144,8 @@ public partial class GUBP : BuildCommand
 
 	List<string> GetCompletedOnlyDependencies(string NodeToDo, bool bFlat = false, bool ECOnly = true)
 	{
-		var Result = new List<string>();
-		foreach (var Node in GUBPNodes[NodeToDo].Node.CompletedDependencies)
+		List<string> Result = new List<string>();
+		foreach (string Node in GUBPNodes[NodeToDo].Node.CompletedDependencies)
 		{
 			bool Usable = GUBPNodes[Node].Node.RunInEC() || !ECOnly;
 			if(Usable)
@@ -157,7 +157,7 @@ public partial class GUBP : BuildCommand
 			}
 			if (bFlat || !Usable)
 			{
-				foreach (var RNode in GetCompletedOnlyDependencies(Node, bFlat, ECOnly))
+				foreach (string RNode in GetCompletedOnlyDependencies(Node, bFlat, ECOnly))
 				{
 					if (!Result.Contains(RNode))
 					{
@@ -174,7 +174,7 @@ public partial class GUBP : BuildCommand
     }
     bool NodeDependsOn(string Rootward, string Leafward)
     {
-        var Deps = GetDependencies(Leafward, true);
+        List<string> Deps = GetDependencies(Leafward, true);
         return Deps.Contains(Rootward);
     }
 
@@ -184,8 +184,8 @@ public partial class GUBP : BuildCommand
         {
             return GUBPNodesControllingTrigger[NodeToDo];
         }
-        var Result = "";
-        foreach (var Node in GUBPNodes[NodeToDo].Node.FullNamesOfDependencies)
+        string Result = "";
+        foreach (string Node in GUBPNodes[NodeToDo].Node.FullNamesOfDependencies)
         {
             if (!GUBPNodes.ContainsKey(Node))
             {
@@ -213,7 +213,7 @@ public partial class GUBP : BuildCommand
                 }
             }
         }
-        foreach (var Node in GUBPNodes[NodeToDo].Node.FullNamesOfPseudosependencies)
+        foreach (string Node in GUBPNodes[NodeToDo].Node.FullNamesOfPseudosependencies)
         {
             if (!GUBPNodes.ContainsKey(Node))
             {
@@ -308,15 +308,15 @@ public partial class GUBP : BuildCommand
 				Result = FrequencyOverride;
 			}
 
-            foreach (var Dep in GUBPNodes[NodeToDo].Node.FullNamesOfDependencies)
+            foreach (string Dep in GUBPNodes[NodeToDo].Node.FullNamesOfDependencies)
             {
                 Result = Math.Max(ComputeDependentCISFrequencyQuantumShift(Dep, FrequencyOverrides), Result);
             }
-            foreach (var Dep in GUBPNodes[NodeToDo].Node.FullNamesOfPseudosependencies)
+            foreach (string Dep in GUBPNodes[NodeToDo].Node.FullNamesOfPseudosependencies)
             {
                 Result = Math.Max(ComputeDependentCISFrequencyQuantumShift(Dep, FrequencyOverrides), Result);
             }
-			foreach (var Dep in GUBPNodes[NodeToDo].Node.CompletedDependencies)
+			foreach (string Dep in GUBPNodes[NodeToDo].Node.CompletedDependencies)
 			{
 				Result = Math.Max(ComputeDependentCISFrequencyQuantumShift(Dep, FrequencyOverrides), Result);
 			}
@@ -386,7 +386,7 @@ public partial class GUBP : BuildCommand
     }
     string GetEMailListForNode(GUBP bp, string NodeToDo, string Emails, string Causers)
     {        
-        var BranchForEmail = "";
+        string BranchForEmail = "";
         if (P4Enabled)
         {
             BranchForEmail = P4Env.BuildRootP4;
@@ -400,7 +400,7 @@ public partial class GUBP : BuildCommand
 
     List<P4Connection.ChangeRecord> GetChanges(int LastOutputForChanges, int TopCL, int LastGreen)
     {
-        var Result = new List<P4Connection.ChangeRecord>();
+        List<P4Connection.ChangeRecord> Result = new List<P4Connection.ChangeRecord>();
         if (TopCL > LastGreen)
         {
             if (LastOutputForChanges > 1990000)
@@ -412,7 +412,7 @@ public partial class GUBP : BuildCommand
                 List<P4Connection.ChangeRecord> ChangeRecords;
 				if (P4.Changes(out ChangeRecords, Cmd, false, true, LongComment: true))
                 {
-                    foreach (var Record in ChangeRecords)
+                    foreach (P4Connection.ChangeRecord Record in ChangeRecords)
                     {
                         if (!Record.User.Equals("buildmachine", StringComparison.InvariantCultureIgnoreCase))
                         {
@@ -435,10 +435,10 @@ public partial class GUBP : BuildCommand
 
     int PrintChanges(int LastOutputForChanges, int TopCL, int LastGreen)
     {
-        var ChangeRecords = GetChanges(LastOutputForChanges, TopCL, LastGreen);
-        foreach (var Record in ChangeRecords)
+        List<P4Connection.ChangeRecord> ChangeRecords = GetChanges(LastOutputForChanges, TopCL, LastGreen);
+        foreach (P4Connection.ChangeRecord Record in ChangeRecords)
         {
-            var Summary = Record.Summary.Replace("\r", "\n");
+            string Summary = Record.Summary.Replace("\r", "\n");
             if (Summary.IndexOf("\n") > 0)
             {
                 Summary = Summary.Substring(0, Summary.IndexOf("\n"));
@@ -450,7 +450,7 @@ public partial class GUBP : BuildCommand
 
     void PrintDetailedChanges(NodeHistory History, bool bShowAllChanges = false)
     {
-        var StartTime = DateTime.UtcNow;
+        DateTime StartTime = DateTime.UtcNow;
 
         string Me = String.Format("{0}   <<<< local sync", P4Env.Changelist);
         int LastOutputForChanges = 0;
@@ -462,7 +462,7 @@ public partial class GUBP : BuildCommand
                 LastGreen = History.AllStarted[0];
             }
         }
-        foreach (var cl in History.AllStarted)
+        foreach (int cl in History.AllStarted)
         {
             if (cl < LastGreen)
             {
@@ -491,7 +491,7 @@ public partial class GUBP : BuildCommand
             LastOutputForChanges = PrintChanges(LastOutputForChanges, P4Env.Changelist, LastGreen);
             Log("         {0}", Me);
         }
-        var BuildDuration = (DateTime.UtcNow - StartTime).TotalMilliseconds;
+        double BuildDuration = (DateTime.UtcNow - StartTime).TotalMilliseconds;
         Log("Took {0}s to get P4 history", BuildDuration / 1000);
 
     }
@@ -511,7 +511,7 @@ public partial class GUBP : BuildCommand
         bool bShowTriggers = true;
         string LastControllingTrigger = "";
         string LastAgentGroup = "";
-        foreach (var NodeToDo in Nodes)
+        foreach (string NodeToDo in Nodes)
         {
             if (ECOnly && !GUBPNodes[NodeToDo].Node.RunInEC())
             {
@@ -583,7 +583,7 @@ public partial class GUBP : BuildCommand
                 );
             if (bShowHistory && GUBPNodesHistory.ContainsKey(NodeToDo))
             {
-                var History = GUBPNodesHistory[NodeToDo];
+                NodeHistory History = GUBPNodesHistory[NodeToDo];
 
                 if (bShowDetailedHistory)
                 {
@@ -599,26 +599,26 @@ public partial class GUBP : BuildCommand
             }
             if (bShowDependencies)
             {
-                foreach (var Dep in GUBPNodes[NodeToDo].Node.FullNamesOfDependencies)
+                foreach (string Dep in GUBPNodes[NodeToDo].Node.FullNamesOfDependencies)
                 {
                     Log("            dep> {0}", Dep);
                 }
-                foreach (var Dep in GUBPNodes[NodeToDo].Node.FullNamesOfPseudosependencies)
+                foreach (string Dep in GUBPNodes[NodeToDo].Node.FullNamesOfPseudosependencies)
                 {
                     Log("           pdep> {0}", Dep);
                 }
-				foreach (var Dep in GUBPNodes[NodeToDo].Node.CompletedDependencies)
+				foreach (string Dep in GUBPNodes[NodeToDo].Node.CompletedDependencies)
 				{
 					Log("			cdep>{0}", Dep);
 				}
             }
             if (bShowECDependencies)
             {
-                foreach (var Dep in GetECDependencies(NodeToDo))
+                foreach (string Dep in GetECDependencies(NodeToDo))
                 {
                     Log("           {0}", Dep);
                 }
-				foreach (var Dep in GetCompletedOnlyDependencies(NodeToDo))
+				foreach (string Dep in GetCompletedOnlyDependencies(NodeToDo))
 				{
 					Log("           compDep> {0}", Dep);
 				}
@@ -626,14 +626,14 @@ public partial class GUBP : BuildCommand
 
 			if(bShowDependednOn)
 			{
-				foreach (var Dep in GUBPNodes[NodeToDo].Node.FullNamesOfDependedOn)
+				foreach (string Dep in GUBPNodes[NodeToDo].Node.FullNamesOfDependedOn)
 				{
 					Log("            depOn> {0}", Dep);
 				}
 			}
 			if (bShowDependentPromotions)
 			{
-				foreach (var Dep in GUBPNodes[NodeToDo].Node.DependentPromotions)
+				foreach (string Dep in GUBPNodes[NodeToDo].Node.DependentPromotions)
 				{
 					Log("            depPro> {0}", Dep);
 				}
@@ -642,15 +642,15 @@ public partial class GUBP : BuildCommand
     }
     public void SaveGraphVisualization(List<string> Nodes)
     {
-        var GraphNodes = new List<GraphNode>();
+        List<GraphNode> GraphNodes = new List<GraphNode>();
 
-        var NodeToGraphNodeMap = new Dictionary<string, GraphNode>();
+        Dictionary<string, GraphNode> NodeToGraphNodeMap = new Dictionary<string, GraphNode>();
 
-        for (var NodeIndex = 0; NodeIndex < Nodes.Count; ++NodeIndex)
+        for (int NodeIndex = 0; NodeIndex < Nodes.Count; ++NodeIndex)
         {
-            var Node = Nodes[NodeIndex];
+            string Node = Nodes[NodeIndex];
 
-            var GraphNode = new GraphNode()
+            GraphNode GraphNode = new GraphNode()
             {
                 Id = GraphNodes.Count,
                 Label = Node
@@ -660,20 +660,20 @@ public partial class GUBP : BuildCommand
         }
 
         // Connect everything together
-        var GraphEdges = new List<GraphEdge>();
+        List<GraphEdge> GraphEdges = new List<GraphEdge>();
 
-        for (var NodeIndex = 0; NodeIndex < Nodes.Count; ++NodeIndex)
+        for (int NodeIndex = 0; NodeIndex < Nodes.Count; ++NodeIndex)
         {
-            var Node = Nodes[NodeIndex];
+            string Node = Nodes[NodeIndex];
             GraphNode NodeGraphNode = NodeToGraphNodeMap[Node];
 
-            foreach (var Dep in GUBPNodes[Node].Node.FullNamesOfDependencies)
+            foreach (string Dep in GUBPNodes[Node].Node.FullNamesOfDependencies)
             {
                 GraphNode PrerequisiteFileGraphNode;
                 if (NodeToGraphNodeMap.TryGetValue(Dep, out PrerequisiteFileGraphNode))
                 {
                     // Connect a file our action is dependent on, to our action itself
-                    var NewGraphEdge = new GraphEdge()
+                    GraphEdge NewGraphEdge = new GraphEdge()
                     {
                         Id = GraphEdges.Count,
                         Source = PrerequisiteFileGraphNode,
@@ -685,13 +685,13 @@ public partial class GUBP : BuildCommand
                 }
 
             }
-            foreach (var Dep in GUBPNodes[Node].Node.FullNamesOfPseudosependencies)
+            foreach (string Dep in GUBPNodes[Node].Node.FullNamesOfPseudosependencies)
             {
                 GraphNode PrerequisiteFileGraphNode;
                 if (NodeToGraphNodeMap.TryGetValue(Dep, out PrerequisiteFileGraphNode))
                 {
                     // Connect a file our action is dependent on, to our action itself
-                    var NewGraphEdge = new GraphEdge()
+                    GraphEdge NewGraphEdge = new GraphEdge()
                     {
                         Id = GraphEdges.Count,
                         Source = PrerequisiteFileGraphNode,
@@ -725,8 +725,8 @@ public partial class GUBP : BuildCommand
 
     List<int> ConvertCLToIntList(List<string> Strings)
     {
-        var Result = new List<int>();
-        foreach (var ThisString in Strings)
+        List<int> Result = new List<int>();
+        foreach (string ThisString in Strings)
         {
             int ThisInt = int.Parse(ThisString);
             if (ThisInt < 1960000 || ThisInt > 3000000)
@@ -787,15 +787,15 @@ public partial class GUBP : BuildCommand
 
     List<string> TopologicalSort(HashSet<string> NodesToDo, Dictionary<string, bool> GUBPNodesCompleted, Dictionary<string, string> GUBPNodesControllingTriggerDotName, Dictionary<string, string> GUBPNodesControllingTrigger, string ExplicitTrigger = "", bool LocalOnly = false, bool SubSort = false, bool DoNotConsiderCompletion = false)
     {
-        var StartTime = DateTime.UtcNow;
+        DateTime StartTime = DateTime.UtcNow;
 
-        var OrdereredToDo = new List<string>();
+        List<string> OrdereredToDo = new List<string>();
 
-        var SortedAgentGroupChains = new Dictionary<string, List<string>>();
+        Dictionary<string, List<string>> SortedAgentGroupChains = new Dictionary<string, List<string>>();
         if (!SubSort)
         {
-            var AgentGroupChains = new Dictionary<string, List<string>>();
-            foreach (var NodeToDo in NodesToDo)
+            Dictionary<string, List<string>> AgentGroupChains = new Dictionary<string, List<string>>();
+            foreach (string NodeToDo in NodesToDo)
             {
                 string MyAgentGroup = GUBPNodes[NodeToDo].Node.AgentSharingGroup;
                 if (MyAgentGroup != "")
@@ -810,7 +810,7 @@ public partial class GUBP : BuildCommand
                     }
                 }
             }
-            foreach (var Chain in AgentGroupChains)
+            foreach (KeyValuePair<string, List<string>> Chain in AgentGroupChains)
             {
                 SortedAgentGroupChains.Add(Chain.Key, TopologicalSort(new HashSet<string>(Chain.Value), GUBPNodesCompleted, GUBPNodesControllingTriggerDotName, GUBPNodesControllingTrigger, ExplicitTrigger, LocalOnly, true, DoNotConsiderCompletion));
             }
@@ -824,10 +824,10 @@ public partial class GUBP : BuildCommand
             float BestPriority = -1E20f;
             string BestNode = "";
             bool BestPseudoReady = false;
-            var NonReadyAgentGroups = new HashSet<string>();
-            var NonPeudoReadyAgentGroups = new HashSet<string>();
-            var ExaminedAgentGroups = new HashSet<string>();
-            foreach (var NodeToDo in NodesToDo)
+            HashSet<string> NonReadyAgentGroups = new HashSet<string>();
+            HashSet<string> NonPeudoReadyAgentGroups = new HashSet<string>();
+            HashSet<string> ExaminedAgentGroups = new HashSet<string>();
+            foreach (string NodeToDo in NodesToDo)
             {
                 bool bReady = true;
 				bool bPseudoReady = true;
@@ -842,9 +842,9 @@ public partial class GUBP : BuildCommand
                     else
                     {
                         ExaminedAgentGroups.Add(GUBPNodes[NodeToDo].Node.AgentSharingGroup);
-                        foreach (var ChainNode in SortedAgentGroupChains[GUBPNodes[NodeToDo].Node.AgentSharingGroup])
+                        foreach (string ChainNode in SortedAgentGroupChains[GUBPNodes[NodeToDo].Node.AgentSharingGroup])
                         {
-                            foreach (var Dep in GUBPNodes[ChainNode].Node.FullNamesOfDependencies)
+                            foreach (string Dep in GUBPNodes[ChainNode].Node.FullNamesOfDependencies)
                             {
                                 if (!GUBPNodes.ContainsKey(Dep))
                                 {
@@ -861,7 +861,7 @@ public partial class GUBP : BuildCommand
                                 NonReadyAgentGroups.Add(GUBPNodes[NodeToDo].Node.AgentSharingGroup);
                                 break;
                             }
-                            foreach (var Dep in GUBPNodes[ChainNode].Node.FullNamesOfPseudosependencies)
+                            foreach (string Dep in GUBPNodes[ChainNode].Node.FullNamesOfPseudosependencies)
                             {
                                 if (!GUBPNodes.ContainsKey(Dep))
                                 {
@@ -879,7 +879,7 @@ public partial class GUBP : BuildCommand
                 }
                 else
                 {
-                    foreach (var Dep in GUBPNodes[NodeToDo].Node.FullNamesOfDependencies)
+                    foreach (string Dep in GUBPNodes[NodeToDo].Node.FullNamesOfDependencies)
                     {
                         if (!GUBPNodes.ContainsKey(Dep))
                         {
@@ -891,7 +891,7 @@ public partial class GUBP : BuildCommand
                             break;
                         }
                     }
-                    foreach (var Dep in GUBPNodes[NodeToDo].Node.FullNamesOfPseudosependencies)
+                    foreach (string Dep in GUBPNodes[NodeToDo].Node.FullNamesOfPseudosependencies)
                     {
                         if (!GUBPNodes.ContainsKey(Dep))
                         {
@@ -903,7 +903,7 @@ public partial class GUBP : BuildCommand
                             break;
                         }
                     }
-					foreach (var Dep in GUBPNodes[NodeToDo].Node.CompletedDependencies)
+					foreach (string Dep in GUBPNodes[NodeToDo].Node.CompletedDependencies)
 					{
 						if (!GUBPNodes.ContainsKey(Dep))
 						{
@@ -916,7 +916,7 @@ public partial class GUBP : BuildCommand
 						}
 					}
                 }
-                var Priority = GUBPNodes[NodeToDo].Node.Priority();
+                float Priority = GUBPNodes[NodeToDo].Node.Priority();
 
                 if (bReady && BestNode != "")
                 {
@@ -979,7 +979,7 @@ public partial class GUBP : BuildCommand
             {
                 if (!SubSort && GUBPNodes[BestNode].Node.AgentSharingGroup != "")
                 {
-                    foreach (var ChainNode in SortedAgentGroupChains[GUBPNodes[BestNode].Node.AgentSharingGroup])
+                    foreach (string ChainNode in SortedAgentGroupChains[GUBPNodes[BestNode].Node.AgentSharingGroup])
                     {
                         OrdereredToDo.Add(ChainNode);
                         NodesToDo.Remove(ChainNode);
@@ -995,14 +995,14 @@ public partial class GUBP : BuildCommand
             if (!bProgressMade && NodesToDo.Count > 0)
             {
                 Log("Cycle in GUBP, could not resolve:");
-                foreach (var NodeToDo in NodesToDo)
+                foreach (string NodeToDo in NodesToDo)
                 {
                     string Deps = "";
                     if (!SubSort && GUBPNodes[NodeToDo].Node.AgentSharingGroup != "")
                     {
-                        foreach (var ChainNode in SortedAgentGroupChains[GUBPNodes[NodeToDo].Node.AgentSharingGroup])
+                        foreach (string ChainNode in SortedAgentGroupChains[GUBPNodes[NodeToDo].Node.AgentSharingGroup])
                         {
-                            foreach (var Dep in GUBPNodes[ChainNode].Node.FullNamesOfDependencies)
+                            foreach (string Dep in GUBPNodes[ChainNode].Node.FullNamesOfDependencies)
                             {
                                 if (!SortedAgentGroupChains[GUBPNodes[NodeToDo].Node.AgentSharingGroup].Contains(Dep) && NodesToDo.Contains(Dep))
                                 {
@@ -1011,14 +1011,14 @@ public partial class GUBP : BuildCommand
                             }
                         }
                     }
-                    foreach (var Dep in GUBPNodes[NodeToDo].Node.FullNamesOfDependencies)
+                    foreach (string Dep in GUBPNodes[NodeToDo].Node.FullNamesOfDependencies)
                     {
                         if (NodesToDo.Contains(Dep))
                         {
                             Deps = Deps + Dep + " ";
                         }
                     }
-                    foreach (var Dep in GUBPNodes[NodeToDo].Node.FullNamesOfPseudosependencies)
+                    foreach (string Dep in GUBPNodes[NodeToDo].Node.FullNamesOfPseudosependencies)
                     {
                         if (NodesToDo.Contains(Dep))
                         {
@@ -1032,7 +1032,7 @@ public partial class GUBP : BuildCommand
         }
         if (!SubSort)
         {
-            var BuildDuration = (DateTime.UtcNow - StartTime).TotalMilliseconds;
+            double BuildDuration = (DateTime.UtcNow - StartTime).TotalMilliseconds;
             Log("Took {0}s to sort {1} nodes", BuildDuration / 1000, OrdereredToDo.Count);
         }
 
@@ -1058,7 +1058,7 @@ public partial class GUBP : BuildCommand
         {
             string GameNameIfAny = GUBPNodes[Node].Node.GameNameIfAnyForTempStorage();
             string NodeStoreWildCard = StoreName.Replace(CLString, "*") + "-" + GUBPNodes[Node].Node.GetFullName();
-            var History = new NodeHistory();
+            NodeHistory History = new NodeHistory();
 
             History.AllStarted = ConvertCLToIntList(TempStorage.FindTempStorageManifests(CmdEnv, NodeStoreWildCard + StartedTempStorageSuffix, false, true, GameNameIfAny));
             History.AllSucceeded = ConvertCLToIntList(TempStorage.FindTempStorageManifests(CmdEnv, NodeStoreWildCard + SucceededTempStorageSuffix, false, true, GameNameIfAny));
@@ -1072,7 +1072,7 @@ public partial class GUBP : BuildCommand
             {
                 History.LastSucceeded = History.AllSucceeded[History.AllSucceeded.Count - 1];
 
-                foreach (var Failed in History.AllFailed)
+                foreach (int Failed in History.AllFailed)
                 {
                     if (Failed > History.LastSucceeded)
                     {
@@ -1080,7 +1080,7 @@ public partial class GUBP : BuildCommand
                         History.FailedString = GUBPNode.MergeSpaceStrings(History.FailedString, String.Format("{0}", Failed));
                     }
                 }
-                foreach (var Started in History.AllStarted)
+                foreach (int Started in History.AllStarted)
                 {
                     if (Started > History.LastSucceeded && !History.Failed.Contains(Started))
                     {
@@ -1105,7 +1105,7 @@ public partial class GUBP : BuildCommand
         int NumPeople = 0;
         if (GUBPNodesHistory.ContainsKey(NodeToDo))
         {
-            var History = GUBPNodesHistory[NodeToDo];
+            NodeHistory History = GUBPNodesHistory[NodeToDo];
 			RunECTool(String.Format("setProperty \"/myWorkflow/LastGreen/{0}\" \"{1}\"", NodeToDo, History.LastSucceeded), true);
 			RunECTool(String.Format("setProperty \"/myWorkflow/LastGreen/{0}\" \"{1}\"", NodeToDo, History.FailedString), true);
 
@@ -1130,15 +1130,15 @@ public partial class GUBP : BuildCommand
                     Log(System.Diagnostics.TraceEventType.Warning, LogUtils.FormatException(Ex));
                 }
 
-                var ChangeRecords = GetChanges(History.LastSucceeded, LastNonDuplicateFail, History.LastSucceeded);
-                foreach (var Record in ChangeRecords)
+                List<P4Connection.ChangeRecord> ChangeRecords = GetChanges(History.LastSucceeded, LastNonDuplicateFail, History.LastSucceeded);
+                foreach (P4Connection.ChangeRecord Record in ChangeRecords)
                 {
                     FailCauserEMails = GUBPNode.MergeSpaceStrings(FailCauserEMails, Record.UserEmail);
                 }
                 if (!String.IsNullOrEmpty(FailCauserEMails))
                 {
                     NumPeople++;
-                    foreach (var AChar in FailCauserEMails.ToCharArray())
+                    foreach (char AChar in FailCauserEMails.ToCharArray())
                     {
                         if (AChar == ' ')
                         {
@@ -1176,7 +1176,7 @@ public partial class GUBP : BuildCommand
 		RunECTool(String.Format("setProperty \"/myWorkflow/FailCausers/{0}\" \"{1}\"", NodeToDo, FailCauserEMails));
 		RunECTool(String.Format("setProperty \"/myWorkflow/EmailNotes/{0}\" \"{1}\"", NodeToDo, EMailNote));
         {
-            var AdditionalEmails = "";
+            string AdditionalEmails = "";
 			string Causers = "";
             if (ParseParam("CIS") && !GUBPNodes[NodeToDo].Node.SendSuccessEmail() && !GUBPNodes[NodeToDo].Node.TriggerNode())
             {
@@ -1206,14 +1206,14 @@ public partial class GUBP : BuildCommand
         {
             return false;
         }
-        foreach (var Elem in A)
+        foreach (string Elem in A)
         {
             if (!B.Contains(Elem))
             {
                 return false;
             }
         }
-        foreach (var Elem in B)
+        foreach (string Elem in B)
         {
             if (!A.Contains(Elem))
             {
@@ -1225,18 +1225,18 @@ public partial class GUBP : BuildCommand
 
     int FindLastNonDuplicateFail(Dictionary<string, NodeHistory> GUBPNodesHistory, string NodeToDo, string CLString)
     {
-        var History = GUBPNodesHistory[NodeToDo];
+        NodeHistory History = GUBPNodesHistory[NodeToDo];
         int Result = P4Env.Changelist;
 
         string GameNameIfAny = GUBPNodes[NodeToDo].Node.GameNameIfAnyForTempStorage();
         string NodeStore = StoreName + "-" + GUBPNodes[NodeToDo].Node.GetFullName() + FailedTempStorageSuffix;
 
-        var BackwardsFails = new List<int>(History.AllFailed);
+        List<int> BackwardsFails = new List<int>(History.AllFailed);
         BackwardsFails.Add(P4Env.Changelist);
         BackwardsFails.Sort();
         BackwardsFails.Reverse();
         HashSet<string> CurrentErrors = null;
-        foreach (var CL in BackwardsFails)
+        foreach (int CL in BackwardsFails)
         {
             if (CL > P4Env.Changelist)
             {
@@ -1246,7 +1246,7 @@ public partial class GUBP : BuildCommand
             {
                 break;
             }
-            var ThisNodeStore = NodeStore.Replace(CLString, String.Format("{0}", CL));
+            string ThisNodeStore = NodeStore.Replace(CLString, String.Format("{0}", CL));
             TempStorage.DeleteLocalTempStorage(CmdEnv, ThisNodeStore, true); // these all clash locally, which is fine we just retrieve them from shared
 
             List<string> Files = null;
@@ -1267,7 +1267,7 @@ public partial class GUBP : BuildCommand
                 throw new AutomationException("Unexpected number of files for fail record {0}", Files.Count);
             }
             string ErrorFile = Files[0];
-            var ThisErrors = ECJobPropsUtils.ErrorsFromProps(ErrorFile);
+            HashSet<string> ThisErrors = ECJobPropsUtils.ErrorsFromProps(ErrorFile);
             if (CurrentErrors == null)
             {
                 CurrentErrors = ThisErrors;
@@ -1285,9 +1285,9 @@ public partial class GUBP : BuildCommand
     }
     List<string> GetECPropsForNode(string NodeToDo, string CLString, out string EMails, bool OnlyLateUpdates = false)
     {
-        var ECProps = new List<string>();
+        List<string> ECProps = new List<string>();
         EMails = "";		
-		var AdditonalEmails = "";
+		string AdditonalEmails = "";
 		string Causers = "";				
 		string AddEmails = ParseParamValue("AddEmails");
 		if (!String.IsNullOrEmpty(AddEmails))
@@ -1319,10 +1319,10 @@ public partial class GUBP : BuildCommand
         {
             Log("Updating node props for node {0}", NodeToDo);
 			string EMails = "";
-            var Props = GetECPropsForNode(NodeToDo, CLString, out EMails, true);
-            foreach (var Prop in Props)
+            List<string> Props = GetECPropsForNode(NodeToDo, CLString, out EMails, true);
+            foreach (string Prop in Props)
             {
-                var Parts = Prop.Split("=".ToCharArray());
+                string[] Parts = Prop.Split("=".ToCharArray());
                 RunECTool(String.Format("setProperty \"/myWorkflow/{0}\" \"{1}\"", Parts[0], Parts[1]), true);
             }			
         }
@@ -1535,7 +1535,7 @@ public partial class GUBP : BuildCommand
 
         LogVerbose("************************* CL:			                    {0}", CL);
 		LogVerbose("************************* P4Enabled:			            {0}", P4Enabled);
-        foreach (var HostPlatform in HostPlatforms)
+        foreach (UnrealTargetPlatform HostPlatform in HostPlatforms)
         {
 			LogVerbose("*************************   HostPlatform:		            {0}", HostPlatform.ToString()); 
         }
@@ -1565,18 +1565,18 @@ public partial class GUBP : BuildCommand
 		Dictionary<string, string> GUBPNodesControllingTrigger = new Dictionary<string, string>();
         Dictionary<string, string> GUBPNodesControllingTriggerDotName = new Dictionary<string, string>();
 
-		var FullNodeDependedOnBy = new Dictionary<string, string>();
-		var FullNodeDependentPromotions = new Dictionary<string, string>();		
-		var SeparatePromotables = new List<string>();
+		Dictionary<string, string> FullNodeDependedOnBy = new Dictionary<string, string>();
+		Dictionary<string, string> FullNodeDependentPromotions = new Dictionary<string, string>();		
+		List<string> SeparatePromotables = new List<string>();
         {
-            foreach (var NodeToDo in GUBPNodes)
+            foreach (KeyValuePair<string, NodeInfo> NodeToDo in GUBPNodes)
             {
                 if (GUBPNodes[NodeToDo.Key].Node.IsSeparatePromotable())
                 {
                     SeparatePromotables.Add(GUBPNodes[NodeToDo.Key].Node.GetFullName());
                     List<string> Dependencies = new List<string>();
                     Dependencies = GetECDependencies(NodeToDo.Key);
-                    foreach (var Dep in Dependencies)
+                    foreach (string Dep in Dependencies)
                     {
                         if (!GUBPNodes.ContainsKey(Dep))
                         {
@@ -1595,16 +1595,16 @@ public partial class GUBP : BuildCommand
 
 			// Compute all the frequencies
 			Dictionary<string, int> FrequencyOverrides = ApplyFrequencyBarriers();
-            foreach (var NodeToDo in GUBPNodes)
+            foreach (KeyValuePair<string, NodeInfo> NodeToDo in GUBPNodes)
             {
                 ComputeDependentCISFrequencyQuantumShift(NodeToDo.Key, FrequencyOverrides);
             }
 
-            foreach (var NodeToDo in GUBPNodes)
+            foreach (KeyValuePair<string, NodeInfo> NodeToDo in GUBPNodes)
             {
-                var Deps = GUBPNodes[NodeToDo.Key].Node.DependentPromotions;
+                List<string> Deps = GUBPNodes[NodeToDo.Key].Node.DependentPromotions;
                 string All = "";
-                foreach (var Dep in Deps)
+                foreach (string Dep in Deps)
                 {
                     if (All != "")
                     {
@@ -1631,7 +1631,7 @@ public partial class GUBP : BuildCommand
             if (!String.IsNullOrEmpty(ExplicitTrigger))
             {
                 bool bFoundIt = false;
-                foreach (var Node in GUBPNodes)
+                foreach (KeyValuePair<string, NodeInfo> Node in GUBPNodes)
                 {
                     if (Node.Value.Node.GetFullName().Equals(ExplicitTrigger, StringComparison.InvariantCultureIgnoreCase))
                     {
@@ -1652,7 +1652,7 @@ public partial class GUBP : BuildCommand
             {
                 if (bSkipTriggers)
                 {
-                    foreach (var Node in GUBPNodes)
+                    foreach (KeyValuePair<string, NodeInfo> Node in GUBPNodes)
                     {
                         if (Node.Value.Node.TriggerNode() && Node.Value.Node.RunInEC())
                         {
@@ -1665,10 +1665,10 @@ public partial class GUBP : BuildCommand
 		if (bPreflightBuild)
 		{
 			LogVerbose("Culling triggers and downstream for preflight builds ");
-			var NewNodesToDo = new HashSet<string>();
-			foreach (var NodeToDo in NodesToDo)
+			HashSet<string> NewNodesToDo = new HashSet<string>();
+			foreach (string NodeToDo in NodesToDo)
 			{
-				var TriggerDot = GetControllingTriggerDotName(NodeToDo, GUBPNodesControllingTriggerDotName, GUBPNodesControllingTrigger);
+				string TriggerDot = GetControllingTriggerDotName(NodeToDo, GUBPNodesControllingTriggerDotName, GUBPNodesControllingTrigger);
 				if (TriggerDot == "" && !GUBPNodes[NodeToDo].Node.TriggerNode())
 				{
 					LogVerbose("  Keeping {0}", NodeToDo);
@@ -1687,14 +1687,14 @@ public partial class GUBP : BuildCommand
 
         LogVerbose("******* Caching completion");
         {
-            var StartTime = DateTime.UtcNow;
-            foreach (var Node in NodesToDo)
+            DateTime StartTime = DateTime.UtcNow;
+            foreach (string Node in NodesToDo)
             {
                 LogVerbose("** {0}", Node);
                 NodeIsAlreadyComplete(GUBPNodesCompleted, Node, LocalOnly); // cache these now to avoid spam later
                 GetControllingTriggerDotName(Node, GUBPNodesControllingTriggerDotName, GUBPNodesControllingTrigger);
             }
-            var BuildDuration = (DateTime.UtcNow - StartTime).TotalMilliseconds;
+            double BuildDuration = (DateTime.UtcNow - StartTime).TotalMilliseconds;
 			LogVerbose("Took {0}s to cache completion for {1} nodes", BuildDuration / 1000, NodesToDo.Count);
         }
         /*if (CLString != "" && StoreName.Contains(CLString) && !ParseParam("NoHistory"))
@@ -1826,13 +1826,13 @@ public partial class GUBP : BuildCommand
 			if (!String.IsNullOrEmpty(NodeSpec))
 			{
 				List<string> Nodes = new List<string>(NodeSpec.Split('+'));
-				foreach (var NodeArg in Nodes)
+				foreach (string NodeArg in Nodes)
 				{
-					var NodeName = NodeArg.Trim();
+					string NodeName = NodeArg.Trim();
 					bool bFoundAnything = false;
 					if (!String.IsNullOrEmpty(NodeName))
 					{
-						foreach (var Node in GUBPNodes)
+						foreach (KeyValuePair<string, NodeInfo> Node in GUBPNodes)
 						{
 							if (Node.Value.Node.GetFullName().Equals(NodeArg, StringComparison.InvariantCultureIgnoreCase) ||
 								Node.Value.Node.AgentSharingGroup.Equals(NodeArg, StringComparison.InvariantCultureIgnoreCase)
@@ -1857,12 +1857,12 @@ public partial class GUBP : BuildCommand
 		if (!String.IsNullOrEmpty(GameSpec))
 		{
 			List<string> Games = new List<string>(GameSpec.Split('+'));
-			foreach (var GameArg in Games)
+			foreach (string GameArg in Games)
 			{
-				var GameName = GameArg.Trim();
+				string GameName = GameArg.Trim();
 				if (!String.IsNullOrEmpty(GameName))
 				{
-					foreach (var GameProj in Branch.CodeProjects)
+					foreach (BranchInfo.BranchUProject GameProj in Branch.CodeProjects)
 					{
 						if (GameProj.GameName.Equals(GameName, StringComparison.InvariantCultureIgnoreCase))
 						{
@@ -1870,7 +1870,7 @@ public partial class GUBP : BuildCommand
 							{
 								NodesToDo.Add(GameAggregatePromotableNode.StaticGetFullName(GameProj));
 							}
-							foreach (var Node in GUBPNodes)
+							foreach (KeyValuePair<string, NodeInfo> Node in GUBPNodes)
 							{
 								if (Node.Value.Node.GameNameIfAnyForTempStorage() == GameProj.GameName)
 								{
@@ -1883,11 +1883,11 @@ public partial class GUBP : BuildCommand
 					}
 					if (GameName != null)
 					{
-						foreach (var GameProj in Branch.NonCodeProjects)
+						foreach (BranchInfo.BranchUProject GameProj in Branch.NonCodeProjects)
 						{
 							if (GameProj.GameName.Equals(GameName, StringComparison.InvariantCultureIgnoreCase))
 							{
-								foreach (var Node in GUBPNodes)
+								foreach (KeyValuePair<string, NodeInfo> Node in GUBPNodes)
 								{
 									if (Node.Value.Node.GameNameIfAnyForTempStorage() == GameProj.GameName)
 									{
@@ -1908,7 +1908,7 @@ public partial class GUBP : BuildCommand
 		if (NodesToDo.Count == 0)
 		{
 			LogVerbose("No nodes specified, adding all nodes");
-			foreach (var Node in GUBPNodes)
+			foreach (KeyValuePair<string, NodeInfo> Node in GUBPNodes)
 			{
 				NodesToDo.Add(Node.Key);
 			}
@@ -1916,7 +1916,7 @@ public partial class GUBP : BuildCommand
 		else if (TimeIndex != 0)
 		{
 			LogVerbose("Check to make sure we didn't ask for nodes that will be culled by time index");
-			foreach (var NodeToDo in NodesToDo)
+			foreach (string NodeToDo in NodesToDo)
 			{
 				if (TimeIndex % (1 << GUBPNodes[NodeToDo].Node.DependentCISFrequencyQuantumShift()) != 0)
 				{
@@ -1926,7 +1926,7 @@ public partial class GUBP : BuildCommand
 		}
 
 		LogVerbose("Desired Nodes");
-		foreach (var NodeToDo in NodesToDo)
+		foreach (string NodeToDo in NodesToDo)
 		{
 			LogVerbose("  {0}", NodeToDo);
 		}
@@ -1938,12 +1938,12 @@ public partial class GUBP : BuildCommand
 			while (!bDoneWithDependencies)
 			{
 				bDoneWithDependencies = true;
-				var Fringe = new HashSet<string>();
-				foreach (var NodeToDo in GUBPNodes)
+				HashSet<string> Fringe = new HashSet<string>();
+				foreach (KeyValuePair<string, NodeInfo> NodeToDo in GUBPNodes)
 				{
 					if (!NodesToDo.Contains(NodeToDo.Key))
 					{
-						foreach (var Dep in GUBPNodes[NodeToDo.Key].Node.FullNamesOfDependencies)
+						foreach (string Dep in GUBPNodes[NodeToDo.Key].Node.FullNamesOfDependencies)
 						{
 							if (!GUBPNodes.ContainsKey(Dep))
 							{
@@ -1955,7 +1955,7 @@ public partial class GUBP : BuildCommand
 								bDoneWithDependencies = false;
 							}
 						}
-						foreach (var Dep in GUBPNodes[NodeToDo.Key].Node.FullNamesOfPseudosependencies)
+						foreach (string Dep in GUBPNodes[NodeToDo.Key].Node.FullNamesOfPseudosependencies)
 						{
 							if (!GUBPNodes.ContainsKey(Dep))
 							{
@@ -1976,13 +1976,13 @@ public partial class GUBP : BuildCommand
 			while (!bDoneWithDependencies)
 			{
 				bDoneWithDependencies = true;
-				var Fringe = new HashSet<string>();
-				foreach (var NodeToDo in NodesToDo)
+				HashSet<string> Fringe = new HashSet<string>();
+				foreach (string NodeToDo in NodesToDo)
 				{
 					if (BranchOptions.NodesToRemovePseudoDependencies.Count > 0)
 					{
-						var ListOfRetainedNodesToRemovePseudoDependencies = BranchOptions.NodesToRemovePseudoDependencies;
-						foreach (var NodeToRemovePseudoDependencies in BranchOptions.NodesToRemovePseudoDependencies)
+						List<string> ListOfRetainedNodesToRemovePseudoDependencies = BranchOptions.NodesToRemovePseudoDependencies;
+						foreach (string NodeToRemovePseudoDependencies in BranchOptions.NodesToRemovePseudoDependencies)
 						{
 							if (NodeToDo == NodeToRemovePseudoDependencies)
 							{
@@ -1993,7 +1993,7 @@ public partial class GUBP : BuildCommand
 						}
 						BranchOptions.NodesToRemovePseudoDependencies = ListOfRetainedNodesToRemovePseudoDependencies;
 					}
-					foreach (var Dep in GUBPNodes[NodeToDo].Node.FullNamesOfDependencies)
+					foreach (string Dep in GUBPNodes[NodeToDo].Node.FullNamesOfDependencies)
 					{
 						if (!GUBPNodes.ContainsKey(Dep))
 						{
@@ -2005,7 +2005,7 @@ public partial class GUBP : BuildCommand
 							bDoneWithDependencies = false;
 						}
 					}
-					foreach (var Dep in GUBPNodes[NodeToDo].Node.FullNamesOfPseudosependencies)
+					foreach (string Dep in GUBPNodes[NodeToDo].Node.FullNamesOfPseudosependencies)
 					{
 						if (!GUBPNodes.ContainsKey(Dep))
 						{
@@ -2019,8 +2019,8 @@ public partial class GUBP : BuildCommand
 		if (TimeIndex != 0)
 		{
 			LogVerbose("Culling based on time index");
-			var NewNodesToDo = new HashSet<string>();
-			foreach (var NodeToDo in NodesToDo)
+			HashSet<string> NewNodesToDo = new HashSet<string>();
+			foreach (string NodeToDo in NodesToDo)
 			{
 				if (TimeIndex % (1 << GUBPNodes[NodeToDo].Node.DependentCISFrequencyQuantumShift()) == 0)
 				{
@@ -2037,8 +2037,8 @@ public partial class GUBP : BuildCommand
 		//Remove Plat if specified
 		if (WithoutLinux)
 		{
-			var NewNodesToDo = new HashSet<string>();
-			foreach (var NodeToDo in NodesToDo)
+			HashSet<string> NewNodesToDo = new HashSet<string>();
+			foreach (string NodeToDo in NodesToDo)
 			{
 				if (!GUBPNodes[NodeToDo].Node.GetFullName().Contains("Linux"))
 				{
@@ -2054,13 +2054,13 @@ public partial class GUBP : BuildCommand
 		//find things that depend on our nodes and setup commander dictionary
 		if (!bOnlyNode)
 		{
-			foreach (var NodeToDo in NodesToDo)
+			foreach (string NodeToDo in NodesToDo)
 			{
 				if (!GUBPNodes[NodeToDo].Node.IsAggregate() && !GUBPNodes[NodeToDo].Node.IsTest())
 				{
 					List<string> ECDependencies = new List<string>();
 					ECDependencies = GetECDependencies(NodeToDo);
-					foreach (var Dep in ECDependencies)
+					foreach (string Dep in ECDependencies)
 					{
 						if (!GUBPNodes.ContainsKey(Dep))
 						{
@@ -2073,11 +2073,11 @@ public partial class GUBP : BuildCommand
 					}
 				}
 			}
-			foreach (var NodeToDo in NodesToDo)
+			foreach (string NodeToDo in NodesToDo)
 			{
-				var Deps = GUBPNodes[NodeToDo].Node.FullNamesOfDependedOn;
+				List<string> Deps = GUBPNodes[NodeToDo].Node.FullNamesOfDependedOn;
 				string All = "";
-				foreach (var Dep in Deps)
+				foreach (string Dep in Deps)
 				{
 					if (All != "")
 					{
@@ -2096,9 +2096,9 @@ public partial class GUBP : BuildCommand
 		FullNodeDirectDependencies = new Dictionary<string,string>();
 
 		Log("******* {0} GUBP Nodes", GUBPNodes.Count);
-		var SortedNodes = TopologicalSort(new HashSet<string>(GUBPNodes.Keys), null/*GUBPNodesCompleted*/, GUBPNodesControllingTriggerDotName, GUBPNodesControllingTrigger, LocalOnly: true, DoNotConsiderCompletion: true);
-		var DependencyStart = DateTime.Now.ToString();
-		foreach (var Node in SortedNodes)
+		List<string> SortedNodes = TopologicalSort(new HashSet<string>(GUBPNodes.Keys), null/*GUBPNodesCompleted*/, GUBPNodesControllingTriggerDotName, GUBPNodesControllingTrigger, LocalOnly: true, DoNotConsiderCompletion: true);
+		string DependencyStart = DateTime.Now.ToString();
+		foreach (string Node in SortedNodes)
 		{
 			string Note = GetControllingTriggerDotName(Node, GUBPNodesControllingTriggerDotName, GUBPNodesControllingTrigger);
 			if (Note == "")
@@ -2111,9 +2111,9 @@ public partial class GUBP : BuildCommand
 			}
 			if (GUBPNodes[Node].Node.RunInEC())
 			{
-				var Deps = GetECDependencies(Node);
+				List<string> Deps = GetECDependencies(Node);
 				string All = "";
-				foreach (var Dep in Deps)
+				foreach (string Dep in Deps)
 				{
 					if (All != "")
 					{
@@ -2130,15 +2130,15 @@ public partial class GUBP : BuildCommand
 				LogVerbose("  {0}: {1} [Aggregate]", Node, Note);
 			}
 		}
-		var DependencyFinish = DateTime.Now.ToString();
+		string DependencyFinish = DateTime.Now.ToString();
 		PrintCSVFile(String.Format("UAT,GetDependencies,{0},{1}", DependencyStart, DependencyFinish));
 	}
 
 	private void ValidateDependencyNames()
 	{
-		foreach (var NodeToDo in GUBPNodes)
+		foreach (KeyValuePair<string, NodeInfo> NodeToDo in GUBPNodes)
 		{
-			foreach (var Dep in GUBPNodes[NodeToDo.Key].Node.FullNamesOfDependencies)
+			foreach (string Dep in GUBPNodes[NodeToDo.Key].Node.FullNamesOfDependencies)
 			{
 				if (!GUBPNodes.ContainsKey(Dep))
 				{
@@ -2149,7 +2149,7 @@ public partial class GUBP : BuildCommand
 					throw new AutomationException("Node {0} has a self arc.", NodeToDo.Key);
 				}
 			}
-			foreach (var Dep in GUBPNodes[NodeToDo.Key].Node.FullNamesOfPseudosependencies)
+			foreach (string Dep in GUBPNodes[NodeToDo.Key].Node.FullNamesOfPseudosependencies)
 			{
 				if (!GUBPNodes.ContainsKey(Dep))
 				{
@@ -2169,8 +2169,8 @@ public partial class GUBP : BuildCommand
 		{
 			throw new AutomationException("Can't have -CIS without P4 support");
 		}
-		var P4IndexFileP4 = CombinePaths(PathSeparator.Slash, CommandUtils.P4Env.BuildRootP4, "Engine", "Build", "CISCounter.txt");
-		var P4IndexFileLocal = CombinePaths(CmdEnv.LocalRoot, "Engine", "Build", "CISCounter.txt");
+		string P4IndexFileP4 = CombinePaths(PathSeparator.Slash, CommandUtils.P4Env.BuildRootP4, "Engine", "Build", "CISCounter.txt");
+		string P4IndexFileLocal = CombinePaths(CmdEnv.LocalRoot, "Engine", "Build", "CISCounter.txt");
 		for(int Retry = 0; Retry < 20; Retry++)
 		{
 			int NowMinutes = (int)((DateTime.UtcNow - new DateTime(2014, 1, 1)).TotalMinutes);
@@ -2221,8 +2221,8 @@ public partial class GUBP : BuildCommand
 					}
 				}
 			}
-			var Data = ReadAllText(P4IndexFileLocal);
-			var Parts = Data.Split(" ".ToCharArray());
+			string Data = ReadAllText(P4IndexFileLocal);
+			string[] Parts = Data.Split(" ".ToCharArray());
 			int Index = int.Parse(Parts[0]);
 			int Minutes = int.Parse(Parts[1]);
 
@@ -2255,7 +2255,7 @@ public partial class GUBP : BuildCommand
 				}
 			}
 			{
-				var Line = String.Format("{0} {1}", NewIndex, NowMinutes);
+				string Line = String.Format("{0} {1}", NewIndex, NowMinutes);
 				LogVerbose("Attempting to write {0} with {1}", P4IndexFileP4, Line);
 				int WorkingCL = -1;
 				try
@@ -2284,10 +2284,10 @@ public partial class GUBP : BuildCommand
 	private List<string> FindUnfinishedTriggers(bool bSkipTriggers, bool LocalOnly, string ExplicitTrigger, Dictionary<string, bool> GUBPNodesCompleted, List<string> OrdereredToDo)
 	{
 		// find all unfinished triggers, excepting the one we are triggering right now
-		var UnfinishedTriggers = new List<string>();
+		List<string> UnfinishedTriggers = new List<string>();
 		if (!bSkipTriggers)
 		{
-			foreach (var NodeToDo in OrdereredToDo)
+			foreach (string NodeToDo in OrdereredToDo)
 			{
 				if (GUBPNodes[NodeToDo].Node.TriggerNode() && !NodeIsAlreadyComplete(GUBPNodesCompleted, NodeToDo, LocalOnly))
 				{
@@ -2303,7 +2303,7 @@ public partial class GUBP : BuildCommand
 
 	private void CheckSortOrder(bool LocalOnly, Dictionary<string, bool> GUBPNodesCompleted, List<string> OrdereredToDo)
 	{
-		foreach (var NodeToDo in OrdereredToDo)
+		foreach (string NodeToDo in OrdereredToDo)
 		{
 			if (GUBPNodes[NodeToDo].Node.TriggerNode() && (GUBPNodes[NodeToDo].Node.IsSticky() || NodeIsAlreadyComplete(GUBPNodesCompleted, NodeToDo, LocalOnly))) // these sticky triggers are ok, everything is already completed anyway
 			{
@@ -2314,7 +2314,7 @@ public partial class GUBP : BuildCommand
 				bHasTests = true;
 			}
 			int MyIndex = OrdereredToDo.IndexOf(NodeToDo);
-			foreach (var Dep in GUBPNodes[NodeToDo].Node.FullNamesOfDependencies)
+			foreach (string Dep in GUBPNodes[NodeToDo].Node.FullNamesOfDependencies)
 			{
 				int DepIndex = OrdereredToDo.IndexOf(Dep);
 				if (DepIndex >= MyIndex)
@@ -2322,7 +2322,7 @@ public partial class GUBP : BuildCommand
 					throw new AutomationException("Topological sort error, node {0} has a dependency of {1} which sorted after it.", NodeToDo, Dep);
 				}
 			}
-			foreach (var Dep in GUBPNodes[NodeToDo].Node.FullNamesOfPseudosependencies)
+			foreach (string Dep in GUBPNodes[NodeToDo].Node.FullNamesOfPseudosependencies)
 			{
 				int DepIndex = OrdereredToDo.IndexOf(Dep);
 				if (DepIndex >= MyIndex)
@@ -2340,33 +2340,33 @@ public partial class GUBP : BuildCommand
 		{
 			throw new AutomationException("No nodes to do!");
 		}
-		var ECProps = new List<string>();
+		List<string> ECProps = new List<string>();
 		ECProps.Add(String.Format("TimeIndex={0}", TimeIndex));
-		foreach (var NodePair in FullNodeList)
+		foreach (KeyValuePair<string, string> NodePair in FullNodeList)
 		{
 			ECProps.Add(string.Format("AllNodes/{0}={1}", NodePair.Key, NodePair.Value));
 		}
-		foreach (var NodePair in FullNodeDirectDependencies)
+		foreach (KeyValuePair<string, string> NodePair in FullNodeDirectDependencies)
 		{
 			ECProps.Add(string.Format("DirectDependencies/{0}={1}", NodePair.Key, NodePair.Value));
 		}
-		foreach (var NodePair in FullNodeListSortKey)
+		foreach (KeyValuePair<string, int> NodePair in FullNodeListSortKey)
 		{
 			ECProps.Add(string.Format("SortKey/{0}={1}", NodePair.Key, NodePair.Value));
 		}
-		foreach (var NodePair in FullNodeDependedOnBy)
+		foreach (KeyValuePair<string, string> NodePair in FullNodeDependedOnBy)
 		{
 			ECProps.Add(string.Format("DependedOnBy/{0}={1}", NodePair.Key, NodePair.Value));
 		}
-		foreach (var NodePair in FullNodeDependentPromotions)
+		foreach (KeyValuePair<string, string> NodePair in FullNodeDependentPromotions)
 		{
 			ECProps.Add(string.Format("DependentPromotions/{0}={1}", NodePair.Key, NodePair.Value));
 		}
-		foreach (var Node in SeparatePromotables)
+		foreach (string Node in SeparatePromotables)
 		{
 			ECProps.Add(string.Format("PossiblePromotables/{0}={1}", Node, ""));
 		}
-		var ECJobProps = new List<string>();
+		List<string> ECJobProps = new List<string>();
 		if (ExplicitTrigger != "")
 		{
 			ECJobProps.Add("IsRoot=0");
@@ -2376,10 +2376,10 @@ public partial class GUBP : BuildCommand
 			ECJobProps.Add("IsRoot=1");
 		}
 
-		var FilteredOrdereredToDo = new List<string>();
-		var StartFilterTime = DateTime.Now.ToString();
+		List<string> FilteredOrdereredToDo = new List<string>();
+		string StartFilterTime = DateTime.Now.ToString();
 		// remove nodes that have unfinished triggers
-		foreach (var NodeToDo in OrdereredToDo)
+		foreach (string NodeToDo in OrdereredToDo)
 		{
 			string ControllingTrigger = GetControllingTrigger(NodeToDo, GUBPNodesControllingTrigger);
 			bool bNoUnfinishedTriggers = !UnfinishedTriggers.Contains(ControllingTrigger);
@@ -2404,16 +2404,16 @@ public partial class GUBP : BuildCommand
 		}
 
 		OrdereredToDo = FilteredOrdereredToDo;
-		var FinishFilterTime = DateTime.Now.ToString();
+		string FinishFilterTime = DateTime.Now.ToString();
 		PrintCSVFile(String.Format("UAT,FilterNodes,{0},{1}", StartFilterTime, FinishFilterTime));
 		Log("*********** EC Nodes, in order.");
 		PrintNodes(this, OrdereredToDo, GUBPNodesCompleted, GUBPNodesHistory, GUBPNodesControllingTriggerDotName, GUBPNodesControllingTrigger, LocalOnly, UnfinishedTriggers);
-		var FinishNodePrint = DateTime.Now.ToString();
+		string FinishNodePrint = DateTime.Now.ToString();
 		PrintCSVFile(String.Format("UAT,SetupCommanderPrint,{0},{1}", FinishFilterTime, FinishNodePrint));
 		// here we are just making sure everything before the explicit trigger is completed.
 		if (!String.IsNullOrEmpty(ExplicitTrigger))
 		{
-			foreach (var NodeToDo in FilteredOrdereredToDo)
+			foreach (string NodeToDo in FilteredOrdereredToDo)
 			{
 				if (GUBPNodes[NodeToDo].Node.RunInEC() && !NodeIsAlreadyComplete(GUBPNodesCompleted, NodeToDo, LocalOnly) && NodeToDo != ExplicitTrigger && !NodeDependsOn(ExplicitTrigger, NodeToDo)) // if something is already finished, we don't put it into EC
 				{
@@ -2433,7 +2433,7 @@ public partial class GUBP : BuildCommand
 		StepList.Add("$ec->setTimeout(600);");
 		StepList.Add("my $batch = $ec->newBatch(\"serial\");");
 		// sticky nodes are ones that we run on the main agent. We run then first and they must not be intermixed with parallel jobs
-		foreach (var NodeToDo in OrdereredToDo)
+		foreach (string NodeToDo in OrdereredToDo)
 		{
 			if (GUBPNodes[NodeToDo].Node.RunInEC() && !NodeIsAlreadyComplete(GUBPNodesCompleted, NodeToDo, LocalOnly)) // if something is already finished, we don't put it into EC
 			{
@@ -2454,7 +2454,7 @@ public partial class GUBP : BuildCommand
 		}
 
 		string ParentPath = ParseParamValue("ParentPath");
-		var StartPerlOutput = DateTime.Now.ToString();
+		string StartPerlOutput = DateTime.Now.ToString();
 		string BaseArgs = String.Format("$batch->createJobStep({{parentPath => '{0}'", ParentPath);
 
 		bool bHasNoop = false;
@@ -2466,10 +2466,10 @@ public partial class GUBP : BuildCommand
 			bHasNoop = true;
 		}
 
-		var FakeECArgs = new List<string>();
-		var AgentGroupChains = new Dictionary<string, List<string>>();
-		var StickyChain = new List<string>();
-		foreach (var NodeToDo in OrdereredToDo)
+		List<string> FakeECArgs = new List<string>();
+		Dictionary<string, List<string>> AgentGroupChains = new Dictionary<string, List<string>>();
+		List<string> StickyChain = new List<string>();
+		foreach (string NodeToDo in OrdereredToDo)
 		{
 			if (GUBPNodes[NodeToDo].Node.RunInEC() && !NodeIsAlreadyComplete(GUBPNodesCompleted, NodeToDo, LocalOnly)) // if something is already finished, we don't put it into EC  
 			{
@@ -2494,12 +2494,12 @@ public partial class GUBP : BuildCommand
 				}
 			}
 		}
-		foreach (var NodeToDo in OrdereredToDo)
+		foreach (string NodeToDo in OrdereredToDo)
 		{
 			if (GUBPNodes[NodeToDo].Node.RunInEC() && !NodeIsAlreadyComplete(GUBPNodesCompleted, NodeToDo, LocalOnly)) // if something is already finished, we don't put it into EC  
 			{
 				string EMails;
-				var NodeProps = GetECPropsForNode(NodeToDo, CLString, out EMails);
+				List<string> NodeProps = GetECPropsForNode(NodeToDo, CLString, out EMails);
 				ECProps.AddRange(NodeProps);
 
 				bool Sticky = GUBPNodes[NodeToDo].Node.IsSticky();
@@ -2532,9 +2532,9 @@ public partial class GUBP : BuildCommand
 				Args = Args + "]";
 				string PreCondition = "";
 				string RunCondition = "";
-				var UncompletedEcDeps = new List<string>();
-				var UncompletedCompletedDeps = new List<string>();
-				var PreConditionUncompletedEcDeps = new List<string>();
+				List<string> UncompletedEcDeps = new List<string>();
+				List<string> UncompletedCompletedDeps = new List<string>();
+				List<string> PreConditionUncompletedEcDeps = new List<string>();
 				string MyAgentGroup = GUBPNodes[NodeToDo].Node.AgentSharingGroup;
 				bool bDoNestedJobstep = false;
 				bool bDoFirstNestedJobstep = false;
@@ -2545,8 +2545,8 @@ public partial class GUBP : BuildCommand
 
 				PreconditionParentPath = ParentPath;
 				{
-					var EcDeps = GetECDependencies(NodeToDo);
-					foreach (var Dep in EcDeps)
+					List<string> EcDeps = GetECDependencies(NodeToDo);
+					foreach (string Dep in EcDeps)
 					{
 						if (GUBPNodes[Dep].Node.RunInEC() && !NodeIsAlreadyComplete(GUBPNodesCompleted, Dep, LocalOnly) && OrdereredToDo.Contains(Dep)) // if something is already finished, we don't put it into EC
 						{
@@ -2559,12 +2559,12 @@ public partial class GUBP : BuildCommand
 					}
 				}
 
-				foreach (var Dep in UncompletedEcDeps)
+				foreach (string Dep in UncompletedEcDeps)
 				{
 					PreConditionUncompletedEcDeps.Add(Dep);
 				}
-				var CompletedDeps = GetCompletedOnlyDependencies(NodeToDo);
-				foreach (var Dep in CompletedDeps)
+				List<string> CompletedDeps = GetCompletedOnlyDependencies(NodeToDo);
+				foreach (string Dep in CompletedDeps)
 				{
 					if (GUBPNodes[Dep].Node.RunInEC() && !NodeIsAlreadyComplete(GUBPNodesCompleted, Dep, LocalOnly) && OrdereredToDo.Contains(Dep)) // if something is already finished, we don't put it into EC
 					{
@@ -2582,7 +2582,7 @@ public partial class GUBP : BuildCommand
 
 					PreConditionUncompletedEcDeps = new List<string>();
 
-					var MyChain = AgentGroupChains[MyAgentGroup];
+					List<string> MyChain = AgentGroupChains[MyAgentGroup];
 					int MyIndex = MyChain.IndexOf(NodeToDo);
 					if (MyIndex > 0)
 					{
@@ -2592,10 +2592,10 @@ public partial class GUBP : BuildCommand
 					{
 						bDoFirstNestedJobstep = bDoNestedJobstep;
 						// to avoid idle agents (and also EC doesn't actually reserve our agent!), we promote all dependencies to the first one
-						foreach (var Chain in MyChain)
+						foreach (string Chain in MyChain)
 						{
-							var EcDeps = GetECDependencies(Chain);
-							foreach (var Dep in EcDeps)
+							List<string> EcDeps = GetECDependencies(Chain);
+							foreach (string Dep in EcDeps)
 							{
 								if (GUBPNodes[Dep].Node.RunInEC() && !NodeIsAlreadyComplete(GUBPNodesCompleted, Dep, LocalOnly) && OrdereredToDo.Contains(Dep)) // if something is already finished, we don't put it into EC
 								{
@@ -2614,7 +2614,7 @@ public partial class GUBP : BuildCommand
 				}
 				if (GUBPNodes[NodeToDo].Node.IsSticky())
 				{
-					var MyChain = StickyChain;
+					List<string> MyChain = StickyChain;
 					int MyIndex = MyChain.IndexOf(NodeToDo);
 					if (MyIndex > 0)
 					{
@@ -2625,8 +2625,8 @@ public partial class GUBP : BuildCommand
 					}
 					else
 					{
-						var EcDeps = GetECDependencies(NodeToDo);
-						foreach (var Dep in EcDeps)
+						List<string> EcDeps = GetECDependencies(NodeToDo);
+						foreach (string Dep in EcDeps)
 						{
 							if (GUBPNodes[Dep].Node.RunInEC() && !NodeIsAlreadyComplete(GUBPNodesCompleted, Dep, LocalOnly) && OrdereredToDo.Contains(Dep)) // if something is already finished, we don't put it into EC
 							{
@@ -2651,7 +2651,7 @@ public partial class GUBP : BuildCommand
 					PreCondition = "\"\\$\" . \"[/javascript if(";
 					// these run "parallel", but we add preconditions to serialize them
 					int Index = 0;
-					foreach (var Dep in PreConditionUncompletedEcDeps)
+					foreach (string Dep in PreConditionUncompletedEcDeps)
 					{
 						PreCondition = PreCondition + "getProperty('" + GetJobStep(PreconditionParentPath, Dep) + "/status\') == \'completed\'";
 						Index++;
@@ -2667,7 +2667,7 @@ public partial class GUBP : BuildCommand
 					PreCondition = "\"\\$\" . \"[/javascript if(";
 					// these run "parallel", but we add preconditions to serialize them
 					int Index = 0;
-					foreach (var Dep in CompletedDeps)
+					foreach (string Dep in CompletedDeps)
 					{
 						if (GUBPNodes[Dep].Node.RunInEC())
 						{
@@ -2686,7 +2686,7 @@ public partial class GUBP : BuildCommand
 				{
 					RunCondition = "\"\\$\" . \"[/javascript if(";
 					int Index = 0;
-					foreach (var Dep in UncompletedEcDeps)
+					foreach (string Dep in UncompletedEcDeps)
 					{
 						RunCondition = RunCondition + "((\'\\$\" . \"[" + GetJobStep(PreconditionParentPath, Dep) + "/outcome]\' == \'success\') || ";
 						RunCondition = RunCondition + "(\'\\$\" . \"[" + GetJobStep(PreconditionParentPath, Dep) + "/outcome]\' == \'warning\'))";
@@ -2797,11 +2797,11 @@ public partial class GUBP : BuildCommand
 
 				if (MyAgentGroup != "" && !bDoNestedJobstep)
 				{
-					var MyChain = AgentGroupChains[MyAgentGroup];
+					List<string> MyChain = AgentGroupChains[MyAgentGroup];
 					int MyIndex = MyChain.IndexOf(NodeToDo);
 					if (MyIndex == MyChain.Count - 1)
 					{
-						var RelPreCondition = "\"\\$\" . \"[/javascript if(";
+						string RelPreCondition = "\"\\$\" . \"[/javascript if(";
 						// this runs "parallel", but we a precondition to serialize it
 						RelPreCondition = RelPreCondition + "getProperty('" + PreconditionParentPath + "/jobSteps[" + NodeToDo + "]/status') == 'completed'";
 						RelPreCondition = RelPreCondition + ") true;]\"";
@@ -2814,7 +2814,7 @@ public partial class GUBP : BuildCommand
 			}
 		}
 		WriteECPerl(StepList);
-		var FinishPerlOutput = DateTime.Now.ToString();
+		string FinishPerlOutput = DateTime.Now.ToString();
 		PrintCSVFile(String.Format("UAT,PerlOutput,{0},{1}", StartPerlOutput, FinishPerlOutput));
 		RunECTool(String.Format("setProperty \"/myWorkflow/HasTests\" \"{0}\"", bHasTests));
 		{
@@ -2831,7 +2831,7 @@ public partial class GUBP : BuildCommand
 		}
 		if (bFakeEC)
 		{
-			foreach (var Args in FakeECArgs)
+			foreach (string Args in FakeECArgs)
 			{
 				RunUAT(CmdEnv, Args);
 			}
@@ -2840,8 +2840,8 @@ public partial class GUBP : BuildCommand
 
 	void ExecuteNodes(List<string> OrdereredToDo, bool bOnlyNode, bool bFakeEC, bool LocalOnly, bool bSaveSharedTempStorage, Dictionary<string, bool> GUBPNodesCompleted, Dictionary<string, NodeHistory> GUBPNodesHistory, string CLString, string FakeFail)
 	{
-        var BuildProductToNodeMap = new Dictionary<string, string>();
-		foreach (var NodeToDo in OrdereredToDo)
+        Dictionary<string, string> BuildProductToNodeMap = new Dictionary<string, string>();
+		foreach (string NodeToDo in OrdereredToDo)
         {
             if (GUBPNodes[NodeToDo].Node.BuildProducts != null || GUBPNodes[NodeToDo].Node.AllDependencyBuildProducts != null)
             {
@@ -2850,7 +2850,7 @@ public partial class GUBP : BuildCommand
 
             GUBPNodes[NodeToDo].Node.AllDependencyBuildProducts = new List<string>();
             GUBPNodes[NodeToDo].Node.AllDependencies = new List<string>();
-            foreach (var Dep in GUBPNodes[NodeToDo].Node.FullNamesOfDependencies)
+            foreach (string Dep in GUBPNodes[NodeToDo].Node.FullNamesOfDependencies)
             {
                 GUBPNodes[NodeToDo].Node.AddAllDependent(Dep);
                 if (GUBPNodes[Dep].Node.AllDependencies == null)
@@ -2862,7 +2862,7 @@ public partial class GUBP : BuildCommand
                 }
                 else
 				{
-					foreach (var DepDep in GUBPNodes[Dep].Node.AllDependencies)
+					foreach (string DepDep in GUBPNodes[Dep].Node.AllDependencies)
 					{
 						GUBPNodes[NodeToDo].Node.AddAllDependent(DepDep);
 					}
@@ -2877,7 +2877,7 @@ public partial class GUBP : BuildCommand
                 }
                 else
                 {
-                    foreach (var Prod in GUBPNodes[Dep].Node.BuildProducts)
+                    foreach (string Prod in GUBPNodes[Dep].Node.BuildProducts)
                     {
                         GUBPNodes[NodeToDo].Node.AddDependentBuildProduct(Prod);
                     }
@@ -2885,7 +2885,7 @@ public partial class GUBP : BuildCommand
                     {
                         throw new AutomationException("Node {0} was not processed yet2?  Processing {1}", Dep, NodeToDo);
                     }
-                    foreach (var Prod in GUBPNodes[Dep].Node.AllDependencyBuildProducts)
+                    foreach (string Prod in GUBPNodes[Dep].Node.AllDependencyBuildProducts)
                     {
                         GUBPNodes[NodeToDo].Node.AddDependentBuildProduct(Prod);
                     }
@@ -2947,7 +2947,7 @@ public partial class GUBP : BuildCommand
                 {
                     SaveStatus(NodeToDo, StartedTempStorageSuffix, NodeStoreName, bSaveSharedTempStorage, GameNameIfAny);
                 }
-				var BuildDuration = 0.0;
+				double BuildDuration = 0.0;
                 try
                 {
                     if (!String.IsNullOrEmpty(FakeFail) && FakeFail.Equals(NodeToDo, StringComparison.InvariantCultureIgnoreCase))
@@ -2962,10 +2962,10 @@ public partial class GUBP : BuildCommand
                     else
                     {                        
 						Log("***** Building GUBP Node {0} for {1}", NodeToDo, NodeStoreName);
-						var StartTime = DateTime.UtcNow;
-                        var StartBuild = DateTime.Now.ToString();
+						DateTime StartTime = DateTime.UtcNow;
+                        string StartBuild = DateTime.Now.ToString();
                         GUBPNodes[NodeToDo].Node.DoBuild(this);
-                        var FinishBuild = DateTime.Now.ToString();
+                        string FinishBuild = DateTime.Now.ToString();
                         if (IsBuildMachine)
                         {
                             PrintCSVFile(String.Format("UAT,DoBuild.{0},{1},{2}", NodeToDo, StartBuild, FinishBuild));
@@ -2975,12 +2975,12 @@ public partial class GUBP : BuildCommand
                     }
                     if (!GUBPNodes[NodeToDo].Node.IsAggregate())
                     {
-						var StoreDuration = 0.0;
-						var StartTime = DateTime.UtcNow;
-                        var StartStore = DateTime.Now.ToString();
+						double StoreDuration = 0.0;
+						DateTime StartTime = DateTime.UtcNow;
+                        string StartStore = DateTime.Now.ToString();
                         TempStorage.StoreToTempStorage(CmdEnv, NodeStoreName, GUBPNodes[NodeToDo].Node.BuildProducts, !bSaveSharedTempStorage, GameNameIfAny, StorageRootIfAny);
 						StoreDuration = (DateTime.UtcNow - StartTime).TotalMilliseconds / 1000;
-                        var FinishStore = DateTime.Now.ToString();                        
+                        string FinishStore = DateTime.Now.ToString();                        
 						Log("Took {0} seconds to store build products", StoreDuration);
 						if (IsBuildMachine)
 						{
@@ -2989,14 +2989,14 @@ public partial class GUBP : BuildCommand
 						}
                         if (ParseParam("StompCheck"))
                         {
-                            foreach (var Dep in GUBPNodes[NodeToDo].Node.AllDependencies)
+                            foreach (string Dep in GUBPNodes[NodeToDo].Node.AllDependencies)
                             {
                                 try
                                 {
                                     bool WasLocal;
-                                    var StartRetrieve = DateTime.Now.ToString();
+                                    string StartRetrieve = DateTime.Now.ToString();
 									TempStorage.RetrieveFromTempStorage(CmdEnv, NodeStoreName, out WasLocal, GameNameIfAny, StorageRootIfAny);
-                                    var FinishRetrieve = DateTime.Now.ToString();
+                                    string FinishRetrieve = DateTime.Now.ToString();
                                     if (IsBuildMachine)
                                     {
                                         PrintCSVFile(String.Format("UAT,RetrieveBuildProducts,{0},{1}", StartRetrieve, FinishRetrieve));
@@ -3020,18 +3020,18 @@ public partial class GUBP : BuildCommand
                 {
                     if (SaveSuccessRecords)
                     {
-                        var StartNodeHistory = DateTime.Now.ToString();
+                        string StartNodeHistory = DateTime.Now.ToString();
                         UpdateNodeHistory(GUBPNodesHistory, NodeToDo, CLString);
-                        var FinishNodeHistory = DateTime.Now.ToString();                        
+                        string FinishNodeHistory = DateTime.Now.ToString();                        
                         SaveStatus(NodeToDo, FailedTempStorageSuffix, NodeStoreName, bSaveSharedTempStorage, GameNameIfAny, ParseParamValue("MyJobStepId"));
-                        var FinishSaveStatus = DateTime.Now.ToString();                        
+                        string FinishSaveStatus = DateTime.Now.ToString();                        
                         UpdateECProps(NodeToDo, CLString);
-                        var FinishECPropUpdate = DateTime.Now.ToString();
+                        string FinishECPropUpdate = DateTime.Now.ToString();
                         
 						if (IsBuildMachine)
 						{
 							GetFailureEmails(GUBPNodesHistory, NodeToDo, CLString);
-                            var FinishFailEmails = DateTime.Now.ToString();
+                            string FinishFailEmails = DateTime.Now.ToString();
                             PrintCSVFile(String.Format("UAT,UpdateNodeHistory,{0},{1}", StartNodeHistory, FinishNodeHistory));
                             PrintCSVFile(String.Format("UAT,SaveNodeStatus,{0},{1}", FinishNodeHistory, FinishSaveStatus));
                             PrintCSVFile(String.Format("UAT,UpdateECProps,{0},{1}", FinishSaveStatus, FinishECPropUpdate));
@@ -3045,7 +3045,7 @@ public partial class GUBP : BuildCommand
 
                     if (GUBPNodesHistory.ContainsKey(NodeToDo))
                     {
-                        var History = GUBPNodesHistory[NodeToDo];
+                        NodeHistory History = GUBPNodesHistory[NodeToDo];
                         Log("Changes since last green *********************************");
                         Log("");
                         Log("");
@@ -3087,18 +3087,18 @@ public partial class GUBP : BuildCommand
                 }
                 if (SaveSuccessRecords) 
                 {
-                    var StartNodeHistory = DateTime.Now.ToString();
+                    string StartNodeHistory = DateTime.Now.ToString();
                     UpdateNodeHistory(GUBPNodesHistory, NodeToDo, CLString);
-                    var FinishNodeHistory = DateTime.Now.ToString();
+                    string FinishNodeHistory = DateTime.Now.ToString();
                     SaveStatus(NodeToDo, SucceededTempStorageSuffix, NodeStoreName, bSaveSharedTempStorage, GameNameIfAny);
-                    var FinishSaveStatus = DateTime.Now.ToString();                    
+                    string FinishSaveStatus = DateTime.Now.ToString();                    
                     UpdateECProps(NodeToDo, CLString);
-                    var FinishECPropUpdate = DateTime.Now.ToString();
+                    string FinishECPropUpdate = DateTime.Now.ToString();
                     
 					if (IsBuildMachine)
 					{
 						GetFailureEmails(GUBPNodesHistory, NodeToDo, CLString);
-                        var FinishFailEmails = DateTime.Now.ToString();
+                        string FinishFailEmails = DateTime.Now.ToString();
                         PrintCSVFile(String.Format("UAT,UpdateNodeHistory,{0},{1}", StartNodeHistory, FinishNodeHistory));
                         PrintCSVFile(String.Format("UAT,SaveNodeStatus,{0},{1}", FinishNodeHistory, FinishSaveStatus));
                         PrintCSVFile(String.Format("UAT,UpdateECProps,{0},{1}", FinishSaveStatus, FinishECPropUpdate));
@@ -3107,7 +3107,7 @@ public partial class GUBP : BuildCommand
 					UpdateECBuildTime(NodeToDo, BuildDuration);
                 }
             }
-            foreach (var Product in GUBPNodes[NodeToDo].Node.BuildProducts)
+            foreach (string Product in GUBPNodes[NodeToDo].Node.BuildProducts)
             {
                 if (BuildProductToNodeMap.ContainsKey(Product))
                 {
