Pass all nodes around as NodeInfo objects, rather than using strings and constantly looking everything up in a dictionary.

diff --git a/Engine/Source/Programs/AutomationTool/GUBP/GUBP.cs b/Engine/Source/Programs/AutomationTool/GUBP/GUBP.cs
--- a/Engine/Source/Programs/AutomationTool/GUBP/GUBP.cs
+++ b/Engine/Source/Programs/AutomationTool/GUBP/GUBP.cs
@@ -32,7 +32,14 @@ public partial class GUBP : BuildCommand
 
 	class NodeInfo
 	{
+		public string Name;
 		public GUBPNode Node;
+
+		public override string ToString()
+		{
+			System.Diagnostics.Trace.TraceWarning("Implicit conversion from NodeInfo to string\n{0}", Environment.StackTrace);
+			return Name;
+		}
 	}
 
     class NodeHistory
@@ -41,7 +48,7 @@ public partial class GUBP : BuildCommand
         public int LastFailed = 0;
         public List<int> InProgress = new List<int>();
         public string InProgressString = "";
-        public List<int> Failed = new List<int>();
+		public List<int> Failed = new List<int>();
         public string FailedString = "";
         public List<int> AllStarted = new List<int>();
         public List<int> AllSucceeded = new List<int>();
@@ -55,7 +62,7 @@ public partial class GUBP : BuildCommand
         {
             throw new AutomationException("Attempt to add a duplicate node {0}", Node.GetFullName());
         }
-        GUBPNodes.Add(Name, new NodeInfo{ Node = Node });
+        GUBPNodes.Add(Name, new NodeInfo{ Name = Name, Node = Node });
         return Name;
     }
 
@@ -94,22 +101,22 @@ public partial class GUBP : BuildCommand
         GUBPNodes[Node].Node.FullNamesOfPseudosependencies.Clear();
     }
 
-    List<string> GetDependencies(string NodeToDo, bool bFlat = false, bool ECOnly = false)
+    List<NodeInfo> GetDependencies(NodeInfo NodeToDo, bool bFlat = false, bool ECOnly = false)
     {
-        List<string> Result = new List<string>();
-        foreach (string Node in GUBPNodes[NodeToDo].Node.FullNamesOfDependencies)
+        List<NodeInfo> Result = new List<NodeInfo>();
+        foreach (NodeInfo Dependency in NodeToDo.Node.FullNamesOfDependencies.Select(x => GUBPNodes[x]))
         {
-            bool Usable = GUBPNodes[Node].Node.RunInEC() || !ECOnly;
+            bool Usable = Dependency.Node.RunInEC() || !ECOnly;
             if (Usable)
             {
-                if (!Result.Contains(Node))
+                if (!Result.Contains(Dependency))
                 {
-                    Result.Add(Node);
+                    Result.Add(Dependency);
                 }
             }
             if (bFlat || !Usable)
             {
-                foreach (string RNode in GetDependencies(Node, bFlat, ECOnly))
+                foreach (NodeInfo RNode in GetDependencies(Dependency, bFlat, ECOnly))
                 {
                     if (!Result.Contains(RNode))
                     {
@@ -118,19 +125,19 @@ public partial class GUBP : BuildCommand
                 }
             }
         }
-        foreach (string Node in GUBPNodes[NodeToDo].Node.FullNamesOfPseudosependencies)
+        foreach (NodeInfo PseudoDependency in NodeToDo.Node.FullNamesOfPseudosependencies.Select(x => GUBPNodes[x]))
         {
-            bool Usable = GUBPNodes[Node].Node.RunInEC() || !ECOnly;
+            bool Usable = PseudoDependency.Node.RunInEC() || !ECOnly;
             if (Usable)
             {
-                if (!Result.Contains(Node))
+                if (!Result.Contains(PseudoDependency))
                 {
-                    Result.Add(Node);
+                    Result.Add(PseudoDependency);
                 }
             }
             if (bFlat || !Usable)
             {
-                foreach (string RNode in GetDependencies(Node, bFlat, ECOnly))
+                foreach (NodeInfo RNode in GetDependencies(PseudoDependency, bFlat, ECOnly))
                 {
                     if (!Result.Contains(RNode))
                     {
@@ -142,22 +149,22 @@ public partial class GUBP : BuildCommand
         return Result;
     }
 
-	List<string> GetCompletedOnlyDependencies(string NodeToDo, bool bFlat = false, bool ECOnly = true)
+	List<NodeInfo> GetCompletedOnlyDependencies(NodeInfo NodeToDo, bool bFlat = false, bool ECOnly = true)
 	{
-		List<string> Result = new List<string>();
-		foreach (string Node in GUBPNodes[NodeToDo].Node.CompletedDependencies)
+		List<NodeInfo> Result = new List<NodeInfo>();
+		foreach (NodeInfo CompletedDependency in NodeToDo.Node.CompletedDependencies.Select(x => GUBPNodes[x]))
 		{
-			bool Usable = GUBPNodes[Node].Node.RunInEC() || !ECOnly;
+			bool Usable = CompletedDependency.Node.RunInEC() || !ECOnly;
 			if(Usable)
 			{
-				if(!Result.Contains(Node))
+				if(!Result.Contains(CompletedDependency))
 				{
-					Result.Add(Node);
+					Result.Add(CompletedDependency);
 				}
 			}
 			if (bFlat || !Usable)
 			{
-				foreach (string RNode in GetCompletedOnlyDependencies(Node, bFlat, ECOnly))
+				foreach (NodeInfo RNode in GetCompletedOnlyDependencies(CompletedDependency, bFlat, ECOnly))
 				{
 					if (!Result.Contains(RNode))
 					{
@@ -168,74 +175,76 @@ public partial class GUBP : BuildCommand
 		}
 		return Result;
 	}
-    List<string> GetECDependencies(string NodeToDo, bool bFlat = false)
+    List<NodeInfo> GetECDependencies(NodeInfo NodeToDo, bool bFlat = false)
     {
         return GetDependencies(NodeToDo, bFlat, true);
     }
-    bool NodeDependsOn(string Rootward, string Leafward)
+    bool NodeDependsOn(NodeInfo Rootward, NodeInfo Leafward)
     {
-        List<string> Deps = GetDependencies(Leafward, true);
+        List<NodeInfo> Deps = GetDependencies(Leafward, true);
         return Deps.Contains(Rootward);
     }
 
-    string GetControllingTrigger(string NodeToDo, Dictionary<string, string> GUBPNodesControllingTrigger)
+    NodeInfo GetControllingTrigger(NodeInfo NodeToDo, Dictionary<NodeInfo, NodeInfo> GUBPNodesControllingTrigger)
     {
         if (GUBPNodesControllingTrigger.ContainsKey(NodeToDo))
         {
             return GUBPNodesControllingTrigger[NodeToDo];
         }
-        string Result = "";
-        foreach (string Node in GUBPNodes[NodeToDo].Node.FullNamesOfDependencies)
-        {
-            if (!GUBPNodes.ContainsKey(Node))
+        NodeInfo Result = null;
+        foreach (string DependencyName in NodeToDo.Node.FullNamesOfDependencies)
+		{
+			NodeInfo Dependency;
+            if (!GUBPNodes.TryGetValue(DependencyName, out Dependency))
             {
-                throw new AutomationException("Dependency {0} in {1} not found.", Node, NodeToDo);
+                throw new AutomationException("Dependency {0} in {1} not found.", DependencyName, NodeToDo.Name);
             }
-            bool IsTrigger = GUBPNodes[Node].Node.TriggerNode();
+            bool IsTrigger = Dependency.Node.TriggerNode();
             if (IsTrigger)
             {
-                if (Node != Result && !string.IsNullOrEmpty(Result))
+                if (Dependency != Result && Result != null)
                 {
-                    throw new AutomationException("Node {0} has two controlling triggers {1} and {2}.", NodeToDo, Node, Result);
+                    throw new AutomationException("Node {0} has two controlling triggers {1} and {2}.", NodeToDo.Name, Dependency.Name, Result);
                 }
-                Result = Node;
+                Result = Dependency;
             }
             else
             {
-                string NewResult = GetControllingTrigger(Node, GUBPNodesControllingTrigger);
-                if (!String.IsNullOrEmpty(NewResult))
+                NodeInfo NewResult = GetControllingTrigger(Dependency, GUBPNodesControllingTrigger);
+                if (NewResult != null)
                 {
-                    if (NewResult != Result && !string.IsNullOrEmpty(Result))
+                    if (NewResult != Result && Result != null)
                     {
-                        throw new AutomationException("Node {0} has two controlling triggers {1} and {2}.", NodeToDo, NewResult, Result);
+                        throw new AutomationException("Node {0} has two controlling triggers {1} and {2}.", NodeToDo.Name, NewResult.Name, Result);
                     }
                     Result = NewResult;
                 }
             }
         }
-        foreach (string Node in GUBPNodes[NodeToDo].Node.FullNamesOfPseudosependencies)
+        foreach (string PseudoDependencyName in NodeToDo.Node.FullNamesOfPseudosependencies)
         {
-            if (!GUBPNodes.ContainsKey(Node))
+			NodeInfo PseudoDependency;
+            if (!GUBPNodes.TryGetValue(PseudoDependencyName, out PseudoDependency))
             {
-                throw new AutomationException("Pseudodependency {0} in {1} not found.", Node, NodeToDo);
+                throw new AutomationException("Pseudodependency {0} in {1} not found.", PseudoDependencyName, NodeToDo.Name);
             }
-            bool IsTrigger = GUBPNodes[Node].Node.TriggerNode();
+            bool IsTrigger = PseudoDependency.Node.TriggerNode();
             if (IsTrigger)
             {
-                if (Node != Result && !string.IsNullOrEmpty(Result))
+                if (PseudoDependency != Result && Result != null)
                 {
-                    throw new AutomationException("Node {0} has two controlling triggers {1} and {2}.", NodeToDo, Node, Result);
+                    throw new AutomationException("Node {0} has two controlling triggers {1} and {2}.", NodeToDo.Name, PseudoDependency.Name, Result);
                 }
-                Result = Node;
+                Result = PseudoDependency;
             }
             else
             {
-                string NewResult = GetControllingTrigger(Node, GUBPNodesControllingTrigger);
-                if (!String.IsNullOrEmpty(NewResult))
+                NodeInfo NewResult = GetControllingTrigger(PseudoDependency, GUBPNodesControllingTrigger);
+                if (NewResult != null)
                 {
-                    if (NewResult != Result && !string.IsNullOrEmpty(Result))
+                    if (NewResult != Result && Result != null)
                     {
-                        throw new AutomationException("Node {0} has two controlling triggers {1} and {2}.", NodeToDo, NewResult, Result);
+                        throw new AutomationException("Node {0} has two controlling triggers {1} and {2}.", NodeToDo.Name, NewResult.Name, Result);
                     }
                     Result = NewResult;
                 }
@@ -244,18 +253,18 @@ public partial class GUBP : BuildCommand
         GUBPNodesControllingTrigger.Add(NodeToDo, Result);
         return Result;
     }
-    string GetControllingTriggerDotName(string NodeToDo, Dictionary<string, string> GUBPNodesControllingTriggerDotName, Dictionary<string, string> GUBPNodesControllingTrigger)
+    string GetControllingTriggerDotName(NodeInfo NodeToDo, Dictionary<NodeInfo, string> GUBPNodesControllingTriggerDotName, Dictionary<NodeInfo, NodeInfo> GUBPNodesControllingTrigger)
     {
         if (GUBPNodesControllingTriggerDotName.ContainsKey(NodeToDo))
         {
             return GUBPNodesControllingTriggerDotName[NodeToDo];
         }
         string Result = "";
-        string WorkingNode = NodeToDo;
+        NodeInfo WorkingNode = NodeToDo;
         while (true)
         {
-            string ThisResult = GetControllingTrigger(WorkingNode, GUBPNodesControllingTrigger);
-            if (ThisResult == "")
+            NodeInfo ThisResult = GetControllingTrigger(WorkingNode, GUBPNodesControllingTrigger);
+            if (ThisResult == null)
             {
                 break;
             }
@@ -263,17 +272,17 @@ public partial class GUBP : BuildCommand
             {
                 Result = "." + Result;
             }
-            Result = ThisResult + Result;
+            Result = ThisResult.Name + Result;
             WorkingNode = ThisResult;
         }
         GUBPNodesControllingTriggerDotName.Add(NodeToDo, Result);
         return Result;
     }
 
-    public string CISFrequencyQuantumShiftString(string NodeToDo)
+    string CISFrequencyQuantumShiftString(NodeInfo NodeToDo)
     {
         string FrequencyString = "";
-        int Quantum = GUBPNodes[NodeToDo].Node.DependentCISFrequencyQuantumShift();
+        int Quantum = NodeToDo.Node.DependentCISFrequencyQuantumShift();
         if (Quantum > 0)
         {
 			int TimeQuantum = 20;
@@ -294,29 +303,29 @@ public partial class GUBP : BuildCommand
         return FrequencyString;
     }
 
-    public int ComputeDependentCISFrequencyQuantumShift(string NodeToDo, Dictionary<string, int> FrequencyOverrides)
+    int ComputeDependentCISFrequencyQuantumShift(NodeInfo NodeToDo, Dictionary<string, int> FrequencyOverrides)
     {
-        int Result = GUBPNodes[NodeToDo].Node.ComputedDependentCISFrequencyQuantumShift;        
+        int Result = NodeToDo.Node.ComputedDependentCISFrequencyQuantumShift;        
         if (Result < 0)
         {
-            Result = GUBPNodes[NodeToDo].Node.CISFrequencyQuantumShift(this);
-            Result = GetFrequencyForNode(this, GUBPNodes[NodeToDo].Node.GetFullName(), Result);
+            Result = NodeToDo.Node.CISFrequencyQuantumShift(this);
+            Result = GetFrequencyForNode(this, NodeToDo, Result);
 
 			int FrequencyOverride;
-			if(FrequencyOverrides.TryGetValue(NodeToDo, out FrequencyOverride) && Result > FrequencyOverride)
+			if(FrequencyOverrides.TryGetValue(NodeToDo.Name, out FrequencyOverride) && Result > FrequencyOverride)
 			{
 				Result = FrequencyOverride;
 			}
 
-            foreach (string Dep in GUBPNodes[NodeToDo].Node.FullNamesOfDependencies)
+            foreach (NodeInfo Dep in NodeToDo.Node.FullNamesOfDependencies.Select(x => GUBPNodes[x]))
             {
                 Result = Math.Max(ComputeDependentCISFrequencyQuantumShift(Dep, FrequencyOverrides), Result);
             }
-            foreach (string Dep in GUBPNodes[NodeToDo].Node.FullNamesOfPseudosependencies)
+            foreach (NodeInfo Dep in NodeToDo.Node.FullNamesOfPseudosependencies.Select(x => GUBPNodes[x]))
             {
                 Result = Math.Max(ComputeDependentCISFrequencyQuantumShift(Dep, FrequencyOverrides), Result);
             }
-			foreach (string Dep in GUBPNodes[NodeToDo].Node.CompletedDependencies)
+			foreach (NodeInfo Dep in NodeToDo.Node.CompletedDependencies.Select(x => GUBPNodes[x]))
 			{
 				Result = Math.Max(ComputeDependentCISFrequencyQuantumShift(Dep, FrequencyOverrides), Result);
 			}
@@ -324,19 +333,19 @@ public partial class GUBP : BuildCommand
             {
                 throw new AutomationException("Failed to compute shift.");
             }
-            GUBPNodes[NodeToDo].Node.ComputedDependentCISFrequencyQuantumShift = Result;
+            NodeToDo.Node.ComputedDependentCISFrequencyQuantumShift = Result;
         }
         return Result;
     }
 
-    bool NodeIsAlreadyComplete(Dictionary<string, bool> GUBPNodesCompleted, string NodeToDo, bool LocalOnly)
+    bool NodeIsAlreadyComplete(Dictionary<string, bool> GUBPNodesCompleted, NodeInfo NodeToDo, bool LocalOnly)
     {
-        if (GUBPNodesCompleted.ContainsKey(NodeToDo))
+        if (GUBPNodesCompleted.ContainsKey(NodeToDo.Name))
         {
-            return GUBPNodesCompleted[NodeToDo];
+            return GUBPNodesCompleted[NodeToDo.Name];
         }
-        string NodeStoreName = StoreName + "-" + GUBPNodes[NodeToDo].Node.GetFullName();
-        string GameNameIfAny = GUBPNodes[NodeToDo].Node.GameNameIfAnyForTempStorage();
+        string NodeStoreName = StoreName + "-" + NodeToDo.Node.GetFullName();
+        string GameNameIfAny = NodeToDo.Node.GameNameIfAnyForTempStorage();
         bool Result;
         if (LocalOnly)
         {
@@ -352,13 +361,13 @@ public partial class GUBP : BuildCommand
         }
         if (Result)
         {
-            LogVerbose("***** GUBP Trigger Node was already triggered {0} -> {1} : {2}", GUBPNodes[NodeToDo].Node.GetFullName(), GameNameIfAny, NodeStoreName);
+            LogVerbose("***** GUBP Trigger Node was already triggered {0} -> {1} : {2}", NodeToDo.Name, GameNameIfAny, NodeStoreName);
         }
         else
         {
-            LogVerbose("***** GUBP Trigger Node was NOT yet triggered {0} -> {1} : {2}", GUBPNodes[NodeToDo].Node.GetFullName(), GameNameIfAny, NodeStoreName);
+            LogVerbose("***** GUBP Trigger Node was NOT yet triggered {0} -> {1} : {2}", NodeToDo.Name, GameNameIfAny, NodeStoreName);
         }
-        GUBPNodesCompleted.Add(NodeToDo, Result);
+        GUBPNodesCompleted.Add(NodeToDo.Name, Result);
         return Result;
     }
     string RunECTool(string Args, bool bQuiet = false)
@@ -384,7 +393,7 @@ public partial class GUBP : BuildCommand
         string ECPerlFile = CommandUtils.CombinePaths(CommandUtils.CmdEnv.LogFolder, "jobsteps.pl");
         WriteAllLines_NoExceptions(ECPerlFile, Args.ToArray());
     }
-    string GetEMailListForNode(GUBP bp, string NodeToDo, string Emails, string Causers)
+    string GetEMailListForNode(GUBP bp, NodeInfo NodeToDo, string Emails, string Causers)
     {        
         string BranchForEmail = "";
         if (P4Enabled)
@@ -393,7 +402,7 @@ public partial class GUBP : BuildCommand
         }
         return HackEmails(Emails, Causers, BranchForEmail, NodeToDo);
     }
-    int GetFrequencyForNode(GUBP bp, string NodeToDo, int BaseFrequency)
+    int GetFrequencyForNode(GUBP bp, NodeInfo NodeToDo, int BaseFrequency)
     {
         return HackFrequency(bp, BranchName, NodeToDo, BaseFrequency);
     }
@@ -495,7 +504,7 @@ public partial class GUBP : BuildCommand
         Log("Took {0}s to get P4 history", BuildDuration / 1000);
 
     }
-    void PrintNodes(GUBP bp, List<string> Nodes, Dictionary<string, bool> GUBPNodesCompleted, Dictionary<string, NodeHistory> GUBPNodesHistory, Dictionary<string, string> GUBPNodesControllingTriggerDotName, Dictionary<string, string> GUBPNodesControllingTrigger, bool LocalOnly, List<string> UnfinishedTriggers = null)
+    void PrintNodes(GUBP bp, List<NodeInfo> Nodes, Dictionary<string, bool> GUBPNodesCompleted, Dictionary<string, NodeHistory> GUBPNodesHistory, Dictionary<NodeInfo, string> GUBPNodesControllingTriggerDotName, Dictionary<NodeInfo, NodeInfo> GUBPNodesControllingTrigger, bool LocalOnly, List<NodeInfo> UnfinishedTriggers = null)
     {
         bool bShowAllChanges = bp.ParseParam("AllChanges") && GUBPNodesHistory != null;
         bool bShowChanges = (bp.ParseParam("Changes") && GUBPNodesHistory != null) || bShowAllChanges;
@@ -511,9 +520,9 @@ public partial class GUBP : BuildCommand
         bool bShowTriggers = true;
         string LastControllingTrigger = "";
         string LastAgentGroup = "";
-        foreach (string NodeToDo in Nodes)
+        foreach (NodeInfo NodeToDo in Nodes)
         {
-            if (ECOnly && !GUBPNodes[NodeToDo].Node.RunInEC())
+            if (ECOnly && !NodeToDo.Node.RunInEC())
             {
                 continue;
             }
@@ -533,7 +542,7 @@ public partial class GUBP : BuildCommand
                         string Finished = "";
                         if (UnfinishedTriggers != null)
                         {
-                            string MyShortControllingTrigger = GetControllingTrigger(NodeToDo, GUBPNodesControllingTrigger);
+                            NodeInfo MyShortControllingTrigger = GetControllingTrigger(NodeToDo, GUBPNodesControllingTrigger);
                             if (UnfinishedTriggers.Contains(MyShortControllingTrigger))
                             {
                                 Finished = "(not yet triggered)";
@@ -547,14 +556,14 @@ public partial class GUBP : BuildCommand
                     }
                 }
             }
-            if (GUBPNodes[NodeToDo].Node.AgentSharingGroup != LastAgentGroup && GUBPNodes[NodeToDo].Node.AgentSharingGroup != "")
+            if (NodeToDo.Node.AgentSharingGroup != LastAgentGroup && NodeToDo.Node.AgentSharingGroup != "")
             {
-                Log("    Agent Group: {0}", GUBPNodes[NodeToDo].Node.AgentSharingGroup);
+                Log("    Agent Group: {0}", NodeToDo.Node.AgentSharingGroup);
             }
-            LastAgentGroup = GUBPNodes[NodeToDo].Node.AgentSharingGroup;
+            LastAgentGroup = NodeToDo.Node.AgentSharingGroup;
 
-            string Agent = GUBPNodes[NodeToDo].Node.ECAgentString();
-			if(ParseParamValue("AgentOverride") != "" && !GUBPNodes[NodeToDo].Node.GetFullName().Contains("Mac"))
+            string Agent = NodeToDo.Node.ECAgentString();
+			if(ParseParamValue("AgentOverride") != "" && !NodeToDo.Node.GetFullName().Contains("Mac"))
 			{
 				Agent = ParseParamValue("AgentOverride");
 			}
@@ -562,7 +571,7 @@ public partial class GUBP : BuildCommand
             {
                 Agent = "[" + Agent + "]";
             }
-            string MemoryReq = "[" + GUBPNodes[NodeToDo].Node.AgentMemoryRequirement(bp).ToString() + "]";            
+            string MemoryReq = "[" + NodeToDo.Node.AgentMemoryRequirement(bp).ToString() + "]";            
             if(MemoryReq == "[0]")
             {
                 MemoryReq = "";                
@@ -571,19 +580,19 @@ public partial class GUBP : BuildCommand
 
             Log("      {0}{1}{2}{3}{4}{5}{6} {7}  {8}",
                 (LastAgentGroup != "" ? "  " : ""),
-                NodeToDo,
+                NodeToDo.Name,
                 FrequencyString,
                 NodeIsAlreadyComplete(GUBPNodesCompleted, NodeToDo, LocalOnly) ? " - (Completed)" : "",
-                GUBPNodes[NodeToDo].Node.TriggerNode() ? " - (TriggerNode)" : "",
-                GUBPNodes[NodeToDo].Node.IsSticky() ? " - (Sticky)" : "",				
+                NodeToDo.Node.TriggerNode() ? " - (TriggerNode)" : "",
+                NodeToDo.Node.IsSticky() ? " - (Sticky)" : "",				
                 Agent,
                 MemoryReq,
                 EMails,
-                ECProc ? GUBPNodes[NodeToDo].Node.ECProcedure() : ""
+                ECProc ? NodeToDo.Node.ECProcedure() : ""
                 );
-            if (bShowHistory && GUBPNodesHistory.ContainsKey(NodeToDo))
+            if (bShowHistory && GUBPNodesHistory.ContainsKey(NodeToDo.Name))
             {
-                NodeHistory History = GUBPNodesHistory[NodeToDo];
+                NodeHistory History = GUBPNodesHistory[NodeToDo.Name];
 
                 if (bShowDetailedHistory)
                 {
@@ -599,48 +608,48 @@ public partial class GUBP : BuildCommand
             }
             if (bShowDependencies)
             {
-                foreach (string Dep in GUBPNodes[NodeToDo].Node.FullNamesOfDependencies)
+                foreach (string Dep in NodeToDo.Node.FullNamesOfDependencies)
                 {
                     Log("            dep> {0}", Dep);
                 }
-                foreach (string Dep in GUBPNodes[NodeToDo].Node.FullNamesOfPseudosependencies)
+                foreach (string Dep in NodeToDo.Node.FullNamesOfPseudosependencies)
                 {
                     Log("           pdep> {0}", Dep);
                 }
-				foreach (string Dep in GUBPNodes[NodeToDo].Node.CompletedDependencies)
+				foreach (string Dep in NodeToDo.Node.CompletedDependencies)
 				{
 					Log("			cdep>{0}", Dep);
 				}
             }
             if (bShowECDependencies)
             {
-                foreach (string Dep in GetECDependencies(NodeToDo))
+                foreach (NodeInfo Dep in GetECDependencies(NodeToDo))
                 {
-                    Log("           {0}", Dep);
+                    Log("           {0}", Dep.Name);
                 }
-				foreach (string Dep in GetCompletedOnlyDependencies(NodeToDo))
+				foreach (NodeInfo Dep in GetCompletedOnlyDependencies(NodeToDo))
 				{
-					Log("           compDep> {0}", Dep);
+					Log("           compDep> {0}", Dep.Name);
 				}
             }
 
 			if(bShowDependednOn)
 			{
-				foreach (string Dep in GUBPNodes[NodeToDo].Node.FullNamesOfDependedOn)
+				foreach (string Dep in NodeToDo.Node.FullNamesOfDependedOn)
 				{
 					Log("            depOn> {0}", Dep);
 				}
 			}
 			if (bShowDependentPromotions)
 			{
-				foreach (string Dep in GUBPNodes[NodeToDo].Node.DependentPromotions)
+				foreach (string Dep in NodeToDo.Node.DependentPromotions)
 				{
 					Log("            depPro> {0}", Dep);
 				}
 			}			
         }
     }
-    public void SaveGraphVisualization(List<string> Nodes)
+    void SaveGraphVisualization(List<NodeInfo> Nodes)
     {
         List<GraphNode> GraphNodes = new List<GraphNode>();
 
@@ -648,7 +657,7 @@ public partial class GUBP : BuildCommand
 
         for (int NodeIndex = 0; NodeIndex < Nodes.Count; ++NodeIndex)
         {
-            string Node = Nodes[NodeIndex];
+            string Node = Nodes[NodeIndex].Name;
 
             GraphNode GraphNode = new GraphNode()
             {
@@ -664,7 +673,7 @@ public partial class GUBP : BuildCommand
 
         for (int NodeIndex = 0; NodeIndex < Nodes.Count; ++NodeIndex)
         {
-            string Node = Nodes[NodeIndex];
+            string Node = Nodes[NodeIndex].Name;
             GraphNode NodeGraphNode = NodeToGraphNodeMap[Node];
 
             foreach (string Dep in GUBPNodes[Node].Node.FullNamesOfDependencies)
@@ -738,7 +747,7 @@ public partial class GUBP : BuildCommand
         Result.Sort();
         return Result;
     }
-    void SaveStatus(string NodeToDo, string Suffix, string NodeStoreName, bool bSaveSharedTempStorage, string GameNameIfAny, string JobStepIDForFailure = null)
+    void SaveStatus(NodeInfo NodeToDo, string Suffix, string NodeStoreName, bool bSaveSharedTempStorage, string GameNameIfAny, string JobStepIDForFailure = null)
     {
         string Contents = "Just a status record: " + Suffix;
         if (!String.IsNullOrEmpty(JobStepIDForFailure) && IsBuildMachine)
@@ -753,7 +762,7 @@ public partial class GUBP : BuildCommand
                 Log(System.Diagnostics.TraceEventType.Warning, LogUtils.FormatException(Ex));
             }
         }
-        string RecordOfSuccess = CombinePaths(CommandUtils.CmdEnv.LocalRoot, "Engine", "Saved", "Logs", NodeToDo + Suffix +".log");
+        string RecordOfSuccess = CombinePaths(CommandUtils.CmdEnv.LocalRoot, "Engine", "Saved", "Logs", NodeToDo.Name + Suffix +".log");
         CreateDirectory(Path.GetDirectoryName(RecordOfSuccess));
         WriteAllText(RecordOfSuccess, Contents);		
 		TempStorage.StoreToTempStorage(CmdEnv, NodeStoreName + Suffix, new List<string> { RecordOfSuccess }, !bSaveSharedTempStorage, GameNameIfAny);		
@@ -785,24 +794,24 @@ public partial class GUBP : BuildCommand
         return Result;
     }
 
-    List<string> TopologicalSort(HashSet<string> NodesToDo, Dictionary<string, bool> GUBPNodesCompleted, Dictionary<string, string> GUBPNodesControllingTriggerDotName, Dictionary<string, string> GUBPNodesControllingTrigger, string ExplicitTrigger = "", bool LocalOnly = false, bool SubSort = false, bool DoNotConsiderCompletion = false)
+    List<NodeInfo> TopologicalSort(HashSet<NodeInfo> NodesToDo, Dictionary<string, bool> GUBPNodesCompleted, Dictionary<NodeInfo, string> GUBPNodesControllingTriggerDotName, Dictionary<NodeInfo, NodeInfo> GUBPNodesControllingTrigger, NodeInfo ExplicitTrigger = null, bool LocalOnly = false, bool SubSort = false, bool DoNotConsiderCompletion = false)
     {
         DateTime StartTime = DateTime.UtcNow;
 
-        List<string> OrdereredToDo = new List<string>();
+        List<NodeInfo> OrdereredToDo = new List<NodeInfo>();
 
-        Dictionary<string, List<string>> SortedAgentGroupChains = new Dictionary<string, List<string>>();
+        Dictionary<string, List<NodeInfo>> SortedAgentGroupChains = new Dictionary<string, List<NodeInfo>>();
         if (!SubSort)
         {
-            Dictionary<string, List<string>> AgentGroupChains = new Dictionary<string, List<string>>();
-            foreach (string NodeToDo in NodesToDo)
+            Dictionary<string, List<NodeInfo>> AgentGroupChains = new Dictionary<string, List<NodeInfo>>();
+            foreach (NodeInfo NodeToDo in NodesToDo)
             {
-                string MyAgentGroup = GUBPNodes[NodeToDo].Node.AgentSharingGroup;
+                string MyAgentGroup = NodeToDo.Node.AgentSharingGroup;
                 if (MyAgentGroup != "")
                 {
                     if (!AgentGroupChains.ContainsKey(MyAgentGroup))
                     {
-                        AgentGroupChains.Add(MyAgentGroup, new List<string> { NodeToDo });
+                        AgentGroupChains.Add(MyAgentGroup, new List<NodeInfo> { NodeToDo });
                     }
                     else
                     {
@@ -810,9 +819,9 @@ public partial class GUBP : BuildCommand
                     }
                 }
             }
-            foreach (KeyValuePair<string, List<string>> Chain in AgentGroupChains)
+            foreach (KeyValuePair<string, List<NodeInfo>> Chain in AgentGroupChains)
             {
-                SortedAgentGroupChains.Add(Chain.Key, TopologicalSort(new HashSet<string>(Chain.Value), GUBPNodesCompleted, GUBPNodesControllingTriggerDotName, GUBPNodesControllingTrigger, ExplicitTrigger, LocalOnly, true, DoNotConsiderCompletion));
+                SortedAgentGroupChains.Add(Chain.Key, TopologicalSort(new HashSet<NodeInfo>(Chain.Value), GUBPNodesCompleted, GUBPNodesControllingTriggerDotName, GUBPNodesControllingTrigger, ExplicitTrigger, LocalOnly, true, DoNotConsiderCompletion));
             }
             Log("***************Done with recursion");
         }
@@ -822,35 +831,35 @@ public partial class GUBP : BuildCommand
         {
             bool bProgressMade = false;
             float BestPriority = -1E20f;
-            string BestNode = "";
+            NodeInfo BestNode = null;
             bool BestPseudoReady = false;
             HashSet<string> NonReadyAgentGroups = new HashSet<string>();
             HashSet<string> NonPeudoReadyAgentGroups = new HashSet<string>();
             HashSet<string> ExaminedAgentGroups = new HashSet<string>();
-            foreach (string NodeToDo in NodesToDo)
+            foreach (NodeInfo NodeToDo in NodesToDo)
             {
                 bool bReady = true;
 				bool bPseudoReady = true;
 				bool bCompReady = true;
-                if (!SubSort && GUBPNodes[NodeToDo].Node.AgentSharingGroup != "")
+                if (!SubSort && NodeToDo.Node.AgentSharingGroup != "")
                 {
-                    if (ExaminedAgentGroups.Contains(GUBPNodes[NodeToDo].Node.AgentSharingGroup))
+                    if (ExaminedAgentGroups.Contains(NodeToDo.Node.AgentSharingGroup))
                     {
-                        bReady = !NonReadyAgentGroups.Contains(GUBPNodes[NodeToDo].Node.AgentSharingGroup);
-                        bPseudoReady = !NonPeudoReadyAgentGroups.Contains(GUBPNodes[NodeToDo].Node.AgentSharingGroup); //this might not be accurate if bReady==false
+                        bReady = !NonReadyAgentGroups.Contains(NodeToDo.Node.AgentSharingGroup);
+                        bPseudoReady = !NonPeudoReadyAgentGroups.Contains(NodeToDo.Node.AgentSharingGroup); //this might not be accurate if bReady==false
                     }
                     else
                     {
-                        ExaminedAgentGroups.Add(GUBPNodes[NodeToDo].Node.AgentSharingGroup);
-                        foreach (string ChainNode in SortedAgentGroupChains[GUBPNodes[NodeToDo].Node.AgentSharingGroup])
+                        ExaminedAgentGroups.Add(NodeToDo.Node.AgentSharingGroup);
+                        foreach (NodeInfo ChainNode in SortedAgentGroupChains[NodeToDo.Node.AgentSharingGroup])
                         {
-                            foreach (string Dep in GUBPNodes[ChainNode].Node.FullNamesOfDependencies)
+                            foreach (string Dep in ChainNode.Node.FullNamesOfDependencies)
                             {
                                 if (!GUBPNodes.ContainsKey(Dep))
                                 {
                                     throw new AutomationException("Dependency {0} node found.", Dep);
                                 }
-                                if (!SortedAgentGroupChains[GUBPNodes[NodeToDo].Node.AgentSharingGroup].Contains(Dep) && NodesToDo.Contains(Dep))
+                                if (!SortedAgentGroupChains[NodeToDo.Node.AgentSharingGroup].Contains(GUBPNodes[Dep]) && NodesToDo.Contains(GUBPNodes[Dep]))
                                 {
                                     bReady = false;
                                     break;
@@ -858,19 +867,19 @@ public partial class GUBP : BuildCommand
                             }
                             if (!bReady)
                             {
-                                NonReadyAgentGroups.Add(GUBPNodes[NodeToDo].Node.AgentSharingGroup);
+                                NonReadyAgentGroups.Add(NodeToDo.Node.AgentSharingGroup);
                                 break;
                             }
-                            foreach (string Dep in GUBPNodes[ChainNode].Node.FullNamesOfPseudosependencies)
+                            foreach (string Dep in ChainNode.Node.FullNamesOfPseudosependencies)
                             {
                                 if (!GUBPNodes.ContainsKey(Dep))
                                 {
                                     throw new AutomationException("Pseudodependency {0} node found.", Dep);
                                 }
-                                if (!SortedAgentGroupChains[GUBPNodes[NodeToDo].Node.AgentSharingGroup].Contains(Dep) && NodesToDo.Contains(Dep))
+                                if (!SortedAgentGroupChains[NodeToDo.Node.AgentSharingGroup].Contains(GUBPNodes[Dep]) && NodesToDo.Contains(GUBPNodes[Dep]))
                                 {
                                     bPseudoReady = false;
-                                    NonPeudoReadyAgentGroups.Add(GUBPNodes[NodeToDo].Node.AgentSharingGroup);
+                                    NonPeudoReadyAgentGroups.Add(NodeToDo.Node.AgentSharingGroup);
                                     break;
                                 }
                             }
@@ -879,46 +888,46 @@ public partial class GUBP : BuildCommand
                 }
                 else
                 {
-                    foreach (string Dep in GUBPNodes[NodeToDo].Node.FullNamesOfDependencies)
+                    foreach (string Dep in NodeToDo.Node.FullNamesOfDependencies)
                     {
                         if (!GUBPNodes.ContainsKey(Dep))
                         {
                             throw new AutomationException("Dependency {0} node found.", Dep);
                         }
-                        if (NodesToDo.Contains(Dep))
+                        if (NodesToDo.Contains(GUBPNodes[Dep]))
                         {
                             bReady = false;
                             break;
                         }
                     }
-                    foreach (string Dep in GUBPNodes[NodeToDo].Node.FullNamesOfPseudosependencies)
+                    foreach (string Dep in NodeToDo.Node.FullNamesOfPseudosependencies)
                     {
                         if (!GUBPNodes.ContainsKey(Dep))
                         {
                             throw new AutomationException("Pseudodependency {0} node found.", Dep);
                         }
-                        if (NodesToDo.Contains(Dep))
+                        if (NodesToDo.Contains(GUBPNodes[Dep]))
                         {
                             bPseudoReady = false;
                             break;
                         }
                     }
-					foreach (string Dep in GUBPNodes[NodeToDo].Node.CompletedDependencies)
+					foreach (string Dep in NodeToDo.Node.CompletedDependencies)
 					{
 						if (!GUBPNodes.ContainsKey(Dep))
 						{
 							throw new AutomationException("Completed Dependency {0} node found.", Dep);
 						}
-						if (NodesToDo.Contains(Dep))
+						if (NodesToDo.Contains(GUBPNodes[Dep]))
 						{
 							bCompReady = false;
 							break;
 						}
 					}
                 }
-                float Priority = GUBPNodes[NodeToDo].Node.Priority();
+                float Priority = NodeToDo.Node.Priority();
 
-                if (bReady && BestNode != "")
+                if (bReady && BestNode != null)
                 {
                     if (String.Compare(GetControllingTriggerDotName(BestNode, GUBPNodesControllingTriggerDotName, GUBPNodesControllingTrigger), GetControllingTriggerDotName(NodeToDo, GUBPNodesControllingTriggerDotName, GUBPNodesControllingTrigger)) < 0) //sorted by controlling trigger
                     {
@@ -926,11 +935,11 @@ public partial class GUBP : BuildCommand
                     }
                     else if (String.Compare(GetControllingTriggerDotName(BestNode, GUBPNodesControllingTriggerDotName, GUBPNodesControllingTrigger), GetControllingTriggerDotName(NodeToDo, GUBPNodesControllingTriggerDotName, GUBPNodesControllingTrigger)) == 0) //sorted by controlling trigger
                     {
-                        if (GUBPNodes[BestNode].Node.IsSticky() && !GUBPNodes[NodeToDo].Node.IsSticky()) //sticky nodes first
+                        if (GUBPNodes[BestNode.Name].Node.IsSticky() && !NodeToDo.Node.IsSticky()) //sticky nodes first
                         {
                             bReady = false;
                         }
-                        else if (GUBPNodes[BestNode].Node.IsSticky() == GUBPNodes[NodeToDo].Node.IsSticky())
+                        else if (GUBPNodes[BestNode.Name].Node.IsSticky() == NodeToDo.Node.IsSticky())
                         {
                             if (BestPseudoReady && !bPseudoReady)
                             {
@@ -938,8 +947,8 @@ public partial class GUBP : BuildCommand
                             }
                             else if (BestPseudoReady == bPseudoReady)
                             {
-                                bool IamLateTrigger = !DoNotConsiderCompletion && GUBPNodes[NodeToDo].Node.TriggerNode() && NodeToDo != ExplicitTrigger && !NodeIsAlreadyComplete(GUBPNodesCompleted, NodeToDo, LocalOnly);
-                                bool BestIsLateTrigger = !DoNotConsiderCompletion && GUBPNodes[BestNode].Node.TriggerNode() && BestNode != ExplicitTrigger && !NodeIsAlreadyComplete(GUBPNodesCompleted, BestNode, LocalOnly);
+                                bool IamLateTrigger = !DoNotConsiderCompletion && NodeToDo.Node.TriggerNode() && NodeToDo != ExplicitTrigger && !NodeIsAlreadyComplete(GUBPNodesCompleted, NodeToDo, LocalOnly);
+                                bool BestIsLateTrigger = !DoNotConsiderCompletion && GUBPNodes[BestNode.Name].Node.TriggerNode() && BestNode != ExplicitTrigger && !NodeIsAlreadyComplete(GUBPNodesCompleted, BestNode, LocalOnly);
                                 if (BestIsLateTrigger && !IamLateTrigger)
                                 {
                                     bReady = false;
@@ -953,7 +962,7 @@ public partial class GUBP : BuildCommand
                                     }
                                     else if (Priority == BestPriority)
                                     {
-                                        if (BestNode.CompareTo(NodeToDo) < 0)
+                                        if (BestNode.Name.CompareTo(NodeToDo.Name) < 0)
                                         {
                                             bReady = false;
                                         }
@@ -977,9 +986,9 @@ public partial class GUBP : BuildCommand
             }
             if (bProgressMade)
             {
-                if (!SubSort && GUBPNodes[BestNode].Node.AgentSharingGroup != "")
+                if (!SubSort && GUBPNodes[BestNode.Name].Node.AgentSharingGroup != "")
                 {
-                    foreach (string ChainNode in SortedAgentGroupChains[GUBPNodes[BestNode].Node.AgentSharingGroup])
+                    foreach (NodeInfo ChainNode in SortedAgentGroupChains[GUBPNodes[BestNode.Name].Node.AgentSharingGroup])
                     {
                         OrdereredToDo.Add(ChainNode);
                         NodesToDo.Remove(ChainNode);
@@ -995,32 +1004,32 @@ public partial class GUBP : BuildCommand
             if (!bProgressMade && NodesToDo.Count > 0)
             {
                 Log("Cycle in GUBP, could not resolve:");
-                foreach (string NodeToDo in NodesToDo)
+                foreach (NodeInfo NodeToDo in NodesToDo)
                 {
                     string Deps = "";
-                    if (!SubSort && GUBPNodes[NodeToDo].Node.AgentSharingGroup != "")
+                    if (!SubSort && NodeToDo.Node.AgentSharingGroup != "")
                     {
-                        foreach (string ChainNode in SortedAgentGroupChains[GUBPNodes[NodeToDo].Node.AgentSharingGroup])
+                        foreach (NodeInfo ChainNode in SortedAgentGroupChains[NodeToDo.Node.AgentSharingGroup])
                         {
-                            foreach (string Dep in GUBPNodes[ChainNode].Node.FullNamesOfDependencies)
+                            foreach (string Dep in GUBPNodes[ChainNode.Name].Node.FullNamesOfDependencies)
                             {
-                                if (!SortedAgentGroupChains[GUBPNodes[NodeToDo].Node.AgentSharingGroup].Contains(Dep) && NodesToDo.Contains(Dep))
+                                if (!SortedAgentGroupChains[NodeToDo.Node.AgentSharingGroup].Contains(GUBPNodes[Dep]) && NodesToDo.Contains(GUBPNodes[Dep]))
                                 {
-                                    Deps = Deps + Dep + "[" + ChainNode + "->" + GUBPNodes[NodeToDo].Node.AgentSharingGroup + "]" + " ";
+                                    Deps = Deps + Dep + "[" + ChainNode.Name + "->" + NodeToDo.Node.AgentSharingGroup + "]" + " ";
                                 }
                             }
                         }
                     }
-                    foreach (string Dep in GUBPNodes[NodeToDo].Node.FullNamesOfDependencies)
+                    foreach (string Dep in NodeToDo.Node.FullNamesOfDependencies)
                     {
-                        if (NodesToDo.Contains(Dep))
+                        if (NodesToDo.Contains(GUBPNodes[Dep]))
                         {
                             Deps = Deps + Dep + " ";
                         }
                     }
-                    foreach (string Dep in GUBPNodes[NodeToDo].Node.FullNamesOfPseudosependencies)
+                    foreach (string Dep in NodeToDo.Node.FullNamesOfPseudosependencies)
                     {
-                        if (NodesToDo.Contains(Dep))
+                        if (NodesToDo.Contains(GUBPNodes[Dep]))
                         {
                             Deps = Deps + Dep + " ";
                         }
@@ -1052,12 +1061,12 @@ public partial class GUBP : BuildCommand
         return ParentPath + "/" + GetJobStepPath(Dep);
     }
 
-    void UpdateNodeHistory(Dictionary<string, NodeHistory> GUBPNodesHistory, string Node, string CLString)
+    void UpdateNodeHistory(Dictionary<string, NodeHistory> GUBPNodesHistory, NodeInfo NodeToDo, string CLString)
     {
-        if (GUBPNodes[Node].Node.RunInEC() && !GUBPNodes[Node].Node.TriggerNode() && CLString != "")
+        if (NodeToDo.Node.RunInEC() && !NodeToDo.Node.TriggerNode() && CLString != "")
         {
-            string GameNameIfAny = GUBPNodes[Node].Node.GameNameIfAnyForTempStorage();
-            string NodeStoreWildCard = StoreName.Replace(CLString, "*") + "-" + GUBPNodes[Node].Node.GetFullName();
+            string GameNameIfAny = NodeToDo.Node.GameNameIfAnyForTempStorage();
+            string NodeStoreWildCard = StoreName.Replace(CLString, "*") + "-" + NodeToDo.Node.GetFullName();
             NodeHistory History = new NodeHistory();
 
             History.AllStarted = ConvertCLToIntList(TempStorage.FindTempStorageManifests(CmdEnv, NodeStoreWildCard + StartedTempStorageSuffix, false, true, GameNameIfAny));
@@ -1089,25 +1098,25 @@ public partial class GUBP : BuildCommand
                     }
                 }                
             }
-			if (GUBPNodesHistory.ContainsKey(Node))
+			if (GUBPNodesHistory.ContainsKey(NodeToDo.Name))
 			{
-				GUBPNodesHistory.Remove(Node);
+				GUBPNodesHistory.Remove(NodeToDo.Name);
 			}
-			GUBPNodesHistory.Add(Node, History);
+			GUBPNodesHistory.Add(NodeToDo.Name, History);
         }
     }
-	void GetFailureEmails(Dictionary<string, NodeHistory> GUBPNodesHistory, string NodeToDo, string CLString, bool OnlyLateUpdates = false)
+	void GetFailureEmails(Dictionary<string, NodeHistory> GUBPNodesHistory, NodeInfo NodeToDo, string CLString, bool OnlyLateUpdates = false)
 	{
         string EMails;
         string FailCauserEMails = "";
         string EMailNote = "";
         bool SendSuccessForGreenAfterRed = false;
         int NumPeople = 0;
-        if (GUBPNodesHistory.ContainsKey(NodeToDo))
+        if (GUBPNodesHistory.ContainsKey(NodeToDo.Name))
         {
-            NodeHistory History = GUBPNodesHistory[NodeToDo];
-			RunECTool(String.Format("setProperty \"/myWorkflow/LastGreen/{0}\" \"{1}\"", NodeToDo, History.LastSucceeded), true);
-			RunECTool(String.Format("setProperty \"/myWorkflow/LastGreen/{0}\" \"{1}\"", NodeToDo, History.FailedString), true);
+            NodeHistory History = GUBPNodesHistory[NodeToDo.Name];
+			RunECTool(String.Format("setProperty \"/myWorkflow/LastGreen/{0}\" \"{1}\"", NodeToDo.Name, History.LastSucceeded), true);
+			RunECTool(String.Format("setProperty \"/myWorkflow/LastGreen/{0}\" \"{1}\"", NodeToDo.Name, History.FailedString), true);
 
             if (History.LastSucceeded > 0 && History.LastSucceeded < P4Env.Changelist)
             {
@@ -1170,15 +1179,15 @@ public partial class GUBP : BuildCommand
         }
         else
         {
-			RunECTool(String.Format("setProperty \"/myWorkflow/LastGreen/{0}\" \"{1}\"", NodeToDo, "0"));
-			RunECTool(String.Format("setProperty \"/myWorkflow/RedsSince/{0}\" \"{1}\"", NodeToDo, ""));
+			RunECTool(String.Format("setProperty \"/myWorkflow/LastGreen/{0}\" \"{1}\"", NodeToDo.Name, "0"));
+			RunECTool(String.Format("setProperty \"/myWorkflow/RedsSince/{0}\" \"{1}\"", NodeToDo.Name, ""));
         }
-		RunECTool(String.Format("setProperty \"/myWorkflow/FailCausers/{0}\" \"{1}\"", NodeToDo, FailCauserEMails));
-		RunECTool(String.Format("setProperty \"/myWorkflow/EmailNotes/{0}\" \"{1}\"", NodeToDo, EMailNote));
+		RunECTool(String.Format("setProperty \"/myWorkflow/FailCausers/{0}\" \"{1}\"", NodeToDo.Name, FailCauserEMails));
+		RunECTool(String.Format("setProperty \"/myWorkflow/EmailNotes/{0}\" \"{1}\"", NodeToDo.Name, EMailNote));
         {
             string AdditionalEmails = "";
 			string Causers = "";
-            if (ParseParam("CIS") && !GUBPNodes[NodeToDo].Node.SendSuccessEmail() && !GUBPNodes[NodeToDo].Node.TriggerNode())
+            if (ParseParam("CIS") && !NodeToDo.Node.SendSuccessEmail() && !NodeToDo.Node.TriggerNode())
             {
 				Causers = FailCauserEMails;
            }
@@ -1188,15 +1197,15 @@ public partial class GUBP : BuildCommand
                 AdditionalEmails = GUBPNode.MergeSpaceStrings(AddEmails, AdditionalEmails);
             }
 			EMails = GetEMailListForNode(this, NodeToDo, AdditionalEmails, Causers);
-			RunECTool(String.Format("setProperty \"/myWorkflow/FailEmails/{0}\" \"{1}\"", NodeToDo, EMails));            
+			RunECTool(String.Format("setProperty \"/myWorkflow/FailEmails/{0}\" \"{1}\"", NodeToDo.Name, EMails));            
         }
-		if (GUBPNodes[NodeToDo].Node.SendSuccessEmail() || SendSuccessForGreenAfterRed)
+		if (NodeToDo.Node.SendSuccessEmail() || SendSuccessForGreenAfterRed)
 		{
-			RunECTool(String.Format("setProperty \"/myWorkflow/SendSuccessEmail/{0}\" \"{1}\"", NodeToDo, "1"));
+			RunECTool(String.Format("setProperty \"/myWorkflow/SendSuccessEmail/{0}\" \"{1}\"", NodeToDo.Name, "1"));
 		}
 		else
 		{
-			RunECTool(String.Format("setProperty \"/myWorkflow/SendSuccessEmail/{0}\" \"{1}\"", NodeToDo, "0"));
+			RunECTool(String.Format("setProperty \"/myWorkflow/SendSuccessEmail/{0}\" \"{1}\"", NodeToDo.Name, "0"));
 		}
 	}
 
@@ -1223,13 +1232,13 @@ public partial class GUBP : BuildCommand
         return true;
     }
 
-    int FindLastNonDuplicateFail(Dictionary<string, NodeHistory> GUBPNodesHistory, string NodeToDo, string CLString)
+    int FindLastNonDuplicateFail(Dictionary<string, NodeHistory> GUBPNodesHistory, NodeInfo NodeToDo, string CLString)
     {
-        NodeHistory History = GUBPNodesHistory[NodeToDo];
+        NodeHistory History = GUBPNodesHistory[NodeToDo.Name];
         int Result = P4Env.Changelist;
 
-        string GameNameIfAny = GUBPNodes[NodeToDo].Node.GameNameIfAnyForTempStorage();
-        string NodeStore = StoreName + "-" + GUBPNodes[NodeToDo].Node.GetFullName() + FailedTempStorageSuffix;
+        string GameNameIfAny = NodeToDo.Node.GameNameIfAnyForTempStorage();
+        string NodeStore = StoreName + "-" + NodeToDo.Node.GetFullName() + FailedTempStorageSuffix;
 
         List<int> BackwardsFails = new List<int>(History.AllFailed);
         BackwardsFails.Add(P4Env.Changelist);
@@ -1283,7 +1292,7 @@ public partial class GUBP : BuildCommand
         }
         return Result;
     }
-    List<string> GetECPropsForNode(string NodeToDo, string CLString, out string EMails, bool OnlyLateUpdates = false)
+    List<string> GetECPropsForNode(NodeInfo NodeToDo, string CLString, out string EMails, bool OnlyLateUpdates = false)
     {
         List<string> ECProps = new List<string>();
         EMails = "";		
@@ -1295,29 +1304,29 @@ public partial class GUBP : BuildCommand
 			AdditonalEmails = GUBPNode.MergeSpaceStrings(AddEmails, AdditonalEmails);
 		}
 		EMails = GetEMailListForNode(this, NodeToDo, AdditonalEmails, Causers);
-		ECProps.Add("FailEmails/" + NodeToDo + "=" + EMails);
+		ECProps.Add("FailEmails/" + NodeToDo.Name + "=" + EMails);
 	
 		if (!OnlyLateUpdates)
 		{
-			string AgentReq = GUBPNodes[NodeToDo].Node.ECAgentString();
-			if(ParseParamValue("AgentOverride") != "" && !GUBPNodes[NodeToDo].Node.GetFullName().Contains("OnMac"))
+			string AgentReq = NodeToDo.Node.ECAgentString();
+			if(ParseParamValue("AgentOverride") != "" && !NodeToDo.Node.GetFullName().Contains("OnMac"))
 			{
 				AgentReq = ParseParamValue("AgentOverride");
 			}
-			ECProps.Add(string.Format("AgentRequirementString/{0}={1}", NodeToDo, AgentReq));
-			ECProps.Add(string.Format("RequiredMemory/{0}={1}", NodeToDo, GUBPNodes[NodeToDo].Node.AgentMemoryRequirement(this)));
-			ECProps.Add(string.Format("Timeouts/{0}={1}", NodeToDo, GUBPNodes[NodeToDo].Node.TimeoutInMinutes()));
-			ECProps.Add(string.Format("JobStepPath/{0}={1}", NodeToDo, GetJobStepPath(NodeToDo)));
+			ECProps.Add(string.Format("AgentRequirementString/{0}={1}", NodeToDo.Name, AgentReq));
+			ECProps.Add(string.Format("RequiredMemory/{0}={1}", NodeToDo.Name, NodeToDo.Node.AgentMemoryRequirement(this)));
+			ECProps.Add(string.Format("Timeouts/{0}={1}", NodeToDo.Name, NodeToDo.Node.TimeoutInMinutes()));
+			ECProps.Add(string.Format("JobStepPath/{0}={1}", NodeToDo.Name, GetJobStepPath(NodeToDo.Name)));
 		}
 		
         return ECProps;
     }
 
-    void UpdateECProps(string NodeToDo, string CLString)
+    void UpdateECProps(NodeInfo NodeToDo, string CLString)
     {
         try
         {
-            Log("Updating node props for node {0}", NodeToDo);
+            Log("Updating node props for node {0}", NodeToDo.Name);
 			string EMails = "";
             List<string> Props = GetECPropsForNode(NodeToDo, CLString, out EMails, true);
             foreach (string Prop in Props)
@@ -1332,12 +1341,12 @@ public partial class GUBP : BuildCommand
             Log(System.Diagnostics.TraceEventType.Warning, LogUtils.FormatException(Ex));
         }
     }
-	void UpdateECBuildTime(string NodeToDo, double BuildDuration)
+	void UpdateECBuildTime(NodeInfo NodeToDo, double BuildDuration)
 	{
 		try
 		{
-			Log("Updating duration prop for node {0}", NodeToDo);
-			RunECTool(String.Format("setProperty \"/myWorkflow/NodeDuration/{0}\" \"{1}\"", NodeToDo, BuildDuration.ToString()));
+			Log("Updating duration prop for node {0}", NodeToDo.Name);
+			RunECTool(String.Format("setProperty \"/myWorkflow/NodeDuration/{0}\" \"{1}\"", NodeToDo.Name, BuildDuration.ToString()));
 			RunECTool(String.Format("setProperty \"/myJobStep/NodeDuration\" \"{0}\"", BuildDuration.ToString()));
 		}
 		catch (Exception Ex)
@@ -1517,17 +1526,17 @@ public partial class GUBP : BuildCommand
         }
 
         bool CommanderSetup = ParseParam("CommanderJobSetupOnly");
-        string ExplicitTrigger = "";
+        string ExplicitTriggerName = "";
         if (CommanderSetup)
         {
-            ExplicitTrigger = ParseParamValue("TriggerNode");
-            if (ExplicitTrigger == null)
+            ExplicitTriggerName = ParseParamValue("TriggerNode");
+            if (ExplicitTriggerName == null)
             {
-                ExplicitTrigger = "";
+                ExplicitTriggerName = "";
             }
         }
 
-        if (ParseParam("CIS") && ExplicitTrigger == "" && CommanderSetup) // explicit triggers will already have a time index assigned
+        if (ParseParam("CIS") && ExplicitTriggerName == "" && CommanderSetup) // explicit triggers will already have a time index assigned
         {
 			TimeIndex = UpdateCISCounter();
             Log("Setting TimeIndex to {0}", TimeIndex);
@@ -1562,8 +1571,8 @@ public partial class GUBP : BuildCommand
             TempStorage.DeleteLocalTempStorageManifests(CmdEnv);
         }
 
-		Dictionary<string, string> GUBPNodesControllingTrigger = new Dictionary<string, string>();
-        Dictionary<string, string> GUBPNodesControllingTriggerDotName = new Dictionary<string, string>();
+		Dictionary<NodeInfo, NodeInfo> GUBPNodesControllingTrigger = new Dictionary<NodeInfo, NodeInfo>();
+        Dictionary<NodeInfo, string> GUBPNodesControllingTriggerDotName = new Dictionary<NodeInfo, string>();
 
 		Dictionary<string, string> FullNodeDependedOnBy = new Dictionary<string, string>();
 		Dictionary<string, string> FullNodeDependentPromotions = new Dictionary<string, string>();		
@@ -1574,19 +1583,19 @@ public partial class GUBP : BuildCommand
                 if (GUBPNodes[NodeToDo.Key].Node.IsSeparatePromotable())
                 {
                     SeparatePromotables.Add(GUBPNodes[NodeToDo.Key].Node.GetFullName());
-                    List<string> Dependencies = new List<string>();
-                    Dependencies = GetECDependencies(NodeToDo.Key);
-                    foreach (string Dep in Dependencies)
+                    List<NodeInfo> Dependencies = new List<NodeInfo>();
+                    Dependencies = GetECDependencies(NodeToDo.Value);
+                    foreach (NodeInfo Dep in Dependencies)
                     {
-                        if (!GUBPNodes.ContainsKey(Dep))
+                        if (!GUBPNodes.ContainsKey(Dep.Name))
                         {
-                            throw new AutomationException("Node {0} is not in the graph.  It is a dependency of {1}.", Dep, NodeToDo);
+                            throw new AutomationException("Node {0} is not in the graph.  It is a dependency of {1}.", Dep.Name, NodeToDo.Key);
                         }
-                        if (!GUBPNodes[Dep].Node.IsPromotableAggregate())
+                        if (!GUBPNodes[Dep.Name].Node.IsPromotableAggregate())
                         {
-                            if (!GUBPNodes[Dep].Node.DependentPromotions.Contains(NodeToDo.Key))
+                            if (!GUBPNodes[Dep.Name].Node.DependentPromotions.Contains(NodeToDo.Key))
                             {
-                                GUBPNodes[Dep].Node.DependentPromotions.Add(NodeToDo.Key);
+                                GUBPNodes[Dep.Name].Node.DependentPromotions.Add(NodeToDo.Key);
                             }
                         }
                     }
@@ -1597,7 +1606,7 @@ public partial class GUBP : BuildCommand
 			Dictionary<string, int> FrequencyOverrides = ApplyFrequencyBarriers();
             foreach (KeyValuePair<string, NodeInfo> NodeToDo in GUBPNodes)
             {
-                ComputeDependentCISFrequencyQuantumShift(NodeToDo.Key, FrequencyOverrides);
+                ComputeDependentCISFrequencyQuantumShift(NodeToDo.Value, FrequencyOverrides);
             }
 
             foreach (KeyValuePair<string, NodeInfo> NodeToDo in GUBPNodes)
@@ -1623,29 +1632,29 @@ public partial class GUBP : BuildCommand
 		Dictionary<string, int> FullNodeListSortKey = GetDisplayOrder(FullNodeList.Keys.ToList(), FullNodeDirectDependencies, GUBPNodes);
 
 		bool bOnlyNode;
-		HashSet<string> NodesToDo;
+		HashSet<NodeInfo> NodesToDo;
 		ParseNodesToDo(WithoutLinux, TimeIndex, CommanderSetup, FullNodeDependedOnBy, out bOnlyNode, out NodesToDo);
 
-        if (CommanderSetup)
+		NodeInfo ExplicitTrigger = null;
+		if (CommanderSetup)
         {
-            if (!String.IsNullOrEmpty(ExplicitTrigger))
+            if (!String.IsNullOrEmpty(ExplicitTriggerName))
             {
-                bool bFoundIt = false;
                 foreach (KeyValuePair<string, NodeInfo> Node in GUBPNodes)
                 {
-                    if (Node.Value.Node.GetFullName().Equals(ExplicitTrigger, StringComparison.InvariantCultureIgnoreCase))
+                    if (Node.Value.Node.GetFullName().Equals(ExplicitTriggerName, StringComparison.InvariantCultureIgnoreCase))
                     {
                         if (Node.Value.Node.TriggerNode() && Node.Value.Node.RunInEC())
                         {
                             Node.Value.Node.SetAsExplicitTrigger();
-                            bFoundIt = true;
+							ExplicitTrigger = Node.Value;
                             break;
                         }
                     }
                 }
-                if (!bFoundIt)
+                if (ExplicitTrigger == null)
                 {
-                    throw new AutomationException("Could not find trigger node named {0}", ExplicitTrigger);
+                    throw new AutomationException("Could not find trigger node named {0}", ExplicitTriggerName);
                 }
             }
             else
@@ -1665,18 +1674,18 @@ public partial class GUBP : BuildCommand
 		if (bPreflightBuild)
 		{
 			LogVerbose("Culling triggers and downstream for preflight builds ");
-			HashSet<string> NewNodesToDo = new HashSet<string>();
-			foreach (string NodeToDo in NodesToDo)
+			HashSet<NodeInfo> NewNodesToDo = new HashSet<NodeInfo>();
+			foreach (NodeInfo NodeToDo in NodesToDo)
 			{
 				string TriggerDot = GetControllingTriggerDotName(NodeToDo, GUBPNodesControllingTriggerDotName, GUBPNodesControllingTrigger);
-				if (TriggerDot == "" && !GUBPNodes[NodeToDo].Node.TriggerNode())
+				if (TriggerDot == "" && !NodeToDo.Node.TriggerNode())
 				{
-					LogVerbose("  Keeping {0}", NodeToDo);
+					LogVerbose("  Keeping {0}", NodeToDo.Name);
 					NewNodesToDo.Add(NodeToDo);
 				}
 				else
 				{
-					LogVerbose("  Rejecting {0}", NodeToDo);
+					LogVerbose("  Rejecting {0}", NodeToDo.Name);
 				}
 			}
 			NodesToDo = NewNodesToDo;
@@ -1688,11 +1697,11 @@ public partial class GUBP : BuildCommand
         LogVerbose("******* Caching completion");
         {
             DateTime StartTime = DateTime.UtcNow;
-            foreach (string Node in NodesToDo)
+            foreach (NodeInfo NodeToDo in NodesToDo)
             {
-                LogVerbose("** {0}", Node);
-                NodeIsAlreadyComplete(GUBPNodesCompleted, Node, LocalOnly); // cache these now to avoid spam later
-                GetControllingTriggerDotName(Node, GUBPNodesControllingTriggerDotName, GUBPNodesControllingTrigger);
+                LogVerbose("** {0}", NodeToDo.Name);
+                NodeIsAlreadyComplete(GUBPNodesCompleted, NodeToDo, LocalOnly); // cache these now to avoid spam later
+                GetControllingTriggerDotName(NodeToDo, GUBPNodesControllingTriggerDotName, GUBPNodesControllingTrigger);
             }
             double BuildDuration = (DateTime.UtcNow - StartTime).TotalMilliseconds;
 			LogVerbose("Took {0}s to cache completion for {1} nodes", BuildDuration / 1000, NodesToDo.Count);
@@ -1712,9 +1721,9 @@ public partial class GUBP : BuildCommand
             Log("Took {0}s to get history for {1} nodes", BuildDuration / 1000, NodesToDo.Count);
         }*/
 
-        List<string> OrdereredToDo = TopologicalSort(NodesToDo, GUBPNodesCompleted, GUBPNodesControllingTriggerDotName, GUBPNodesControllingTrigger, ExplicitTrigger, LocalOnly);
+        List<NodeInfo> OrdereredToDo = TopologicalSort(NodesToDo, GUBPNodesCompleted, GUBPNodesControllingTriggerDotName, GUBPNodesControllingTrigger, ExplicitTrigger, LocalOnly);
 
-		List<string> UnfinishedTriggers = FindUnfinishedTriggers(bSkipTriggers, LocalOnly, ExplicitTrigger, GUBPNodesCompleted, OrdereredToDo);
+		List<NodeInfo> UnfinishedTriggers = FindUnfinishedTriggers(bSkipTriggers, LocalOnly, ExplicitTrigger, GUBPNodesCompleted, OrdereredToDo);
 
         LogVerbose("*********** Desired And Dependent Nodes, in order.");
         PrintNodes(this, OrdereredToDo, GUBPNodesCompleted, GUBPNodesHistory, GUBPNodesControllingTriggerDotName, GUBPNodesControllingTrigger, LocalOnly, UnfinishedTriggers);		
@@ -1787,11 +1796,11 @@ public partial class GUBP : BuildCommand
 		return FrequencyOverrides;
 	}
 
-	private void ParseNodesToDo(bool WithoutLinux, int TimeIndex, bool CommanderSetup, Dictionary<string, string> FullNodeDependedOnBy, out bool bOnlyNode, out HashSet<string> NodesToDo)
+	private void ParseNodesToDo(bool WithoutLinux, int TimeIndex, bool CommanderSetup, Dictionary<string, string> FullNodeDependedOnBy, out bool bOnlyNode, out HashSet<NodeInfo> NodesToDo)
 	{
 		bOnlyNode = false;
 		bool bRelatedToNode = false;
-		NodesToDo = new HashSet<string>();
+		NodesToDo = new HashSet<NodeInfo>();
 
 		{
 			string NodeSpec = ParseParamValue("Node");
@@ -1838,9 +1847,9 @@ public partial class GUBP : BuildCommand
 								Node.Value.Node.AgentSharingGroup.Equals(NodeArg, StringComparison.InvariantCultureIgnoreCase)
 								)
 							{
-								if (!NodesToDo.Contains(Node.Key))
+								if (!NodesToDo.Contains(Node.Value))
 								{
-									NodesToDo.Add(Node.Key);
+									NodesToDo.Add(Node.Value);
 								}
 								bFoundAnything = true;
 							}
@@ -1868,13 +1877,13 @@ public partial class GUBP : BuildCommand
 						{
 							if (GameProj.Options(UnrealTargetPlatform.Win64).bIsPromotable)
 							{
-								NodesToDo.Add(GameAggregatePromotableNode.StaticGetFullName(GameProj));
+								NodesToDo.Add(GUBPNodes[GameAggregatePromotableNode.StaticGetFullName(GameProj)]);
 							}
 							foreach (KeyValuePair<string, NodeInfo> Node in GUBPNodes)
 							{
 								if (Node.Value.Node.GameNameIfAnyForTempStorage() == GameProj.GameName)
 								{
-									NodesToDo.Add(Node.Key);
+									NodesToDo.Add(Node.Value);
 								}
 							}
 
@@ -1891,7 +1900,7 @@ public partial class GUBP : BuildCommand
 								{
 									if (Node.Value.Node.GameNameIfAnyForTempStorage() == GameProj.GameName)
 									{
-										NodesToDo.Add(Node.Key);
+										NodesToDo.Add(Node.Value);
 									}
 								}
 								GameName = null;
@@ -1910,25 +1919,25 @@ public partial class GUBP : BuildCommand
 			LogVerbose("No nodes specified, adding all nodes");
 			foreach (KeyValuePair<string, NodeInfo> Node in GUBPNodes)
 			{
-				NodesToDo.Add(Node.Key);
+				NodesToDo.Add(Node.Value);
 			}
 		}
 		else if (TimeIndex != 0)
 		{
 			LogVerbose("Check to make sure we didn't ask for nodes that will be culled by time index");
-			foreach (string NodeToDo in NodesToDo)
+			foreach (NodeInfo NodeToDo in NodesToDo)
 			{
-				if (TimeIndex % (1 << GUBPNodes[NodeToDo].Node.DependentCISFrequencyQuantumShift()) != 0)
+				if (TimeIndex % (1 << NodeToDo.Node.DependentCISFrequencyQuantumShift()) != 0)
 				{
-					throw new AutomationException("You asked specifically for node {0}, but it is culled by the time quantum: TimeIndex = {1}, DependentCISFrequencyQuantumShift = {2}.", NodeToDo, TimeIndex, GUBPNodes[NodeToDo].Node.DependentCISFrequencyQuantumShift());
+					throw new AutomationException("You asked specifically for node {0}, but it is culled by the time quantum: TimeIndex = {1}, DependentCISFrequencyQuantumShift = {2}.", NodeToDo.Name, TimeIndex, NodeToDo.Node.DependentCISFrequencyQuantumShift());
 				}
 			}
 		}
 
 		LogVerbose("Desired Nodes");
-		foreach (string NodeToDo in NodesToDo)
+		foreach (NodeInfo NodeToDo in NodesToDo)
 		{
-			LogVerbose("  {0}", NodeToDo);
+			LogVerbose("  {0}", NodeToDo.Name);
 		}
 		// if we are doing related to, then find things that depend on the selected nodes
 		if (bRelatedToNode)
@@ -1938,10 +1947,10 @@ public partial class GUBP : BuildCommand
 			while (!bDoneWithDependencies)
 			{
 				bDoneWithDependencies = true;
-				HashSet<string> Fringe = new HashSet<string>();
+				HashSet<NodeInfo> Fringe = new HashSet<NodeInfo>();
 				foreach (KeyValuePair<string, NodeInfo> NodeToDo in GUBPNodes)
 				{
-					if (!NodesToDo.Contains(NodeToDo.Key))
+					if (!NodesToDo.Contains(NodeToDo.Value))
 					{
 						foreach (string Dep in GUBPNodes[NodeToDo.Key].Node.FullNamesOfDependencies)
 						{
@@ -1949,9 +1958,9 @@ public partial class GUBP : BuildCommand
 							{
 								throw new AutomationException("Node {0} is not in the graph. It is a dependency of {1}.", Dep, NodeToDo.Key);
 							}
-							if (NodesToDo.Contains(Dep))
+							if (NodesToDo.Contains(GUBPNodes[Dep]))
 							{
-								Fringe.Add(NodeToDo.Key);
+								Fringe.Add(NodeToDo.Value);
 								bDoneWithDependencies = false;
 							}
 						}
@@ -1976,40 +1985,40 @@ public partial class GUBP : BuildCommand
 			while (!bDoneWithDependencies)
 			{
 				bDoneWithDependencies = true;
-				HashSet<string> Fringe = new HashSet<string>();
-				foreach (string NodeToDo in NodesToDo)
+				HashSet<NodeInfo> Fringe = new HashSet<NodeInfo>();
+				foreach (NodeInfo NodeToDo in NodesToDo)
 				{
 					if (BranchOptions.NodesToRemovePseudoDependencies.Count > 0)
 					{
 						List<string> ListOfRetainedNodesToRemovePseudoDependencies = BranchOptions.NodesToRemovePseudoDependencies;
 						foreach (string NodeToRemovePseudoDependencies in BranchOptions.NodesToRemovePseudoDependencies)
 						{
-							if (NodeToDo == NodeToRemovePseudoDependencies)
+							if (NodeToDo.Name == NodeToRemovePseudoDependencies)
 							{
-								RemoveAllPseudodependenciesFromNode(NodeToDo);
-								ListOfRetainedNodesToRemovePseudoDependencies.Remove(NodeToDo);
+								RemoveAllPseudodependenciesFromNode(NodeToDo.Name);
+								ListOfRetainedNodesToRemovePseudoDependencies.Remove(NodeToDo.Name);
 								break;
 							}
 						}
 						BranchOptions.NodesToRemovePseudoDependencies = ListOfRetainedNodesToRemovePseudoDependencies;
 					}
-					foreach (string Dep in GUBPNodes[NodeToDo].Node.FullNamesOfDependencies)
+					foreach (string Dep in NodeToDo.Node.FullNamesOfDependencies)
 					{
 						if (!GUBPNodes.ContainsKey(Dep))
 						{
-							throw new AutomationException("Node {0} is not in the graph. It is a dependency of {1}.", Dep, NodeToDo);
+							throw new AutomationException("Node {0} is not in the graph. It is a dependency of {1}.", Dep, NodeToDo.Name);
 						}
-						if (!NodesToDo.Contains(Dep))
+						if (!NodesToDo.Contains(GUBPNodes[Dep]))
 						{
-							Fringe.Add(Dep);
+							Fringe.Add(GUBPNodes[Dep]);
 							bDoneWithDependencies = false;
 						}
 					}
-					foreach (string Dep in GUBPNodes[NodeToDo].Node.FullNamesOfPseudosependencies)
+					foreach (string Dep in NodeToDo.Node.FullNamesOfPseudosependencies)
 					{
 						if (!GUBPNodes.ContainsKey(Dep))
 						{
-							throw new AutomationException("Node {0} is not in the graph. It is a pseudodependency of {1}.", Dep, NodeToDo);
+							throw new AutomationException("Node {0} is not in the graph. It is a pseudodependency of {1}.", Dep, NodeToDo.Name);
 						}
 					}
 				}
@@ -2019,17 +2028,17 @@ public partial class GUBP : BuildCommand
 		if (TimeIndex != 0)
 		{
 			LogVerbose("Culling based on time index");
-			HashSet<string> NewNodesToDo = new HashSet<string>();
-			foreach (string NodeToDo in NodesToDo)
+			HashSet<NodeInfo> NewNodesToDo = new HashSet<NodeInfo>();
+			foreach (NodeInfo NodeToDo in NodesToDo)
 			{
-				if (TimeIndex % (1 << GUBPNodes[NodeToDo].Node.DependentCISFrequencyQuantumShift()) == 0)
+				if (TimeIndex % (1 << NodeToDo.Node.DependentCISFrequencyQuantumShift()) == 0)
 				{
-					LogVerbose("  Keeping {0}", NodeToDo);
+					LogVerbose("  Keeping {0}", NodeToDo.Name);
 					NewNodesToDo.Add(NodeToDo);
 				}
 				else
 				{
-					LogVerbose("  Rejecting {0}", NodeToDo);
+					LogVerbose("  Rejecting {0}", NodeToDo.Name);
 				}
 			}
 			NodesToDo = NewNodesToDo;
@@ -2037,10 +2046,10 @@ public partial class GUBP : BuildCommand
 		//Remove Plat if specified
 		if (WithoutLinux)
 		{
-			HashSet<string> NewNodesToDo = new HashSet<string>();
-			foreach (string NodeToDo in NodesToDo)
+			HashSet<NodeInfo> NewNodesToDo = new HashSet<NodeInfo>();
+			foreach (NodeInfo NodeToDo in NodesToDo)
 			{
-				if (!GUBPNodes[NodeToDo].Node.GetFullName().Contains("Linux"))
+				if (!NodeToDo.Node.GetFullName().Contains("Linux"))
 				{
 					NewNodesToDo.Add(NodeToDo);
 				}
@@ -2054,28 +2063,28 @@ public partial class GUBP : BuildCommand
 		//find things that depend on our nodes and setup commander dictionary
 		if (!bOnlyNode)
 		{
-			foreach (string NodeToDo in NodesToDo)
+			foreach (NodeInfo NodeToDo in NodesToDo)
 			{
-				if (!GUBPNodes[NodeToDo].Node.IsAggregate() && !GUBPNodes[NodeToDo].Node.IsTest())
+				if (!NodeToDo.Node.IsAggregate() && !NodeToDo.Node.IsTest())
 				{
-					List<string> ECDependencies = new List<string>();
+					List<NodeInfo> ECDependencies = new List<NodeInfo>();
 					ECDependencies = GetECDependencies(NodeToDo);
-					foreach (string Dep in ECDependencies)
+					foreach (NodeInfo Dep in ECDependencies)
 					{
-						if (!GUBPNodes.ContainsKey(Dep))
+						if (!GUBPNodes.ContainsKey(Dep.Name))
 						{
-							throw new AutomationException("Node {0} is not in the graph. It is a dependency of {1}.", Dep, NodeToDo);
+							throw new AutomationException("Node {0} is not in the graph. It is a dependency of {1}.", Dep.Name, NodeToDo.Name);
 						}
-						if (!GUBPNodes[Dep].Node.FullNamesOfDependedOn.Contains(NodeToDo))
+						if (!GUBPNodes[Dep.Name].Node.FullNamesOfDependedOn.Contains(NodeToDo.Name))
 						{
-							GUBPNodes[Dep].Node.FullNamesOfDependedOn.Add(NodeToDo);
+							GUBPNodes[Dep.Name].Node.FullNamesOfDependedOn.Add(NodeToDo.Name);
 						}
 					}
 				}
 			}
-			foreach (string NodeToDo in NodesToDo)
+			foreach (NodeInfo NodeToDo in NodesToDo)
 			{
-				List<string> Deps = GUBPNodes[NodeToDo].Node.FullNamesOfDependedOn;
+				List<string> Deps = NodeToDo.Node.FullNamesOfDependedOn;
 				string All = "";
 				foreach (string Dep in Deps)
 				{
@@ -2085,20 +2094,20 @@ public partial class GUBP : BuildCommand
 					}
 					All += Dep;
 				}
-				FullNodeDependedOnBy.Add(NodeToDo, All);
+				FullNodeDependedOnBy.Add(NodeToDo.Name, All);
 			}
 		}
 	}
 
-	private void GetFullNodeListAndDirectDependencies(Dictionary<string, string> GUBPNodesControllingTrigger, Dictionary<string, string> GUBPNodesControllingTriggerDotName, out Dictionary<string, string> FullNodeList, out Dictionary<string, string> FullNodeDirectDependencies)
+	private void GetFullNodeListAndDirectDependencies(Dictionary<NodeInfo, NodeInfo> GUBPNodesControllingTrigger, Dictionary<NodeInfo, string> GUBPNodesControllingTriggerDotName, out Dictionary<string, string> FullNodeList, out Dictionary<string, string> FullNodeDirectDependencies)
 	{
 		FullNodeList = new Dictionary<string,string>();
 		FullNodeDirectDependencies = new Dictionary<string,string>();
 
 		Log("******* {0} GUBP Nodes", GUBPNodes.Count);
-		List<string> SortedNodes = TopologicalSort(new HashSet<string>(GUBPNodes.Keys), null/*GUBPNodesCompleted*/, GUBPNodesControllingTriggerDotName, GUBPNodesControllingTrigger, LocalOnly: true, DoNotConsiderCompletion: true);
+		List<NodeInfo> SortedNodes = TopologicalSort(new HashSet<NodeInfo>(GUBPNodes.Values), null/*GUBPNodesCompleted*/, GUBPNodesControllingTriggerDotName, GUBPNodesControllingTrigger, LocalOnly: true, DoNotConsiderCompletion: true);
 		string DependencyStart = DateTime.Now.ToString();
-		foreach (string Node in SortedNodes)
+		foreach (NodeInfo Node in SortedNodes)
 		{
 			string Note = GetControllingTriggerDotName(Node, GUBPNodesControllingTriggerDotName, GUBPNodesControllingTrigger);
 			if (Note == "")
@@ -2109,25 +2118,25 @@ public partial class GUBP : BuildCommand
 			{
 				Note = "always";
 			}
-			if (GUBPNodes[Node].Node.RunInEC())
+			if (GUBPNodes[Node.Name].Node.RunInEC())
 			{
-				List<string> Deps = GetECDependencies(Node);
+				List<NodeInfo> Deps = GetECDependencies(Node);
 				string All = "";
-				foreach (string Dep in Deps)
+				foreach (NodeInfo Dep in Deps)
 				{
 					if (All != "")
 					{
 						All += " ";
 					}
-					All += Dep;
+					All += Dep.Name;
 				}
-				LogVerbose("  {0}: {1}      {2}", Node, Note, All);
-				FullNodeList.Add(Node, Note);
-				FullNodeDirectDependencies.Add(Node, All);
+				LogVerbose("  {0}: {1}      {2}", Node.Name, Note, All);
+				FullNodeList.Add(Node.Name, Note);
+				FullNodeDirectDependencies.Add(Node.Name, All);
 			}
 			else
 			{
-				LogVerbose("  {0}: {1} [Aggregate]", Node, Note);
+				LogVerbose("  {0}: {1} [Aggregate]", Node.Name, Note);
 			}
 		}
 		string DependencyFinish = DateTime.Now.ToString();
@@ -2281,17 +2290,17 @@ public partial class GUBP : BuildCommand
 		throw new AutomationException("Failed to update the CIS counter after 20 tries.");
 	}
 
-	private List<string> FindUnfinishedTriggers(bool bSkipTriggers, bool LocalOnly, string ExplicitTrigger, Dictionary<string, bool> GUBPNodesCompleted, List<string> OrdereredToDo)
+	private List<NodeInfo> FindUnfinishedTriggers(bool bSkipTriggers, bool LocalOnly, NodeInfo ExplicitTrigger, Dictionary<string, bool> GUBPNodesCompleted, List<NodeInfo> OrdereredToDo)
 	{
 		// find all unfinished triggers, excepting the one we are triggering right now
-		List<string> UnfinishedTriggers = new List<string>();
+		List<NodeInfo> UnfinishedTriggers = new List<NodeInfo>();
 		if (!bSkipTriggers)
 		{
-			foreach (string NodeToDo in OrdereredToDo)
+			foreach (NodeInfo NodeToDo in OrdereredToDo)
 			{
-				if (GUBPNodes[NodeToDo].Node.TriggerNode() && !NodeIsAlreadyComplete(GUBPNodesCompleted, NodeToDo, LocalOnly))
+				if (NodeToDo.Node.TriggerNode() && !NodeIsAlreadyComplete(GUBPNodesCompleted, NodeToDo, LocalOnly))
 				{
-					if (String.IsNullOrEmpty(ExplicitTrigger) || ExplicitTrigger != NodeToDo)
+					if (ExplicitTrigger != NodeToDo)
 					{
 						UnfinishedTriggers.Add(NodeToDo);
 					}
@@ -2301,39 +2310,39 @@ public partial class GUBP : BuildCommand
 		return UnfinishedTriggers;
 	}
 
-	private void CheckSortOrder(bool LocalOnly, Dictionary<string, bool> GUBPNodesCompleted, List<string> OrdereredToDo)
+	private void CheckSortOrder(bool LocalOnly, Dictionary<string, bool> GUBPNodesCompleted, List<NodeInfo> OrdereredToDo)
 	{
-		foreach (string NodeToDo in OrdereredToDo)
+		foreach (NodeInfo NodeToDo in OrdereredToDo)
 		{
-			if (GUBPNodes[NodeToDo].Node.TriggerNode() && (GUBPNodes[NodeToDo].Node.IsSticky() || NodeIsAlreadyComplete(GUBPNodesCompleted, NodeToDo, LocalOnly))) // these sticky triggers are ok, everything is already completed anyway
+			if (NodeToDo.Node.TriggerNode() && (NodeToDo.Node.IsSticky() || NodeIsAlreadyComplete(GUBPNodesCompleted, NodeToDo, LocalOnly))) // these sticky triggers are ok, everything is already completed anyway
 			{
 				continue;
 			}
-			if (GUBPNodes[NodeToDo].Node.IsTest())
+			if (NodeToDo.Node.IsTest())
 			{
 				bHasTests = true;
 			}
 			int MyIndex = OrdereredToDo.IndexOf(NodeToDo);
-			foreach (string Dep in GUBPNodes[NodeToDo].Node.FullNamesOfDependencies)
+			foreach (NodeInfo Dep in NodeToDo.Node.FullNamesOfDependencies.Select(x => GUBPNodes[x]))
 			{
 				int DepIndex = OrdereredToDo.IndexOf(Dep);
 				if (DepIndex >= MyIndex)
 				{
-					throw new AutomationException("Topological sort error, node {0} has a dependency of {1} which sorted after it.", NodeToDo, Dep);
+					throw new AutomationException("Topological sort error, node {0} has a dependency of {1} which sorted after it.", NodeToDo.Name, Dep.Name);
 				}
 			}
-			foreach (string Dep in GUBPNodes[NodeToDo].Node.FullNamesOfPseudosependencies)
+			foreach (NodeInfo Dep in NodeToDo.Node.FullNamesOfPseudosependencies.Select(x => GUBPNodes[x]))
 			{
 				int DepIndex = OrdereredToDo.IndexOf(Dep);
 				if (DepIndex >= MyIndex)
 				{
-					throw new AutomationException("Topological sort error, node {0} has a pseduodependency of {1} which sorted after it.", NodeToDo, Dep);
+					throw new AutomationException("Topological sort error, node {0} has a pseduodependency of {1} which sorted after it.", NodeToDo.Name, Dep.Name);
 				}
 			}
 		}
 	}
 
-	private void DoCommanderSetup(int TimeIndex, bool bSkipTriggers, bool bFakeEC, bool LocalOnly, string CLString, string ExplicitTrigger, Dictionary<string, string> GUBPNodesControllingTrigger, Dictionary<string, string> GUBPNodesControllingTriggerDotName, Dictionary<string, string> FullNodeList, Dictionary<string, string> FullNodeDirectDependencies, Dictionary<string, string> FullNodeDependedOnBy, Dictionary<string, string> FullNodeDependentPromotions, List<string> SeparatePromotables, Dictionary<string, int> FullNodeListSortKey, Dictionary<string, bool> GUBPNodesCompleted, Dictionary<string, NodeHistory> GUBPNodesHistory, List<string> OrdereredToDo, List<string> UnfinishedTriggers, string FakeFail)
+	private void DoCommanderSetup(int TimeIndex, bool bSkipTriggers, bool bFakeEC, bool LocalOnly, string CLString, NodeInfo ExplicitTrigger, Dictionary<NodeInfo, NodeInfo> GUBPNodesControllingTrigger, Dictionary<NodeInfo, string> GUBPNodesControllingTriggerDotName, Dictionary<string, string> FullNodeList, Dictionary<string, string> FullNodeDirectDependencies, Dictionary<string, string> FullNodeDependedOnBy, Dictionary<string, string> FullNodeDependentPromotions, List<string> SeparatePromotables, Dictionary<string, int> FullNodeListSortKey, Dictionary<string, bool> GUBPNodesCompleted, Dictionary<string, NodeHistory> GUBPNodesHistory, List<NodeInfo> OrdereredToDo, List<NodeInfo> UnfinishedTriggers, string FakeFail)
 	{
 
 		if (OrdereredToDo.Count == 0)
@@ -2367,7 +2376,7 @@ public partial class GUBP : BuildCommand
 			ECProps.Add(string.Format("PossiblePromotables/{0}={1}", Node, ""));
 		}
 		List<string> ECJobProps = new List<string>();
-		if (ExplicitTrigger != "")
+		if (ExplicitTrigger != null)
 		{
 			ECJobProps.Add("IsRoot=0");
 		}
@@ -2376,25 +2385,25 @@ public partial class GUBP : BuildCommand
 			ECJobProps.Add("IsRoot=1");
 		}
 
-		List<string> FilteredOrdereredToDo = new List<string>();
+		List<NodeInfo> FilteredOrdereredToDo = new List<NodeInfo>();
 		string StartFilterTime = DateTime.Now.ToString();
 		// remove nodes that have unfinished triggers
-		foreach (string NodeToDo in OrdereredToDo)
+		foreach (NodeInfo NodeToDo in OrdereredToDo)
 		{
-			string ControllingTrigger = GetControllingTrigger(NodeToDo, GUBPNodesControllingTrigger);
+			NodeInfo ControllingTrigger = GetControllingTrigger(NodeToDo, GUBPNodesControllingTrigger);
 			bool bNoUnfinishedTriggers = !UnfinishedTriggers.Contains(ControllingTrigger);
 
 			if (bNoUnfinishedTriggers)
 			{
 				// if we are triggering, then remove nodes that are not controlled by the trigger or are dependencies of this trigger
-				if (!String.IsNullOrEmpty(ExplicitTrigger))
+				if (ExplicitTrigger != null)
 				{
 					if (ExplicitTrigger != NodeToDo && !NodeDependsOn(NodeToDo, ExplicitTrigger) && !NodeDependsOn(ExplicitTrigger, NodeToDo))
 					{
 						continue; // this wasn't on the chain related to the trigger we are triggering, so it is not relevant
 					}
 				}
-				if (bPreflightBuild && !bSkipTriggers && GUBPNodes[NodeToDo].Node.TriggerNode())
+				if (bPreflightBuild && !bSkipTriggers && NodeToDo.Node.TriggerNode())
 				{
 					// in preflight builds, we are either skipping triggers (and running things downstream) or we just stop at triggers and don't make them available for triggering.
 					continue;
@@ -2411,18 +2420,18 @@ public partial class GUBP : BuildCommand
 		string FinishNodePrint = DateTime.Now.ToString();
 		PrintCSVFile(String.Format("UAT,SetupCommanderPrint,{0},{1}", FinishFilterTime, FinishNodePrint));
 		// here we are just making sure everything before the explicit trigger is completed.
-		if (!String.IsNullOrEmpty(ExplicitTrigger))
+		if (ExplicitTrigger != null)
 		{
-			foreach (string NodeToDo in FilteredOrdereredToDo)
+			foreach (NodeInfo NodeToDo in FilteredOrdereredToDo)
 			{
-				if (GUBPNodes[NodeToDo].Node.RunInEC() && !NodeIsAlreadyComplete(GUBPNodesCompleted, NodeToDo, LocalOnly) && NodeToDo != ExplicitTrigger && !NodeDependsOn(ExplicitTrigger, NodeToDo)) // if something is already finished, we don't put it into EC
+				if (NodeToDo.Node.RunInEC() && !NodeIsAlreadyComplete(GUBPNodesCompleted, NodeToDo, LocalOnly) && NodeToDo != ExplicitTrigger && !NodeDependsOn(ExplicitTrigger, NodeToDo)) // if something is already finished, we don't put it into EC
 				{
-					throw new AutomationException("We are being asked to process node {0}, however, this is an explicit trigger {1}, so everything before it should already be handled. It seems likely that you waited too long to run the trigger. You will have to do a new build from scratch.", NodeToDo, ExplicitTrigger);
+					throw new AutomationException("We are being asked to process node {0}, however, this is an explicit trigger {1}, so everything before it should already be handled. It seems likely that you waited too long to run the trigger. You will have to do a new build from scratch.", NodeToDo.Name, ExplicitTrigger.Name);
 				}
 			}
 		}
 
-		string LastSticky = "";
+		NodeInfo LastSticky = null;
 		bool HitNonSticky = false;
 		bool bHaveECNodes = false;
 		List<string> StepList = new List<string>();
@@ -2433,12 +2442,12 @@ public partial class GUBP : BuildCommand
 		StepList.Add("$ec->setTimeout(600);");
 		StepList.Add("my $batch = $ec->newBatch(\"serial\");");
 		// sticky nodes are ones that we run on the main agent. We run then first and they must not be intermixed with parallel jobs
-		foreach (string NodeToDo in OrdereredToDo)
+		foreach (NodeInfo NodeToDo in OrdereredToDo)
 		{
-			if (GUBPNodes[NodeToDo].Node.RunInEC() && !NodeIsAlreadyComplete(GUBPNodesCompleted, NodeToDo, LocalOnly)) // if something is already finished, we don't put it into EC
+			if (NodeToDo.Node.RunInEC() && !NodeIsAlreadyComplete(GUBPNodesCompleted, NodeToDo, LocalOnly)) // if something is already finished, we don't put it into EC
 			{
 				bHaveECNodes = true;
-				if (GUBPNodes[NodeToDo].Node.IsSticky())
+				if (NodeToDo.Node.IsSticky())
 				{
 					LastSticky = NodeToDo;
 					if (HitNonSticky && !bSkipTriggers)
@@ -2458,7 +2467,7 @@ public partial class GUBP : BuildCommand
 		string BaseArgs = String.Format("$batch->createJobStep({{parentPath => '{0}'", ParentPath);
 
 		bool bHasNoop = false;
-		if (LastSticky == "" && bHaveECNodes)
+		if (LastSticky == null && bHaveECNodes)
 		{
 			// if we don't have any sticky nodes and we have other nodes, we run a fake noop just to release the resource 
 			string Args = String.Format("{0}, subprocedure => 'GUBP_UAT_Node', parallel => '0', jobStepName => 'Noop', actualParameter => [{{actualParameterName => 'NodeName', value => 'Noop'}}, {{actualParameterName => 'Sticky', value =>'1' }}], releaseMode => 'release'}});", BaseArgs);
@@ -2467,18 +2476,18 @@ public partial class GUBP : BuildCommand
 		}
 
 		List<string> FakeECArgs = new List<string>();
-		Dictionary<string, List<string>> AgentGroupChains = new Dictionary<string, List<string>>();
-		List<string> StickyChain = new List<string>();
-		foreach (string NodeToDo in OrdereredToDo)
+		Dictionary<string, List<NodeInfo>> AgentGroupChains = new Dictionary<string, List<NodeInfo>>();
+		List<NodeInfo> StickyChain = new List<NodeInfo>();
+		foreach (NodeInfo NodeToDo in OrdereredToDo)
 		{
-			if (GUBPNodes[NodeToDo].Node.RunInEC() && !NodeIsAlreadyComplete(GUBPNodesCompleted, NodeToDo, LocalOnly)) // if something is already finished, we don't put it into EC  
+			if (NodeToDo.Node.RunInEC() && !NodeIsAlreadyComplete(GUBPNodesCompleted, NodeToDo, LocalOnly)) // if something is already finished, we don't put it into EC  
 			{
-				string MyAgentGroup = GUBPNodes[NodeToDo].Node.AgentSharingGroup;
+				string MyAgentGroup = NodeToDo.Node.AgentSharingGroup;
 				if (MyAgentGroup != "")
 				{
 					if (!AgentGroupChains.ContainsKey(MyAgentGroup))
 					{
-						AgentGroupChains.Add(MyAgentGroup, new List<string> { NodeToDo });
+						AgentGroupChains.Add(MyAgentGroup, new List<NodeInfo> { NodeToDo });
 					}
 					else
 					{
@@ -2486,7 +2495,7 @@ public partial class GUBP : BuildCommand
 					}
 				}
 			}
-			if (GUBPNodes[NodeToDo].Node.IsSticky())
+			if (NodeToDo.Node.IsSticky())
 			{
 				if (!StickyChain.Contains(NodeToDo))
 				{
@@ -2494,32 +2503,32 @@ public partial class GUBP : BuildCommand
 				}
 			}
 		}
-		foreach (string NodeToDo in OrdereredToDo)
+		foreach (NodeInfo NodeToDo in OrdereredToDo)
 		{
-			if (GUBPNodes[NodeToDo].Node.RunInEC() && !NodeIsAlreadyComplete(GUBPNodesCompleted, NodeToDo, LocalOnly)) // if something is already finished, we don't put it into EC  
+			if (NodeToDo.Node.RunInEC() && !NodeIsAlreadyComplete(GUBPNodesCompleted, NodeToDo, LocalOnly)) // if something is already finished, we don't put it into EC  
 			{
 				string EMails;
 				List<string> NodeProps = GetECPropsForNode(NodeToDo, CLString, out EMails);
 				ECProps.AddRange(NodeProps);
 
-				bool Sticky = GUBPNodes[NodeToDo].Node.IsSticky();
+				bool Sticky = NodeToDo.Node.IsSticky();
 				bool DoParallel = !Sticky;
-				if (GUBPNodes[NodeToDo].Node.ECProcedure() == "GUBP_UAT_Node_Parallel_AgentShare_Editor")
+				if (NodeToDo.Node.ECProcedure() == "GUBP_UAT_Node_Parallel_AgentShare_Editor")
 				{
 					DoParallel = true;
 				}
-				if (Sticky && GUBPNodes[NodeToDo].Node.ECAgentString() != "")
+				if (Sticky && NodeToDo.Node.ECAgentString() != "")
 				{
-					throw new AutomationException("Node {1} is sticky but has agent requirements.", NodeToDo);
+					throw new AutomationException("Node {1} is sticky but has agent requirements.", NodeToDo.Name);
 				}
-				string Procedure = GUBPNodes[NodeToDo].Node.ECProcedure();
-				if (GUBPNodes[NodeToDo].Node.IsSticky() && NodeToDo == LastSticky)
+				string Procedure = NodeToDo.Node.ECProcedure();
+				if (NodeToDo.Node.IsSticky() && NodeToDo == LastSticky)
 				{
 					Procedure = Procedure + "_Release";
 				}
 				string Args = String.Format("{0}, subprocedure => '{1}', parallel => '{2}', jobStepName => '{3}', actualParameter => [{{actualParameterName => 'NodeName', value =>'{4}'}}",
-					BaseArgs, Procedure, DoParallel ? 1 : 0, NodeToDo, NodeToDo);
-				string ProcedureParams = GUBPNodes[NodeToDo].Node.ECProcedureParams();
+					BaseArgs, Procedure, DoParallel ? 1 : 0, NodeToDo.Name, NodeToDo.Name);
+				string ProcedureParams = NodeToDo.Node.ECProcedureParams();
 				if (!String.IsNullOrEmpty(ProcedureParams))
 				{
 					Args = Args + ProcedureParams;
@@ -2535,7 +2544,7 @@ public partial class GUBP : BuildCommand
 				List<string> UncompletedEcDeps = new List<string>();
 				List<string> UncompletedCompletedDeps = new List<string>();
 				List<string> PreConditionUncompletedEcDeps = new List<string>();
-				string MyAgentGroup = GUBPNodes[NodeToDo].Node.AgentSharingGroup;
+				string MyAgentGroup = NodeToDo.Node.AgentSharingGroup;
 				bool bDoNestedJobstep = false;
 				bool bDoFirstNestedJobstep = false;
 
@@ -2545,16 +2554,16 @@ public partial class GUBP : BuildCommand
 
 				PreconditionParentPath = ParentPath;
 				{
-					List<string> EcDeps = GetECDependencies(NodeToDo);
-					foreach (string Dep in EcDeps)
+					List<NodeInfo> EcDeps = GetECDependencies(NodeToDo);
+					foreach (NodeInfo Dep in EcDeps)
 					{
-						if (GUBPNodes[Dep].Node.RunInEC() && !NodeIsAlreadyComplete(GUBPNodesCompleted, Dep, LocalOnly) && OrdereredToDo.Contains(Dep)) // if something is already finished, we don't put it into EC
+						if (Dep.Node.RunInEC() && !NodeIsAlreadyComplete(GUBPNodesCompleted, Dep, LocalOnly) && OrdereredToDo.Contains(Dep)) // if something is already finished, we don't put it into EC
 						{
 							if (OrdereredToDo.IndexOf(Dep) > OrdereredToDo.IndexOf(NodeToDo))
 							{
-								throw new AutomationException("Topological sort error, node {0} has a dependency of {1} which sorted after it.", NodeToDo, Dep);
+								throw new AutomationException("Topological sort error, node {0} has a dependency of {1} which sorted after it.", NodeToDo.Name, Dep.Name);
 							}
-							UncompletedEcDeps.Add(Dep);
+							UncompletedEcDeps.Add(Dep.Name);
 						}
 					}
 				}
@@ -2563,16 +2572,16 @@ public partial class GUBP : BuildCommand
 				{
 					PreConditionUncompletedEcDeps.Add(Dep);
 				}
-				List<string> CompletedDeps = GetCompletedOnlyDependencies(NodeToDo);
-				foreach (string Dep in CompletedDeps)
+				List<NodeInfo> CompletedDeps = GetCompletedOnlyDependencies(NodeToDo);
+				foreach (NodeInfo Dep in CompletedDeps)
 				{
-					if (GUBPNodes[Dep].Node.RunInEC() && !NodeIsAlreadyComplete(GUBPNodesCompleted, Dep, LocalOnly) && OrdereredToDo.Contains(Dep)) // if something is already finished, we don't put it into EC
+					if (Dep.Node.RunInEC() && !NodeIsAlreadyComplete(GUBPNodesCompleted, Dep, LocalOnly) && OrdereredToDo.Contains(Dep)) // if something is already finished, we don't put it into EC
 					{
 						if (OrdereredToDo.IndexOf(Dep) > OrdereredToDo.IndexOf(NodeToDo))
 						{
-							throw new AutomationException("Topological sort error, node {0} has a dependency of {1} which sorted after it.", NodeToDo, Dep);
+							throw new AutomationException("Topological sort error, node {0} has a dependency of {1} which sorted after it.", NodeToDo.Name, Dep.Name);
 						}
-						UncompletedCompletedDeps.Add(Dep);
+						UncompletedCompletedDeps.Add(Dep.Name);
 					}
 				}
 				if (MyAgentGroup != "")
@@ -2582,61 +2591,61 @@ public partial class GUBP : BuildCommand
 
 					PreConditionUncompletedEcDeps = new List<string>();
 
-					List<string> MyChain = AgentGroupChains[MyAgentGroup];
+					List<NodeInfo> MyChain = AgentGroupChains[MyAgentGroup];
 					int MyIndex = MyChain.IndexOf(NodeToDo);
 					if (MyIndex > 0)
 					{
-						PreConditionUncompletedEcDeps.Add(MyChain[MyIndex - 1]);
+						PreConditionUncompletedEcDeps.Add(MyChain[MyIndex - 1].Name);
 					}
 					else
 					{
 						bDoFirstNestedJobstep = bDoNestedJobstep;
 						// to avoid idle agents (and also EC doesn't actually reserve our agent!), we promote all dependencies to the first one
-						foreach (string Chain in MyChain)
+						foreach (NodeInfo Chain in MyChain)
 						{
-							List<string> EcDeps = GetECDependencies(Chain);
-							foreach (string Dep in EcDeps)
+							List<NodeInfo> EcDeps = GetECDependencies(Chain);
+							foreach (NodeInfo Dep in EcDeps)
 							{
-								if (GUBPNodes[Dep].Node.RunInEC() && !NodeIsAlreadyComplete(GUBPNodesCompleted, Dep, LocalOnly) && OrdereredToDo.Contains(Dep)) // if something is already finished, we don't put it into EC
+								if (GUBPNodes[Dep.Name].Node.RunInEC() && !NodeIsAlreadyComplete(GUBPNodesCompleted, Dep, LocalOnly) && OrdereredToDo.Contains(Dep)) // if something is already finished, we don't put it into EC
 								{
 									if (OrdereredToDo.IndexOf(Dep) > OrdereredToDo.IndexOf(Chain))
 									{
-										throw new AutomationException("Topological sort error, node {0} has a dependency of {1} which sorted after it.", Chain, Dep);
+										throw new AutomationException("Topological sort error, node {0} has a dependency of {1} which sorted after it.", Chain.Name, Dep.Name);
 									}
-									if (!MyChain.Contains(Dep) && !PreConditionUncompletedEcDeps.Contains(Dep))
+									if (!MyChain.Contains(Dep) && !PreConditionUncompletedEcDeps.Contains(Dep.Name))
 									{
-										PreConditionUncompletedEcDeps.Add(Dep);
+										PreConditionUncompletedEcDeps.Add(Dep.Name);
 									}
 								}
 							}
 						}
 					}
 				}
-				if (GUBPNodes[NodeToDo].Node.IsSticky())
+				if (NodeToDo.Node.IsSticky())
 				{
-					List<string> MyChain = StickyChain;
+					List<NodeInfo> MyChain = StickyChain;
 					int MyIndex = MyChain.IndexOf(NodeToDo);
 					if (MyIndex > 0)
 					{
-						if (!PreConditionUncompletedEcDeps.Contains(MyChain[MyIndex - 1]) && !NodeIsAlreadyComplete(GUBPNodesCompleted, MyChain[MyIndex - 1], LocalOnly))
+						if (!PreConditionUncompletedEcDeps.Contains(MyChain[MyIndex - 1].Name) && !NodeIsAlreadyComplete(GUBPNodesCompleted, MyChain[MyIndex - 1], LocalOnly))
 						{
-							PreConditionUncompletedEcDeps.Add(MyChain[MyIndex - 1]);
+							PreConditionUncompletedEcDeps.Add(MyChain[MyIndex - 1].Name);
 						}
 					}
 					else
 					{
-						List<string> EcDeps = GetECDependencies(NodeToDo);
-						foreach (string Dep in EcDeps)
+						List<NodeInfo> EcDeps = GetECDependencies(NodeToDo);
+						foreach (NodeInfo Dep in EcDeps)
 						{
-							if (GUBPNodes[Dep].Node.RunInEC() && !NodeIsAlreadyComplete(GUBPNodesCompleted, Dep, LocalOnly) && OrdereredToDo.Contains(Dep)) // if something is already finished, we don't put it into EC
+							if (Dep.Node.RunInEC() && !NodeIsAlreadyComplete(GUBPNodesCompleted, Dep, LocalOnly) && OrdereredToDo.Contains(Dep)) // if something is already finished, we don't put it into EC
 							{
 								if (OrdereredToDo.IndexOf(Dep) > OrdereredToDo.IndexOf(NodeToDo))
 								{
-									throw new AutomationException("Topological sort error, node {0} has a dependency of {1} which sorted after it.", NodeToDo, Dep);
+									throw new AutomationException("Topological sort error, node {0} has a dependency of {1} which sorted after it.", NodeToDo.Name, Dep.Name);
 								}
-								if (!MyChain.Contains(Dep) && !PreConditionUncompletedEcDeps.Contains(Dep))
+								if (!MyChain.Contains(Dep) && !PreConditionUncompletedEcDeps.Contains(Dep.Name))
 								{
-									PreConditionUncompletedEcDeps.Add(Dep);
+									PreConditionUncompletedEcDeps.Add(Dep.Name);
 								}
 							}
 						}
@@ -2667,11 +2676,11 @@ public partial class GUBP : BuildCommand
 					PreCondition = "\"\\$\" . \"[/javascript if(";
 					// these run "parallel", but we add preconditions to serialize them
 					int Index = 0;
-					foreach (string Dep in CompletedDeps)
+					foreach (NodeInfo Dep in CompletedDeps)
 					{
-						if (GUBPNodes[Dep].Node.RunInEC())
+						if (Dep.Node.RunInEC())
 						{
-							PreCondition = PreCondition + "getProperty('" + GetJobStep(PreconditionParentPath, Dep) + "/status\') == \'completed\'";
+							PreCondition = PreCondition + "getProperty('" + GetJobStep(PreconditionParentPath, Dep.Name) + "/status\') == \'completed\'";
 							Index++;
 							if (Index != CompletedDeps.Count)
 							{
@@ -2716,7 +2725,7 @@ public partial class GUBP : BuildCommand
 						}
 						{
 							string NestArgs = String.Format("$batch->createJobStep({{parentPath => '{0}/jobSteps[{1}]', jobStepName => '{2}_GetPool', subprocedure => 'GUBP{3}_AgentShare_GetPool', parallel => '1', actualParameter => [{{actualParameterName => 'AgentSharingGroup', value => '{4}'}}, {{actualParameterName => 'NodeName', value => '{5}'}}]",
-								ParentPath, MyAgentGroup, MyAgentGroup, GUBPNodes[NodeToDo].Node.ECProcedureInfix(), MyAgentGroup, NodeToDo);
+								ParentPath, MyAgentGroup, MyAgentGroup, NodeToDo.Node.ECProcedureInfix(), MyAgentGroup, NodeToDo.Name);
 							if (!String.IsNullOrEmpty(PreCondition))
 							{
 								NestArgs = NestArgs + ", precondition => " + PreCondition;
@@ -2726,9 +2735,9 @@ public partial class GUBP : BuildCommand
 						}
 						{
 							string NestArgs = String.Format("$batch->createJobStep({{parentPath => '{0}/jobSteps[{1}]', jobStepName => '{2}_GetAgent', subprocedure => 'GUBP{3}_AgentShare_GetAgent', parallel => '1', exclusiveMode => 'call', resourceName => '{4}', actualParameter => [{{actualParameterName => 'AgentSharingGroup', value => '{5}'}}, {{actualParameterName => 'NodeName', value=> '{6}'}}]",
-								ParentPath, MyAgentGroup, MyAgentGroup, GUBPNodes[NodeToDo].Node.ECProcedureInfix(),
+								ParentPath, MyAgentGroup, MyAgentGroup, NodeToDo.Node.ECProcedureInfix(),
 								String.Format("$[/myJob/jobSteps[{0}]/ResourcePool]", MyAgentGroup),
-								MyAgentGroup, NodeToDo);
+								MyAgentGroup, NodeToDo.Name);
 							{
 								NestArgs = NestArgs + ", precondition  => ";
 								NestArgs = NestArgs + "\"\\$\" . \"[/javascript if(";
@@ -2771,7 +2780,7 @@ public partial class GUBP : BuildCommand
 				StepList.Add(Args);
 				if (bFakeEC &&
 					!UnfinishedTriggers.Contains(NodeToDo) &&
-					(GUBPNodes[NodeToDo].Node.ECProcedure().StartsWith("GUBP_UAT_Node") || GUBPNodes[NodeToDo].Node.ECProcedure().StartsWith("GUBP_Mac_UAT_Node")) // other things we really can't test
+					(NodeToDo.Node.ECProcedure().StartsWith("GUBP_UAT_Node") || NodeToDo.Node.ECProcedure().StartsWith("GUBP_Mac_UAT_Node")) // other things we really can't test
 					) // unfinished triggers are never run directly by EC, rather it does another job setup
 				{
 					string Arg = String.Format("gubp -Node={0} -FakeEC {1} {2} {3} {4} {5}",
@@ -2797,13 +2806,13 @@ public partial class GUBP : BuildCommand
 
 				if (MyAgentGroup != "" && !bDoNestedJobstep)
 				{
-					List<string> MyChain = AgentGroupChains[MyAgentGroup];
+					List<NodeInfo> MyChain = AgentGroupChains[MyAgentGroup];
 					int MyIndex = MyChain.IndexOf(NodeToDo);
 					if (MyIndex == MyChain.Count - 1)
 					{
 						string RelPreCondition = "\"\\$\" . \"[/javascript if(";
 						// this runs "parallel", but we a precondition to serialize it
-						RelPreCondition = RelPreCondition + "getProperty('" + PreconditionParentPath + "/jobSteps[" + NodeToDo + "]/status') == 'completed'";
+						RelPreCondition = RelPreCondition + "getProperty('" + PreconditionParentPath + "/jobSteps[" + NodeToDo.Name + "]/status') == 'completed'";
 						RelPreCondition = RelPreCondition + ") true;]\"";
 						// we need to release the resource
 						string RelArgs = String.Format("{0}, subprocedure => 'GUBP_Release_AgentShare', parallel => '1', jobStepName => 'Release_{1}', actualParameter => [{{actualParameterName => 'AgentSharingGroup', valued => '{2}'}}], releaseMode => 'release', precondition => '{3}'",
@@ -2838,64 +2847,64 @@ public partial class GUBP : BuildCommand
 		}
 	}
 
-	void ExecuteNodes(List<string> OrdereredToDo, bool bOnlyNode, bool bFakeEC, bool LocalOnly, bool bSaveSharedTempStorage, Dictionary<string, bool> GUBPNodesCompleted, Dictionary<string, NodeHistory> GUBPNodesHistory, string CLString, string FakeFail)
+	void ExecuteNodes(List<NodeInfo> OrdereredToDo, bool bOnlyNode, bool bFakeEC, bool LocalOnly, bool bSaveSharedTempStorage, Dictionary<string, bool> GUBPNodesCompleted, Dictionary<string, NodeHistory> GUBPNodesHistory, string CLString, string FakeFail)
 	{
-        Dictionary<string, string> BuildProductToNodeMap = new Dictionary<string, string>();
-		foreach (string NodeToDo in OrdereredToDo)
+        Dictionary<string, NodeInfo> BuildProductToNodeMap = new Dictionary<string, NodeInfo>();
+		foreach (NodeInfo NodeToDo in OrdereredToDo)
         {
-            if (GUBPNodes[NodeToDo].Node.BuildProducts != null || GUBPNodes[NodeToDo].Node.AllDependencyBuildProducts != null)
+            if (NodeToDo.Node.BuildProducts != null || NodeToDo.Node.AllDependencyBuildProducts != null)
             {
                 throw new AutomationException("topological sort error");
             }
 
-            GUBPNodes[NodeToDo].Node.AllDependencyBuildProducts = new List<string>();
-            GUBPNodes[NodeToDo].Node.AllDependencies = new List<string>();
-            foreach (string Dep in GUBPNodes[NodeToDo].Node.FullNamesOfDependencies)
+            NodeToDo.Node.AllDependencyBuildProducts = new List<string>();
+            NodeToDo.Node.AllDependencies = new List<string>();
+            foreach (NodeInfo Dep in NodeToDo.Node.FullNamesOfDependencies.Select(x => GUBPNodes[x]))
             {
-                GUBPNodes[NodeToDo].Node.AddAllDependent(Dep);
-                if (GUBPNodes[Dep].Node.AllDependencies == null)
+                NodeToDo.Node.AddAllDependent(Dep.Name);
+                if (Dep.Node.AllDependencies == null)
                 {
                     if (!bOnlyNode)
                     {
-                        throw new AutomationException("Node {0} was not processed yet3?  Processing {1}", Dep, NodeToDo);
+                        throw new AutomationException("Node {0} was not processed yet3?  Processing {1}", Dep, NodeToDo.Name);
                     }
                 }
                 else
 				{
-					foreach (string DepDep in GUBPNodes[Dep].Node.AllDependencies)
+					foreach (string DepDep in Dep.Node.AllDependencies)
 					{
-						GUBPNodes[NodeToDo].Node.AddAllDependent(DepDep);
+						NodeToDo.Node.AddAllDependent(DepDep);
 					}
 				}
 				
-                if (GUBPNodes[Dep].Node.BuildProducts == null)
+                if (Dep.Node.BuildProducts == null)
                 {
                     if (!bOnlyNode)
                     {
-                        throw new AutomationException("Node {0} was not processed yet? Processing {1}", Dep, NodeToDo);
+                        throw new AutomationException("Node {0} was not processed yet? Processing {1}", Dep, NodeToDo.Name);
                     }
                 }
                 else
                 {
-                    foreach (string Prod in GUBPNodes[Dep].Node.BuildProducts)
+                    foreach (string Prod in Dep.Node.BuildProducts)
                     {
-                        GUBPNodes[NodeToDo].Node.AddDependentBuildProduct(Prod);
+                        NodeToDo.Node.AddDependentBuildProduct(Prod);
                     }
-                    if (GUBPNodes[Dep].Node.AllDependencyBuildProducts == null)
+                    if (Dep.Node.AllDependencyBuildProducts == null)
                     {
-                        throw new AutomationException("Node {0} was not processed yet2?  Processing {1}", Dep, NodeToDo);
+                        throw new AutomationException("Node {0} was not processed yet2?  Processing {1}", Dep.Name, NodeToDo.Name);
                     }
-                    foreach (string Prod in GUBPNodes[Dep].Node.AllDependencyBuildProducts)
+                    foreach (string Prod in Dep.Node.AllDependencyBuildProducts)
                     {
-                        GUBPNodes[NodeToDo].Node.AddDependentBuildProduct(Prod);
+                        NodeToDo.Node.AddDependentBuildProduct(Prod);
                     }
                 }
             }
 
-            string NodeStoreName = StoreName + "-" + GUBPNodes[NodeToDo].Node.GetFullName();
+            string NodeStoreName = StoreName + "-" + NodeToDo.Node.GetFullName();
             
-            string GameNameIfAny = GUBPNodes[NodeToDo].Node.GameNameIfAnyForTempStorage();
-            string StorageRootIfAny = GUBPNodes[NodeToDo].Node.RootIfAnyForTempStorage();
+            string GameNameIfAny = NodeToDo.Node.GameNameIfAnyForTempStorage();
+            string StorageRootIfAny = NodeToDo.Node.RootIfAnyForTempStorage();
 					
             if (bFake)
             {
@@ -2904,40 +2913,40 @@ public partial class GUBP : BuildCommand
 
             // this is kinda complicated
             bool SaveSuccessRecords = (IsBuildMachine || bFakeEC) && // no real reason to make these locally except for fakeEC tests
-                (!GUBPNodes[NodeToDo].Node.TriggerNode() || GUBPNodes[NodeToDo].Node.IsSticky()) // trigger nodes are run twice, one to start the new workflow and once when it is actually triggered, we will save reconds for the latter
-                && (GUBPNodes[NodeToDo].Node.RunInEC() || !GUBPNodes[NodeToDo].Node.IsAggregate()); //aggregates not in EC can be "run" multiple times, so we can't track those
+                (!NodeToDo.Node.TriggerNode() || NodeToDo.Node.IsSticky()) // trigger nodes are run twice, one to start the new workflow and once when it is actually triggered, we will save reconds for the latter
+                && (NodeToDo.Node.RunInEC() || !NodeToDo.Node.IsAggregate()); //aggregates not in EC can be "run" multiple times, so we can't track those
 
-            Log("***** Running GUBP Node {0} -> {1} : {2}", GUBPNodes[NodeToDo].Node.GetFullName(), GameNameIfAny, NodeStoreName);
+            Log("***** Running GUBP Node {0} -> {1} : {2}", NodeToDo.Node.GetFullName(), GameNameIfAny, NodeStoreName);
             if (NodeIsAlreadyComplete(GUBPNodesCompleted, NodeToDo, LocalOnly))
             {
-                if (NodeToDo == VersionFilesNode.StaticGetFullName() && !IsBuildMachine)
+                if (NodeToDo.Name == VersionFilesNode.StaticGetFullName() && !IsBuildMachine)
                 {
-                    Log("***** NOT ****** Retrieving GUBP Node {0} from {1}; it is the version files.", GUBPNodes[NodeToDo].Node.GetFullName(), NodeStoreName);
-                    GUBPNodes[NodeToDo].Node.BuildProducts = new List<string>();
+                    Log("***** NOT ****** Retrieving GUBP Node {0} from {1}; it is the version files.", NodeToDo.Node.GetFullName(), NodeStoreName);
+                    NodeToDo.Node.BuildProducts = new List<string>();
 
                 }
                 else
                 {
-                    Log("***** Retrieving GUBP Node {0} from {1}", GUBPNodes[NodeToDo].Node.GetFullName(), NodeStoreName);
+                    Log("***** Retrieving GUBP Node {0} from {1}", NodeToDo.Node.GetFullName(), NodeStoreName);
                     bool WasLocal;
 					try
 					{
-						GUBPNodes[NodeToDo].Node.BuildProducts = TempStorage.RetrieveFromTempStorage(CmdEnv, NodeStoreName, out WasLocal, GameNameIfAny, StorageRootIfAny);
+						NodeToDo.Node.BuildProducts = TempStorage.RetrieveFromTempStorage(CmdEnv, NodeStoreName, out WasLocal, GameNameIfAny, StorageRootIfAny);
 					}
 					catch
 					{
 						if(GameNameIfAny != "")
 						{
-							GUBPNodes[NodeToDo].Node.BuildProducts = TempStorage.RetrieveFromTempStorage(CmdEnv, NodeStoreName, out WasLocal, "", StorageRootIfAny);
+							NodeToDo.Node.BuildProducts = TempStorage.RetrieveFromTempStorage(CmdEnv, NodeStoreName, out WasLocal, "", StorageRootIfAny);
 						}
 						else
 						{
-							throw new AutomationException("Build Products cannot be found for node {0}", NodeToDo);
+							throw new AutomationException("Build Products cannot be found for node {0}", NodeToDo.Name);
 						}
 					}
                     if (!WasLocal)
                     {
-                        GUBPNodes[NodeToDo].Node.PostLoadFromSharedTempStorage(this);
+                        NodeToDo.Node.PostLoadFromSharedTempStorage(this);
                     }
                 }
             }
@@ -2950,35 +2959,35 @@ public partial class GUBP : BuildCommand
 				double BuildDuration = 0.0;
                 try
                 {
-                    if (!String.IsNullOrEmpty(FakeFail) && FakeFail.Equals(NodeToDo, StringComparison.InvariantCultureIgnoreCase))
+                    if (!String.IsNullOrEmpty(FakeFail) && FakeFail.Equals(NodeToDo.Name, StringComparison.InvariantCultureIgnoreCase))
                     {
-                        throw new AutomationException("Failing node {0} by request.", NodeToDo);
+                        throw new AutomationException("Failing node {0} by request.", NodeToDo.Name);
                     }
                     if (bFake)
                     {
-                        Log("***** FAKE!! Building GUBP Node {0} for {1}", NodeToDo, NodeStoreName);
-                        GUBPNodes[NodeToDo].Node.DoFakeBuild(this);
+                        Log("***** FAKE!! Building GUBP Node {0} for {1}", NodeToDo.Name, NodeStoreName);
+                        NodeToDo.Node.DoFakeBuild(this);
                     }
                     else
                     {                        
-						Log("***** Building GUBP Node {0} for {1}", NodeToDo, NodeStoreName);
+						Log("***** Building GUBP Node {0} for {1}", NodeToDo.Name, NodeStoreName);
 						DateTime StartTime = DateTime.UtcNow;
                         string StartBuild = DateTime.Now.ToString();
-                        GUBPNodes[NodeToDo].Node.DoBuild(this);
+                        NodeToDo.Node.DoBuild(this);
                         string FinishBuild = DateTime.Now.ToString();
                         if (IsBuildMachine)
                         {
-                            PrintCSVFile(String.Format("UAT,DoBuild.{0},{1},{2}", NodeToDo, StartBuild, FinishBuild));
+                            PrintCSVFile(String.Format("UAT,DoBuild.{0},{1},{2}", NodeToDo.Name, StartBuild, FinishBuild));
                         }
 						BuildDuration = (DateTime.UtcNow - StartTime).TotalMilliseconds / 1000;
 						
                     }
-                    if (!GUBPNodes[NodeToDo].Node.IsAggregate())
+                    if (!NodeToDo.Node.IsAggregate())
                     {
 						double StoreDuration = 0.0;
 						DateTime StartTime = DateTime.UtcNow;
                         string StartStore = DateTime.Now.ToString();
-                        TempStorage.StoreToTempStorage(CmdEnv, NodeStoreName, GUBPNodes[NodeToDo].Node.BuildProducts, !bSaveSharedTempStorage, GameNameIfAny, StorageRootIfAny);
+                        TempStorage.StoreToTempStorage(CmdEnv, NodeStoreName, NodeToDo.Node.BuildProducts, !bSaveSharedTempStorage, GameNameIfAny, StorageRootIfAny);
 						StoreDuration = (DateTime.UtcNow - StartTime).TotalMilliseconds / 1000;
                         string FinishStore = DateTime.Now.ToString();                        
 						Log("Took {0} seconds to store build products", StoreDuration);
@@ -2989,7 +2998,7 @@ public partial class GUBP : BuildCommand
 						}
                         if (ParseParam("StompCheck"))
                         {
-                            foreach (string Dep in GUBPNodes[NodeToDo].Node.AllDependencies)
+                            foreach (string Dep in NodeToDo.Node.AllDependencies)
                             {
                                 try
                                 {
@@ -3009,7 +3018,7 @@ public partial class GUBP : BuildCommand
                                 }
                                 catch(Exception Ex)
                                 {
-                                    throw new AutomationException("Node {0} stomped Node {1}   Ex: {2}", NodeToDo, Dep, LogUtils.FormatException(Ex));
+                                    throw new AutomationException("Node {0} stomped Node {1}   Ex: {2}", NodeToDo.Name, Dep, LogUtils.FormatException(Ex));
                                 }
                             }
                         }
@@ -3043,9 +3052,9 @@ public partial class GUBP : BuildCommand
                     Log("{0}", ExceptionToString(Ex));
 
 
-                    if (GUBPNodesHistory.ContainsKey(NodeToDo))
+                    if (GUBPNodesHistory.ContainsKey(NodeToDo.Name))
                     {
-                        NodeHistory History = GUBPNodesHistory[NodeToDo];
+                        NodeHistory History = GUBPNodesHistory[NodeToDo.Name];
                         Log("Changes since last green *********************************");
                         Log("");
                         Log("");
@@ -3107,11 +3116,11 @@ public partial class GUBP : BuildCommand
 					UpdateECBuildTime(NodeToDo, BuildDuration);
                 }
             }
-            foreach (string Product in GUBPNodes[NodeToDo].Node.BuildProducts)
+            foreach (string Product in NodeToDo.Node.BuildProducts)
             {
                 if (BuildProductToNodeMap.ContainsKey(Product))
                 {
-                    throw new AutomationException("Overlapping build product: {0} and {1} both produce {2}", BuildProductToNodeMap[Product], NodeToDo, Product);
+                    throw new AutomationException("Overlapping build product: {0} and {1} both produce {2}", BuildProductToNodeMap[Product], NodeToDo.Name, Product);
                 }
                 BuildProductToNodeMap.Add(Product, NodeToDo);
             }
diff --git a/Engine/Source/Programs/AutomationTool/GUBP/LegacyBranchSetup.cs b/Engine/Source/Programs/AutomationTool/GUBP/LegacyBranchSetup.cs
--- a/Engine/Source/Programs/AutomationTool/GUBP/LegacyBranchSetup.cs
+++ b/Engine/Source/Programs/AutomationTool/GUBP/LegacyBranchSetup.cs
@@ -38,10 +38,10 @@ partial class GUBP
         {
             Adders = new List<GUBPNodeAdder>();
             Assembly[] LoadedAssemblies = AppDomain.CurrentDomain.GetAssemblies();
-            foreach (var Dll in LoadedAssemblies)
+            foreach (Assembly Dll in LoadedAssemblies)
             {
 				Type[] AllTypes = GetTypesFromAssembly(Dll);
-                foreach (var PotentialConfigType in AllTypes)
+                foreach (Type PotentialConfigType in AllTypes)
                 {
                     if (PotentialConfigType != typeof(GUBPNodeAdder) && typeof(GUBPNodeAdder).IsAssignableFrom(PotentialConfigType))
                     {
@@ -54,7 +54,7 @@ partial class GUBP
                 }
             }
         }
-        foreach(var Adder in Adders)
+        foreach(GUBPNodeAdder Adder in Adders)
         {
             Adder.AddNodes(this, InHostPlatform, InActivePlatforms);
         }
@@ -93,7 +93,7 @@ partial class GUBP
         {
             BranchHackers = new List<GUBPBranchHacker>();
             Assembly[] LoadedAssemblies = AppDomain.CurrentDomain.GetAssemblies();
-            foreach (var Dll in LoadedAssemblies)
+            foreach (Assembly Dll in LoadedAssemblies)
             {
 				Type[] AllTypes = GetTypesFromAssembly(Dll);
                 foreach (var PotentialConfigType in AllTypes)
@@ -109,8 +109,8 @@ partial class GUBP
                 }
             }
         }
-        var Result = new GUBPBranchHacker.BranchOptions();
-        foreach (var Hacker in BranchHackers)
+        GUBPBranchHacker.BranchOptions Result = new GUBPBranchHacker.BranchOptions();
+        foreach (GUBPBranchHacker Hacker in BranchHackers)
         {
             Hacker.ModifyOptions(this, ref Result, Branch);
         }
@@ -138,7 +138,7 @@ partial class GUBP
     }
 
     private static List<GUBPEmailHacker> EmailHackers;
-    private string HackEmails(string Emails, string Causers, string Branch, string NodeName)
+    private string HackEmails(string Emails, string Causers, string Branch, NodeInfo NodeInfo)
     {
         string OnlyEmail = ParseParamValue("OnlyEmail");
         if (!String.IsNullOrEmpty(OnlyEmail))
@@ -171,26 +171,26 @@ partial class GUBP
             }
         }
         List<string> Result = new List<string>(Emails.Split(' '));
-		if(!EmailHackers.Any(x => x.VetoEmailingCausers(this, Branch, NodeName)))
+		if(!EmailHackers.Any(x => x.VetoEmailingCausers(this, Branch, NodeInfo.Name)))
 		{
 			Result.AddRange(Causers.Split(' '));
 		}
         foreach (var EmailHacker in EmailHackers)
         {
-            Result.AddRange(EmailHacker.AddEmails(this, Branch, NodeName));
+            Result.AddRange(EmailHacker.AddEmails(this, Branch, NodeInfo.Name));
         }
         foreach (var EmailHacker in EmailHackers)
         {
             var NewResult = new List<string>();
             foreach (var Email in Result)
             {
-                NewResult.AddRange(EmailHacker.ModifyEmail(Email, this, Branch, NodeName));
+                NewResult.AddRange(EmailHacker.ModifyEmail(Email, this, Branch, NodeInfo.Name));
             }
             Result = NewResult;
         }
         foreach (var EmailHacker in EmailHackers)
         {
-            Result = EmailHacker.FinalizeEmails(Result, this, Branch, NodeName);
+            Result = EmailHacker.FinalizeEmails(Result, this, Branch, NodeInfo.Name);
         }
         string FinalEmails = "";
         int Count = 0;
@@ -209,7 +209,7 @@ partial class GUBP
         }
     }
     private static List<GUBPFrequencyHacker> FrequencyHackers;
-    private int HackFrequency(GUBP bp, string Branch, string NodeName, int BaseFrequency)
+    private int HackFrequency(GUBP bp, string Branch, NodeInfo NodeInfo, int BaseFrequency)
     {
         int Frequency = BaseFrequency;
         if (FrequencyHackers == null)
@@ -234,7 +234,7 @@ partial class GUBP
         }
         foreach(var FrequencyHacker in FrequencyHackers)
         {
-           Frequency = FrequencyHacker.GetNodeFrequency(bp, Branch, NodeName, BaseFrequency);            
+           Frequency = FrequencyHacker.GetNodeFrequency(bp, Branch, NodeInfo.Name, BaseFrequency);            
         }
         return Frequency;
     }    
