Cache dependency list in the NodeInfo instead of looking them up by name each time.

diff --git a/Engine/Source/Programs/AutomationTool/GUBP/GUBP.cs b/Engine/Source/Programs/AutomationTool/GUBP/GUBP.cs
--- a/Engine/Source/Programs/AutomationTool/GUBP/GUBP.cs
+++ b/Engine/Source/Programs/AutomationTool/GUBP/GUBP.cs
@@ -34,6 +34,8 @@ public partial class GUBP : BuildCommand
 	{
 		public string Name;
 		public GUBPNode Node;
+		public List<NodeInfo> Dependencies;
+		public List<NodeInfo> PseudoDependencies;
 
 		public override string ToString()
 		{
@@ -104,7 +106,7 @@ public partial class GUBP : BuildCommand
     List<NodeInfo> GetDependencies(NodeInfo NodeToDo, bool bFlat = false, bool ECOnly = false)
     {
         List<NodeInfo> Result = new List<NodeInfo>();
-        foreach (NodeInfo Dependency in NodeToDo.Node.FullNamesOfDependencies.Select(x => GUBPNodes[x]))
+        foreach (NodeInfo Dependency in NodeToDo.Dependencies)
         {
             bool Usable = Dependency.Node.RunInEC() || !ECOnly;
             if (Usable)
@@ -125,7 +127,7 @@ public partial class GUBP : BuildCommand
                 }
             }
         }
-        foreach (NodeInfo PseudoDependency in NodeToDo.Node.FullNamesOfPseudosependencies.Select(x => GUBPNodes[x]))
+        foreach (NodeInfo PseudoDependency in NodeToDo.PseudoDependencies)
         {
             bool Usable = PseudoDependency.Node.RunInEC() || !ECOnly;
             if (Usable)
@@ -192,13 +194,8 @@ public partial class GUBP : BuildCommand
             return GUBPNodesControllingTrigger[NodeToDo];
         }
         NodeInfo Result = null;
-        foreach (string DependencyName in NodeToDo.Node.FullNamesOfDependencies)
+        foreach (NodeInfo Dependency in NodeToDo.Dependencies)
 		{
-			NodeInfo Dependency;
-            if (!GUBPNodes.TryGetValue(DependencyName, out Dependency))
-            {
-                throw new AutomationException("Dependency {0} in {1} not found.", DependencyName, NodeToDo.Name);
-            }
             bool IsTrigger = Dependency.Node.TriggerNode();
             if (IsTrigger)
             {
@@ -221,13 +218,8 @@ public partial class GUBP : BuildCommand
                 }
             }
         }
-        foreach (string PseudoDependencyName in NodeToDo.Node.FullNamesOfPseudosependencies)
+        foreach (NodeInfo PseudoDependency in NodeToDo.PseudoDependencies)
         {
-			NodeInfo PseudoDependency;
-            if (!GUBPNodes.TryGetValue(PseudoDependencyName, out PseudoDependency))
-            {
-                throw new AutomationException("Pseudodependency {0} in {1} not found.", PseudoDependencyName, NodeToDo.Name);
-            }
             bool IsTrigger = PseudoDependency.Node.TriggerNode();
             if (IsTrigger)
             {
@@ -303,7 +295,7 @@ public partial class GUBP : BuildCommand
         return FrequencyString;
     }
 
-    int ComputeDependentCISFrequencyQuantumShift(NodeInfo NodeToDo, Dictionary<string, int> FrequencyOverrides)
+    int ComputeDependentCISFrequencyQuantumShift(NodeInfo NodeToDo, Dictionary<NodeInfo, int> FrequencyOverrides)
     {
         int Result = NodeToDo.Node.ComputedDependentCISFrequencyQuantumShift;        
         if (Result < 0)
@@ -312,16 +304,16 @@ public partial class GUBP : BuildCommand
             Result = GetFrequencyForNode(this, NodeToDo, Result);
 
 			int FrequencyOverride;
-			if(FrequencyOverrides.TryGetValue(NodeToDo.Name, out FrequencyOverride) && Result > FrequencyOverride)
+			if(FrequencyOverrides.TryGetValue(NodeToDo, out FrequencyOverride) && Result > FrequencyOverride)
 			{
 				Result = FrequencyOverride;
 			}
 
-            foreach (NodeInfo Dep in NodeToDo.Node.FullNamesOfDependencies.Select(x => GUBPNodes[x]))
+            foreach (NodeInfo Dep in NodeToDo.Dependencies)
             {
                 Result = Math.Max(ComputeDependentCISFrequencyQuantumShift(Dep, FrequencyOverrides), Result);
             }
-            foreach (NodeInfo Dep in NodeToDo.Node.FullNamesOfPseudosependencies.Select(x => GUBPNodes[x]))
+            foreach (NodeInfo Dep in NodeToDo.PseudoDependencies)
             {
                 Result = Math.Max(ComputeDependentCISFrequencyQuantumShift(Dep, FrequencyOverrides), Result);
             }
@@ -608,13 +600,13 @@ public partial class GUBP : BuildCommand
             }
             if (bShowDependencies)
             {
-                foreach (string Dep in NodeToDo.Node.FullNamesOfDependencies)
+                foreach (NodeInfo Dep in NodeToDo.Dependencies)
                 {
-                    Log("            dep> {0}", Dep);
+                    Log("            dep> {0}", Dep.Name);
                 }
-                foreach (string Dep in NodeToDo.Node.FullNamesOfPseudosependencies)
+                foreach (NodeInfo Dep in NodeToDo.PseudoDependencies)
                 {
-                    Log("           pdep> {0}", Dep);
+                    Log("           pdep> {0}", Dep.Name);
                 }
 				foreach (string Dep in NodeToDo.Node.CompletedDependencies)
 				{
@@ -653,16 +645,16 @@ public partial class GUBP : BuildCommand
     {
         List<GraphNode> GraphNodes = new List<GraphNode>();
 
-        Dictionary<string, GraphNode> NodeToGraphNodeMap = new Dictionary<string, GraphNode>();
+        Dictionary<NodeInfo, GraphNode> NodeToGraphNodeMap = new Dictionary<NodeInfo, GraphNode>();
 
         for (int NodeIndex = 0; NodeIndex < Nodes.Count; ++NodeIndex)
         {
-            string Node = Nodes[NodeIndex].Name;
+            NodeInfo Node = Nodes[NodeIndex];
 
             GraphNode GraphNode = new GraphNode()
             {
                 Id = GraphNodes.Count,
-                Label = Node
+                Label = Node.Name
             };
             GraphNodes.Add(GraphNode);
             NodeToGraphNodeMap.Add(Node, GraphNode);
@@ -673,10 +665,10 @@ public partial class GUBP : BuildCommand
 
         for (int NodeIndex = 0; NodeIndex < Nodes.Count; ++NodeIndex)
         {
-            string Node = Nodes[NodeIndex].Name;
+            NodeInfo Node = Nodes[NodeIndex];
             GraphNode NodeGraphNode = NodeToGraphNodeMap[Node];
 
-            foreach (string Dep in GUBPNodes[Node].Node.FullNamesOfDependencies)
+            foreach (NodeInfo Dep in Node.Dependencies)
             {
                 GraphNode PrerequisiteFileGraphNode;
                 if (NodeToGraphNodeMap.TryGetValue(Dep, out PrerequisiteFileGraphNode))
@@ -694,7 +686,7 @@ public partial class GUBP : BuildCommand
                 }
 
             }
-            foreach (string Dep in GUBPNodes[Node].Node.FullNamesOfPseudosependencies)
+            foreach (NodeInfo Dep in Node.PseudoDependencies)
             {
                 GraphNode PrerequisiteFileGraphNode;
                 if (NodeToGraphNodeMap.TryGetValue(Dep, out PrerequisiteFileGraphNode))
@@ -853,13 +845,9 @@ public partial class GUBP : BuildCommand
                         ExaminedAgentGroups.Add(NodeToDo.Node.AgentSharingGroup);
                         foreach (NodeInfo ChainNode in SortedAgentGroupChains[NodeToDo.Node.AgentSharingGroup])
                         {
-                            foreach (string Dep in ChainNode.Node.FullNamesOfDependencies)
+                            foreach (NodeInfo Dep in ChainNode.Dependencies)
                             {
-                                if (!GUBPNodes.ContainsKey(Dep))
-                                {
-                                    throw new AutomationException("Dependency {0} node found.", Dep);
-                                }
-                                if (!SortedAgentGroupChains[NodeToDo.Node.AgentSharingGroup].Contains(GUBPNodes[Dep]) && NodesToDo.Contains(GUBPNodes[Dep]))
+                                if (!SortedAgentGroupChains[NodeToDo.Node.AgentSharingGroup].Contains(Dep) && NodesToDo.Contains(Dep))
                                 {
                                     bReady = false;
                                     break;
@@ -870,13 +858,9 @@ public partial class GUBP : BuildCommand
                                 NonReadyAgentGroups.Add(NodeToDo.Node.AgentSharingGroup);
                                 break;
                             }
-                            foreach (string Dep in ChainNode.Node.FullNamesOfPseudosependencies)
+                            foreach (NodeInfo Dep in ChainNode.PseudoDependencies)
                             {
-                                if (!GUBPNodes.ContainsKey(Dep))
-                                {
-                                    throw new AutomationException("Pseudodependency {0} node found.", Dep);
-                                }
-                                if (!SortedAgentGroupChains[NodeToDo.Node.AgentSharingGroup].Contains(GUBPNodes[Dep]) && NodesToDo.Contains(GUBPNodes[Dep]))
+                                if (!SortedAgentGroupChains[NodeToDo.Node.AgentSharingGroup].Contains(Dep) && NodesToDo.Contains(Dep))
                                 {
                                     bPseudoReady = false;
                                     NonPeudoReadyAgentGroups.Add(NodeToDo.Node.AgentSharingGroup);
@@ -888,25 +872,17 @@ public partial class GUBP : BuildCommand
                 }
                 else
                 {
-                    foreach (string Dep in NodeToDo.Node.FullNamesOfDependencies)
+                    foreach (NodeInfo Dep in NodeToDo.Dependencies)
                     {
-                        if (!GUBPNodes.ContainsKey(Dep))
-                        {
-                            throw new AutomationException("Dependency {0} node found.", Dep);
-                        }
-                        if (NodesToDo.Contains(GUBPNodes[Dep]))
+                        if (NodesToDo.Contains(Dep))
                         {
                             bReady = false;
                             break;
                         }
                     }
-                    foreach (string Dep in NodeToDo.Node.FullNamesOfPseudosependencies)
+                    foreach (NodeInfo Dep in NodeToDo.PseudoDependencies)
                     {
-                        if (!GUBPNodes.ContainsKey(Dep))
-                        {
-                            throw new AutomationException("Pseudodependency {0} node found.", Dep);
-                        }
-                        if (NodesToDo.Contains(GUBPNodes[Dep]))
+                        if (NodesToDo.Contains(Dep))
                         {
                             bPseudoReady = false;
                             break;
@@ -1011,30 +987,30 @@ public partial class GUBP : BuildCommand
                     {
                         foreach (NodeInfo ChainNode in SortedAgentGroupChains[NodeToDo.Node.AgentSharingGroup])
                         {
-                            foreach (string Dep in GUBPNodes[ChainNode.Name].Node.FullNamesOfDependencies)
+                            foreach (NodeInfo Dep in GUBPNodes[ChainNode.Name].Dependencies)
                             {
-                                if (!SortedAgentGroupChains[NodeToDo.Node.AgentSharingGroup].Contains(GUBPNodes[Dep]) && NodesToDo.Contains(GUBPNodes[Dep]))
+                                if (!SortedAgentGroupChains[NodeToDo.Node.AgentSharingGroup].Contains(Dep) && NodesToDo.Contains(Dep))
                                 {
-                                    Deps = Deps + Dep + "[" + ChainNode.Name + "->" + NodeToDo.Node.AgentSharingGroup + "]" + " ";
+                                    Deps = Deps + Dep.Name + "[" + ChainNode.Name + "->" + NodeToDo.Node.AgentSharingGroup + "]" + " ";
                                 }
                             }
                         }
                     }
-                    foreach (string Dep in NodeToDo.Node.FullNamesOfDependencies)
+                    foreach (NodeInfo Dep in NodeToDo.Dependencies)
                     {
-                        if (NodesToDo.Contains(GUBPNodes[Dep]))
+                        if (NodesToDo.Contains(Dep))
                         {
-                            Deps = Deps + Dep + " ";
+                            Deps = Deps + Dep.Name + " ";
                         }
                     }
-                    foreach (string Dep in NodeToDo.Node.FullNamesOfPseudosependencies)
+                    foreach (NodeInfo Dep in NodeToDo.PseudoDependencies)
                     {
-                        if (NodesToDo.Contains(GUBPNodes[Dep]))
+                        if (NodesToDo.Contains(Dep))
                         {
-                            Deps = Deps + Dep + " ";
+                            Deps = Deps + Dep.Name + " ";
                         }
                     }
-                    Log("  {0}    deps: {1}", NodeToDo, Deps);
+                    Log("  {0}    deps: {1}", NodeToDo.Name, Deps);
                 }
                 throw new AutomationException("Cycle in GUBP");
             }
@@ -1564,7 +1540,7 @@ public partial class GUBP : BuildCommand
 
 		AddNodesForBranch(TimeIndex, bNoAutomatedTesting);
 
-		ValidateDependencyNames();
+		LinkGraph(GUBPNodes);
 
         if (bCleanLocalTempStorage)  // shared temp storage can never be wiped
         {
@@ -1603,7 +1579,7 @@ public partial class GUBP : BuildCommand
             }
 
 			// Compute all the frequencies
-			Dictionary<string, int> FrequencyOverrides = ApplyFrequencyBarriers();
+			Dictionary<NodeInfo, int> FrequencyOverrides = ApplyFrequencyBarriers();
             foreach (KeyValuePair<string, NodeInfo> NodeToDo in GUBPNodes)
             {
                 ComputeDependentCISFrequencyQuantumShift(NodeToDo.Value, FrequencyOverrides);
@@ -1751,37 +1727,97 @@ public partial class GUBP : BuildCommand
 		ExecuteNodes(OrdereredToDo, bOnlyNode, bFakeEC, LocalOnly, bSaveSharedTempStorage, GUBPNodesCompleted, GUBPNodesHistory, CLString, FakeFail);
 	}
 
-	private Dictionary<string, int> ApplyFrequencyBarriers()
+	private static void LinkGraph(Dictionary<string, NodeInfo> GUBPNodes)
+	{
+		int NumErrors = 0;
+		foreach (NodeInfo NodeToDo in GUBPNodes.Values)
+		{
+			// Find all the dependencies
+			NodeToDo.Dependencies = new List<NodeInfo>();
+			foreach (string DependencyName in NodeToDo.Node.FullNamesOfDependencies)
+			{
+				NodeInfo Dependency;
+				if (!GUBPNodes.TryGetValue(DependencyName, out Dependency))
+				{
+					CommandUtils.LogError("Node {0} is not in the graph. It is a dependency of {1}.", DependencyName, NodeToDo.Name);
+					NumErrors++;
+				}
+				else if(Dependency == NodeToDo)
+				{
+					CommandUtils.LogError("Node {0} has a self arc.", NodeToDo.Name);
+					NumErrors++;
+				}
+				else
+				{
+					NodeToDo.Dependencies.Add(Dependency);
+				}
+			}
+
+			// Find all the pseudo-dependencies
+			NodeToDo.PseudoDependencies = new List<NodeInfo>();
+			foreach (string PseudoDependencyName in NodeToDo.Node.FullNamesOfPseudosependencies)
+			{
+				NodeInfo PseudoDependency;
+				if (!GUBPNodes.TryGetValue(PseudoDependencyName, out PseudoDependency))
+				{
+					CommandUtils.LogError("Node {0} is not in the graph. It is a pseudodependency of {1}.", PseudoDependencyName, NodeToDo.Name);
+					NumErrors++;
+				}
+				else if(PseudoDependency == NodeToDo)
+				{
+					CommandUtils.LogError("Node {0} has a self pseudoarc.", NodeToDo.Name);
+					NumErrors++;
+				}
+				else
+				{
+					NodeToDo.PseudoDependencies.Add(PseudoDependency);
+				}
+			}
+		}
+		if (NumErrors > 0)
+		{
+			throw new AutomationException("Failed to link graph ({0} errors).", NumErrors);
+		}
+	}
+
+	private Dictionary<NodeInfo, int> ApplyFrequencyBarriers()
 	{
 		// Make sure that everything that's listed as a frequency barrier is completed with the given interval
-		Dictionary<string, int> FrequencyOverrides = new Dictionary<string, int>();
+		Dictionary<NodeInfo, int> FrequencyOverrides = new Dictionary<NodeInfo, int>();
 		foreach (KeyValuePair<string, sbyte> Barrier in BranchOptions.FrequencyBarriers)
 		{
+			// Find the barrier node
+			NodeInfo BarrierNode;
+			if(!GUBPNodes.TryGetValue(Barrier.Key, out BarrierNode))
+			{
+				throw new AutomationException("Couldn't find node '{0}' for frequency barrier", Barrier.Key);
+			}
+
 			// All the nodes which are dependencies of the barrier node
-			HashSet<string> IncludedNodes = new HashSet<string> { Barrier.Key };
+			HashSet<NodeInfo> IncludedNodes = new HashSet<NodeInfo> { BarrierNode };
 
 			// Find all the nodes which are indirect dependencies of this node
-			List<string> SearchNodes = new List<string> { Barrier.Key };
+			List<NodeInfo> SearchNodes = new List<NodeInfo> { BarrierNode };
 			for (int Idx = 0; Idx < SearchNodes.Count; Idx++)
 			{
-				GUBPNode Node = GUBPNodes[SearchNodes[Idx]].Node;
-				foreach (string DependencyName in Node.FullNamesOfDependencies.Union(Node.FullNamesOfPseudosependencies))
+				NodeInfo Node = SearchNodes[Idx];
+				foreach (NodeInfo Dependency in Node.Dependencies.Union(Node.PseudoDependencies))
 				{
-					if (!IncludedNodes.Contains(DependencyName))
+					if (!IncludedNodes.Contains(Dependency))
 					{
-						IncludedNodes.Add(DependencyName);
-						SearchNodes.Add(DependencyName);
+						IncludedNodes.Add(Dependency);
+						SearchNodes.Add(Dependency);
 					}
 				}
 			}
 
 			// Make sure that everything included in this list is before the cap, and everything not in the list is after it
-			foreach (KeyValuePair<string, NodeInfo> NodePair in GUBPNodes)
+			foreach (NodeInfo Node in GUBPNodes.Values)
 			{
-				if (IncludedNodes.Contains(NodePair.Key))
+				if (IncludedNodes.Contains(Node))
 				{
 					int Frequency;
-					if (FrequencyOverrides.TryGetValue(NodePair.Key, out Frequency))
+					if (FrequencyOverrides.TryGetValue(Node, out Frequency))
 					{
 						Frequency = Math.Min(Frequency, Barrier.Value);
 					}
@@ -1789,7 +1825,7 @@ public partial class GUBP : BuildCommand
 					{
 						Frequency = Barrier.Value;
 					}
-					FrequencyOverrides[NodePair.Key] = Frequency;
+					FrequencyOverrides[Node] = Frequency;
 				}
 			}
 		}
@@ -1948,29 +1984,18 @@ public partial class GUBP : BuildCommand
 			{
 				bDoneWithDependencies = true;
 				HashSet<NodeInfo> Fringe = new HashSet<NodeInfo>();
-				foreach (KeyValuePair<string, NodeInfo> NodeToDo in GUBPNodes)
+				foreach (NodeInfo NodeToDo in GUBPNodes.Values)
 				{
-					if (!NodesToDo.Contains(NodeToDo.Value))
+					if (!NodesToDo.Contains(NodeToDo))
 					{
-						foreach (string Dep in GUBPNodes[NodeToDo.Key].Node.FullNamesOfDependencies)
+						foreach (NodeInfo Dep in NodeToDo.Dependencies)
 						{
-							if (!GUBPNodes.ContainsKey(Dep))
+							if (NodesToDo.Contains(Dep))
 							{
-								throw new AutomationException("Node {0} is not in the graph. It is a dependency of {1}.", Dep, NodeToDo.Key);
-							}
-							if (NodesToDo.Contains(GUBPNodes[Dep]))
-							{
-								Fringe.Add(NodeToDo.Value);
+								Fringe.Add(NodeToDo);
 								bDoneWithDependencies = false;
 							}
 						}
-						foreach (string Dep in GUBPNodes[NodeToDo.Key].Node.FullNamesOfPseudosependencies)
-						{
-							if (!GUBPNodes.ContainsKey(Dep))
-							{
-								throw new AutomationException("Node {0} is not in the graph. It is a pseudodependency of {1}.", Dep, NodeToDo.Key);
-							}
-						}
 					}
 				}
 				NodesToDo.UnionWith(Fringe);
@@ -2002,25 +2027,14 @@ public partial class GUBP : BuildCommand
 						}
 						BranchOptions.NodesToRemovePseudoDependencies = ListOfRetainedNodesToRemovePseudoDependencies;
 					}
-					foreach (string Dep in NodeToDo.Node.FullNamesOfDependencies)
+					foreach (NodeInfo Dep in NodeToDo.Dependencies)
 					{
-						if (!GUBPNodes.ContainsKey(Dep))
-						{
-							throw new AutomationException("Node {0} is not in the graph. It is a dependency of {1}.", Dep, NodeToDo.Name);
-						}
-						if (!NodesToDo.Contains(GUBPNodes[Dep]))
+						if (!NodesToDo.Contains(Dep))
 						{
-							Fringe.Add(GUBPNodes[Dep]);
+							Fringe.Add(Dep);
 							bDoneWithDependencies = false;
 						}
 					}
-					foreach (string Dep in NodeToDo.Node.FullNamesOfPseudosependencies)
-					{
-						if (!GUBPNodes.ContainsKey(Dep))
-						{
-							throw new AutomationException("Node {0} is not in the graph. It is a pseudodependency of {1}.", Dep, NodeToDo.Name);
-						}
-					}
 				}
 				NodesToDo.UnionWith(Fringe);
 			}
@@ -2143,35 +2157,6 @@ public partial class GUBP : BuildCommand
 		PrintCSVFile(String.Format("UAT,GetDependencies,{0},{1}", DependencyStart, DependencyFinish));
 	}
 
-	private void ValidateDependencyNames()
-	{
-		foreach (KeyValuePair<string, NodeInfo> NodeToDo in GUBPNodes)
-		{
-			foreach (string Dep in GUBPNodes[NodeToDo.Key].Node.FullNamesOfDependencies)
-			{
-				if (!GUBPNodes.ContainsKey(Dep))
-				{
-					throw new AutomationException("Node {0} is not in the full graph. It is a dependency of {1}.", Dep, NodeToDo.Key);
-				}
-				if (Dep == NodeToDo.Key)
-				{
-					throw new AutomationException("Node {0} has a self arc.", NodeToDo.Key);
-				}
-			}
-			foreach (string Dep in GUBPNodes[NodeToDo.Key].Node.FullNamesOfPseudosependencies)
-			{
-				if (!GUBPNodes.ContainsKey(Dep))
-				{
-					throw new AutomationException("Node {0} is not in the full graph. It is a pseudodependency of {1}.", Dep, NodeToDo.Key);
-				}
-				if (Dep == NodeToDo.Key)
-				{
-					throw new AutomationException("Node {0} has a self pseudoarc.", NodeToDo.Key);
-				}
-			}
-		}
-	}
-
 	private int UpdateCISCounter()
 	{
 		if (!P4Enabled)
@@ -2323,7 +2308,7 @@ public partial class GUBP : BuildCommand
 				bHasTests = true;
 			}
 			int MyIndex = OrdereredToDo.IndexOf(NodeToDo);
-			foreach (NodeInfo Dep in NodeToDo.Node.FullNamesOfDependencies.Select(x => GUBPNodes[x]))
+			foreach (NodeInfo Dep in NodeToDo.Dependencies)
 			{
 				int DepIndex = OrdereredToDo.IndexOf(Dep);
 				if (DepIndex >= MyIndex)
@@ -2331,7 +2316,7 @@ public partial class GUBP : BuildCommand
 					throw new AutomationException("Topological sort error, node {0} has a dependency of {1} which sorted after it.", NodeToDo.Name, Dep.Name);
 				}
 			}
-			foreach (NodeInfo Dep in NodeToDo.Node.FullNamesOfPseudosependencies.Select(x => GUBPNodes[x]))
+			foreach (NodeInfo Dep in NodeToDo.PseudoDependencies)
 			{
 				int DepIndex = OrdereredToDo.IndexOf(Dep);
 				if (DepIndex >= MyIndex)
@@ -2859,7 +2844,7 @@ public partial class GUBP : BuildCommand
 
             NodeToDo.Node.AllDependencyBuildProducts = new List<string>();
             NodeToDo.Node.AllDependencies = new List<string>();
-            foreach (NodeInfo Dep in NodeToDo.Node.FullNamesOfDependencies.Select(x => GUBPNodes[x]))
+            foreach (NodeInfo Dep in NodeToDo.Dependencies)
             {
                 NodeToDo.Node.AddAllDependent(Dep.Name);
                 if (Dep.Node.AllDependencies == null)
