Remove 'completed only' dependencies. I think they were a hack, and nothing is using them any more.

diff --git a/Engine/Source/Programs/AutomationTool/GUBP/GUBP.cs b/Engine/Source/Programs/AutomationTool/GUBP/GUBP.cs
--- a/Engine/Source/Programs/AutomationTool/GUBP/GUBP.cs
+++ b/Engine/Source/Programs/AutomationTool/GUBP/GUBP.cs
@@ -151,33 +151,7 @@ public partial class GUBP : BuildCommand
         return Result;
     }
 
-	List<NodeInfo> GetCompletedOnlyDependencies(NodeInfo NodeToDo, bool bFlat = false, bool ECOnly = true)
-	{
-		List<NodeInfo> Result = new List<NodeInfo>();
-		foreach (NodeInfo CompletedDependency in NodeToDo.Node.CompletedDependencies.Select(x => GUBPNodes[x]))
-		{
-			bool Usable = CompletedDependency.Node.RunInEC() || !ECOnly;
-			if(Usable)
-			{
-				if(!Result.Contains(CompletedDependency))
-				{
-					Result.Add(CompletedDependency);
-				}
-			}
-			if (bFlat || !Usable)
-			{
-				foreach (NodeInfo RNode in GetCompletedOnlyDependencies(CompletedDependency, bFlat, ECOnly))
-				{
-					if (!Result.Contains(RNode))
-					{
-						Result.Add(RNode);
-					}
-				}
-			}
-		}
-		return Result;
-	}
-    List<NodeInfo> GetECDependencies(NodeInfo NodeToDo, bool bFlat = false)
+	List<NodeInfo> GetECDependencies(NodeInfo NodeToDo, bool bFlat = false)
     {
         return GetDependencies(NodeToDo, bFlat, true);
     }
@@ -317,10 +291,6 @@ public partial class GUBP : BuildCommand
             {
                 Result = Math.Max(ComputeDependentCISFrequencyQuantumShift(Dep, FrequencyOverrides), Result);
             }
-			foreach (NodeInfo Dep in NodeToDo.Node.CompletedDependencies.Select(x => GUBPNodes[x]))
-			{
-				Result = Math.Max(ComputeDependentCISFrequencyQuantumShift(Dep, FrequencyOverrides), Result);
-			}
             if (Result < 0)
             {
                 throw new AutomationException("Failed to compute shift.");
@@ -608,10 +578,6 @@ public partial class GUBP : BuildCommand
                 {
                     Log("           pdep> {0}", Dep.Name);
                 }
-				foreach (string Dep in NodeToDo.Node.CompletedDependencies)
-				{
-					Log("			cdep>{0}", Dep);
-				}
             }
             if (bShowECDependencies)
             {
@@ -619,10 +585,6 @@ public partial class GUBP : BuildCommand
                 {
                     Log("           {0}", Dep.Name);
                 }
-				foreach (NodeInfo Dep in GetCompletedOnlyDependencies(NodeToDo))
-				{
-					Log("           compDep> {0}", Dep.Name);
-				}
             }
 
 			if(bShowDependednOn)
@@ -888,18 +850,6 @@ public partial class GUBP : BuildCommand
                             break;
                         }
                     }
-					foreach (string Dep in NodeToDo.Node.CompletedDependencies)
-					{
-						if (!GUBPNodes.ContainsKey(Dep))
-						{
-							throw new AutomationException("Completed Dependency {0} node found.", Dep);
-						}
-						if (NodesToDo.Contains(GUBPNodes[Dep]))
-						{
-							bCompReady = false;
-							break;
-						}
-					}
                 }
                 float Priority = NodeToDo.Node.Priority();
 
@@ -2527,7 +2477,6 @@ public partial class GUBP : BuildCommand
 				string PreCondition = "";
 				string RunCondition = "";
 				List<string> UncompletedEcDeps = new List<string>();
-				List<string> UncompletedCompletedDeps = new List<string>();
 				List<string> PreConditionUncompletedEcDeps = new List<string>();
 				string MyAgentGroup = NodeToDo.Node.AgentSharingGroup;
 				bool bDoNestedJobstep = false;
@@ -2557,18 +2506,6 @@ public partial class GUBP : BuildCommand
 				{
 					PreConditionUncompletedEcDeps.Add(Dep);
 				}
-				List<NodeInfo> CompletedDeps = GetCompletedOnlyDependencies(NodeToDo);
-				foreach (NodeInfo Dep in CompletedDeps)
-				{
-					if (Dep.Node.RunInEC() && !NodeIsAlreadyComplete(GUBPNodesCompleted, Dep, LocalOnly) && OrdereredToDo.Contains(Dep)) // if something is already finished, we don't put it into EC
-					{
-						if (OrdereredToDo.IndexOf(Dep) > OrdereredToDo.IndexOf(NodeToDo))
-						{
-							throw new AutomationException("Topological sort error, node {0} has a dependency of {1} which sorted after it.", NodeToDo.Name, Dep.Name);
-						}
-						UncompletedCompletedDeps.Add(Dep.Name);
-					}
-				}
 				if (MyAgentGroup != "")
 				{
 					bDoNestedJobstep = true;
@@ -2656,25 +2593,6 @@ public partial class GUBP : BuildCommand
 					}
 					PreCondition = PreCondition + ") true;]\"";
 				}
-				if (UncompletedCompletedDeps.Count > 0)
-				{
-					PreCondition = "\"\\$\" . \"[/javascript if(";
-					// these run "parallel", but we add preconditions to serialize them
-					int Index = 0;
-					foreach (NodeInfo Dep in CompletedDeps)
-					{
-						if (Dep.Node.RunInEC())
-						{
-							PreCondition = PreCondition + "getProperty('" + GetJobStep(PreconditionParentPath, Dep.Name) + "/status\') == \'completed\'";
-							Index++;
-							if (Index != CompletedDeps.Count)
-							{
-								PreCondition = PreCondition + " && ";
-							}
-						}
-					}
-					PreCondition = PreCondition + ") true;]\"";
-				}
 
 				if (UncompletedEcDeps.Count > 0)
 				{
diff --git a/Engine/Source/Programs/AutomationTool/GUBP/LegacyNodes.cs b/Engine/Source/Programs/AutomationTool/GUBP/LegacyNodes.cs
--- a/Engine/Source/Programs/AutomationTool/GUBP/LegacyNodes.cs
+++ b/Engine/Source/Programs/AutomationTool/GUBP/LegacyNodes.cs
@@ -21,7 +21,6 @@ partial class GUBP
 		public List<string> DependentPromotions = new List<string>();
         public List<string> AllDependencyBuildProducts = null;
         public List<string> AllDependencies = null;		
-		public List<string> CompletedDependencies = new List<string>();
         public string AgentSharingGroup = "";
         public int ComputedDependentCISFrequencyQuantumShift = -1;
 
@@ -203,13 +202,6 @@ partial class GUBP
                 FullNamesOfPseudosependencies.Add(Node);
             }
         }
-		public void AddCompletedDependency(string Node)
-		{
-			if(!CompletedDependencies.Contains(Node))
-			{
-				CompletedDependencies.Add(Node);
-			}
-		}
         public void RemovePseudodependency(string Node)
         {
             if (FullNamesOfPseudosependencies.Contains(Node))
