PrintCSVFile detects a file of "nul" and ignores output.

diff --git a/Engine/Source/Programs/AutomationTool/AutomationUtils/CommandUtils.cs b/Engine/Source/Programs/AutomationTool/AutomationUtils/CommandUtils.cs
--- a/Engine/Source/Programs/AutomationTool/AutomationUtils/CommandUtils.cs
+++ b/Engine/Source/Programs/AutomationTool/AutomationUtils/CommandUtils.cs
@@ -1898,7 +1898,8 @@ namespace AutomationTool
 		private static UnrealBuildTool.UnrealTargetPlatform[] UBTTargetPlatforms; 
         public static void PrintCSVFile(string Input)
         {
-            if (IsBuildMachine && CmdEnv.CSVFile != "")
+            // sometimes nul is used as the CSV file so batch scripts can append using output redirectors easily. We have to ignore those.
+            if (IsBuildMachine && CmdEnv.CSVFile != "" && CmdEnv.CSVFile != "nul")
             {
                 var CSVBuilder = new StringBuilder();
                 var CSVLineToAppend = String.Format("{0}{1}", Input, Environment.NewLine);
