Move the controlling trigger and controlling trigger path into NodeInfo.

diff --git a/Engine/Source/Programs/AutomationTool/GUBP/GUBP.cs b/Engine/Source/Programs/AutomationTool/GUBP/GUBP.cs
--- a/Engine/Source/Programs/AutomationTool/GUBP/GUBP.cs
+++ b/Engine/Source/Programs/AutomationTool/GUBP/GUBP.cs
@@ -36,6 +36,8 @@ public partial class GUBP : BuildCommand
 		public GUBPNode Node;
 		public List<NodeInfo> Dependencies;
 		public List<NodeInfo> PseudoDependencies;
+		public NodeInfo ControllingTrigger;
+		public string ControllingTriggerDotName;
 
 		public override string ToString()
 		{
@@ -161,89 +163,58 @@ public partial class GUBP : BuildCommand
         return Deps.Contains(Rootward);
     }
 
-    NodeInfo GetControllingTrigger(NodeInfo NodeToDo, Dictionary<NodeInfo, NodeInfo> GUBPNodesControllingTrigger)
+    static void FindControllingTriggers(NodeInfo NodeToDo)
     {
-        if (GUBPNodesControllingTrigger.ContainsKey(NodeToDo))
-        {
-            return GUBPNodesControllingTrigger[NodeToDo];
-        }
-        NodeInfo Result = null;
-        foreach (NodeInfo Dependency in NodeToDo.Dependencies)
+		if(NodeToDo.ControllingTriggerDotName == null)
 		{
-            bool IsTrigger = Dependency.Node.TriggerNode();
-            if (IsTrigger)
-            {
-                if (Dependency != Result && Result != null)
-                {
-                    throw new AutomationException("Node {0} has two controlling triggers {1} and {2}.", NodeToDo.Name, Dependency.Name, Result);
-                }
-                Result = Dependency;
-            }
-            else
-            {
-                NodeInfo NewResult = GetControllingTrigger(Dependency, GUBPNodesControllingTrigger);
-                if (NewResult != null)
-                {
-                    if (NewResult != Result && Result != null)
-                    {
-                        throw new AutomationException("Node {0} has two controlling triggers {1} and {2}.", NodeToDo.Name, NewResult.Name, Result);
-                    }
-                    Result = NewResult;
-                }
-            }
-        }
-        foreach (NodeInfo PseudoDependency in NodeToDo.PseudoDependencies)
-        {
-            bool IsTrigger = PseudoDependency.Node.TriggerNode();
-            if (IsTrigger)
-            {
-                if (PseudoDependency != Result && Result != null)
-                {
-                    throw new AutomationException("Node {0} has two controlling triggers {1} and {2}.", NodeToDo.Name, PseudoDependency.Name, Result);
-                }
-                Result = PseudoDependency;
-            }
-            else
-            {
-                NodeInfo NewResult = GetControllingTrigger(PseudoDependency, GUBPNodesControllingTrigger);
-                if (NewResult != null)
-                {
-                    if (NewResult != Result && Result != null)
-                    {
-                        throw new AutomationException("Node {0} has two controlling triggers {1} and {2}.", NodeToDo.Name, NewResult.Name, Result);
-                    }
-                    Result = NewResult;
-                }
-            }
-        }
-        GUBPNodesControllingTrigger.Add(NodeToDo, Result);
-        return Result;
-    }
-    string GetControllingTriggerDotName(NodeInfo NodeToDo, Dictionary<NodeInfo, string> GUBPNodesControllingTriggerDotName, Dictionary<NodeInfo, NodeInfo> GUBPNodesControllingTrigger)
-    {
-        if (GUBPNodesControllingTriggerDotName.ContainsKey(NodeToDo))
-        {
-            return GUBPNodesControllingTriggerDotName[NodeToDo];
-        }
-        string Result = "";
-        NodeInfo WorkingNode = NodeToDo;
-        while (true)
-        {
-            NodeInfo ThisResult = GetControllingTrigger(WorkingNode, GUBPNodesControllingTrigger);
-            if (ThisResult == null)
-            {
-                break;
-            }
-            if (Result != "")
-            {
-                Result = "." + Result;
-            }
-            Result = ThisResult.Name + Result;
-            WorkingNode = ThisResult;
-        }
-        GUBPNodesControllingTriggerDotName.Add(NodeToDo, Result);
-        return Result;
-    }
+			// Find all the dependencies of this node
+			List<NodeInfo> AllDependencies = new List<NodeInfo>();
+			AllDependencies.AddRange(NodeToDo.Dependencies);
+			AllDependencies.AddRange(NodeToDo.PseudoDependencies);
+
+			// Find the immediate trigger controlling this one
+			foreach (NodeInfo Dependency in AllDependencies)
+			{
+				FindControllingTriggers(Dependency);
+
+				NodeInfo Trigger = null;
+				if(Dependency.Node.TriggerNode())
+				{
+					Trigger = Dependency;
+				}
+				else
+				{
+					Trigger = Dependency.ControllingTrigger;
+				}
+				
+				if(Trigger != null)
+				{
+					if(NodeToDo.ControllingTrigger == null)
+					{
+						NodeToDo.ControllingTrigger = Trigger;
+					}
+					else if(NodeToDo.ControllingTrigger != Trigger)
+					{
+						throw new AutomationException("Node {0} has two controlling triggers {1} and {2}.", NodeToDo.Name, NodeToDo.ControllingTrigger.Name, Trigger.Name);
+					}
+				}
+			}
+
+			// Set the path of controlling triggers for this node
+			if(NodeToDo.ControllingTrigger == null)
+			{
+				NodeToDo.ControllingTriggerDotName = "";
+			}
+			else if(NodeToDo.ControllingTrigger.ControllingTriggerDotName.Length == 0)
+			{
+				NodeToDo.ControllingTriggerDotName = NodeToDo.ControllingTrigger.Name;
+			}
+			else
+			{
+				NodeToDo.ControllingTriggerDotName = NodeToDo.ControllingTrigger.ControllingTriggerDotName + "." + NodeToDo.ControllingTrigger.Name;
+			}
+		}
+	}
 
     string CISFrequencyQuantumShiftString(NodeInfo NodeToDo)
     {
@@ -466,7 +437,7 @@ public partial class GUBP : BuildCommand
         Log("Took {0}s to get P4 history", BuildDuration / 1000);
 
     }
-    void PrintNodes(GUBP bp, List<NodeInfo> Nodes, Dictionary<string, bool> GUBPNodesCompleted, Dictionary<string, NodeHistory> GUBPNodesHistory, Dictionary<NodeInfo, string> GUBPNodesControllingTriggerDotName, Dictionary<NodeInfo, NodeInfo> GUBPNodesControllingTrigger, bool LocalOnly, List<NodeInfo> UnfinishedTriggers = null)
+    void PrintNodes(GUBP bp, List<NodeInfo> Nodes, Dictionary<string, bool> GUBPNodesCompleted, Dictionary<string, NodeHistory> GUBPNodesHistory, bool LocalOnly, List<NodeInfo> UnfinishedTriggers = null)
     {
         bool bShowAllChanges = bp.ParseParam("AllChanges") && GUBPNodesHistory != null;
         bool bShowChanges = (bp.ParseParam("Changes") && GUBPNodesHistory != null) || bShowAllChanges;
@@ -495,7 +466,7 @@ public partial class GUBP : BuildCommand
             }
             if (bShowTriggers)
             {
-                string MyControllingTrigger = GetControllingTriggerDotName(NodeToDo, GUBPNodesControllingTriggerDotName, GUBPNodesControllingTrigger);
+                string MyControllingTrigger = NodeToDo.ControllingTriggerDotName;
                 if (MyControllingTrigger != LastControllingTrigger)
                 {
                     LastControllingTrigger = MyControllingTrigger;
@@ -504,7 +475,7 @@ public partial class GUBP : BuildCommand
                         string Finished = "";
                         if (UnfinishedTriggers != null)
                         {
-                            NodeInfo MyShortControllingTrigger = GetControllingTrigger(NodeToDo, GUBPNodesControllingTrigger);
+                            NodeInfo MyShortControllingTrigger = NodeToDo.ControllingTrigger;
                             if (UnfinishedTriggers.Contains(MyShortControllingTrigger))
                             {
                                 Finished = "(not yet triggered)";
@@ -748,7 +719,7 @@ public partial class GUBP : BuildCommand
         return Result;
     }
 
-    List<NodeInfo> TopologicalSort(HashSet<NodeInfo> NodesToDo, Dictionary<string, bool> GUBPNodesCompleted, Dictionary<NodeInfo, string> GUBPNodesControllingTriggerDotName, Dictionary<NodeInfo, NodeInfo> GUBPNodesControllingTrigger, NodeInfo ExplicitTrigger = null, bool LocalOnly = false, bool SubSort = false, bool DoNotConsiderCompletion = false)
+    List<NodeInfo> TopologicalSort(HashSet<NodeInfo> NodesToDo, Dictionary<string, bool> GUBPNodesCompleted, NodeInfo ExplicitTrigger = null, bool LocalOnly = false, bool SubSort = false, bool DoNotConsiderCompletion = false)
     {
         DateTime StartTime = DateTime.UtcNow;
 
@@ -775,7 +746,7 @@ public partial class GUBP : BuildCommand
             }
             foreach (KeyValuePair<string, List<NodeInfo>> Chain in AgentGroupChains)
             {
-                SortedAgentGroupChains.Add(Chain.Key, TopologicalSort(new HashSet<NodeInfo>(Chain.Value), GUBPNodesCompleted, GUBPNodesControllingTriggerDotName, GUBPNodesControllingTrigger, ExplicitTrigger, LocalOnly, true, DoNotConsiderCompletion));
+                SortedAgentGroupChains.Add(Chain.Key, TopologicalSort(new HashSet<NodeInfo>(Chain.Value), GUBPNodesCompleted, ExplicitTrigger, LocalOnly, true, DoNotConsiderCompletion));
             }
             Log("***************Done with recursion");
         }
@@ -855,11 +826,11 @@ public partial class GUBP : BuildCommand
 
                 if (bReady && BestNode != null)
                 {
-                    if (String.Compare(GetControllingTriggerDotName(BestNode, GUBPNodesControllingTriggerDotName, GUBPNodesControllingTrigger), GetControllingTriggerDotName(NodeToDo, GUBPNodesControllingTriggerDotName, GUBPNodesControllingTrigger)) < 0) //sorted by controlling trigger
+                    if (String.Compare(BestNode.ControllingTriggerDotName, NodeToDo.ControllingTriggerDotName) < 0) //sorted by controlling trigger
                     {
                         bReady = false;
                     }
-                    else if (String.Compare(GetControllingTriggerDotName(BestNode, GUBPNodesControllingTriggerDotName, GUBPNodesControllingTrigger), GetControllingTriggerDotName(NodeToDo, GUBPNodesControllingTriggerDotName, GUBPNodesControllingTrigger)) == 0) //sorted by controlling trigger
+                    else if (String.Compare(BestNode.ControllingTriggerDotName, NodeToDo.ControllingTriggerDotName) == 0) //sorted by controlling trigger
                     {
                         if (GUBPNodes[BestNode.Name].Node.IsSticky() && !NodeToDo.Node.IsSticky()) //sticky nodes first
                         {
@@ -1492,14 +1463,16 @@ public partial class GUBP : BuildCommand
 
 		LinkGraph(GUBPNodes);
 
+		foreach(NodeInfo NodeToDo in GUBPNodes.Values)
+		{
+			FindControllingTriggers(NodeToDo);
+		}
+
         if (bCleanLocalTempStorage)  // shared temp storage can never be wiped
         {
             TempStorage.DeleteLocalTempStorageManifests(CmdEnv);
         }
 
-		Dictionary<NodeInfo, NodeInfo> GUBPNodesControllingTrigger = new Dictionary<NodeInfo, NodeInfo>();
-        Dictionary<NodeInfo, string> GUBPNodesControllingTriggerDotName = new Dictionary<NodeInfo, string>();
-
 		Dictionary<string, string> FullNodeDependedOnBy = new Dictionary<string, string>();
 		Dictionary<string, string> FullNodeDependentPromotions = new Dictionary<string, string>();		
 		List<string> SeparatePromotables = new List<string>();
@@ -1553,7 +1526,7 @@ public partial class GUBP : BuildCommand
 
         Dictionary<string, string> FullNodeList;
         Dictionary<string, string> FullNodeDirectDependencies;
-		GetFullNodeListAndDirectDependencies(GUBPNodesControllingTrigger, GUBPNodesControllingTriggerDotName, out FullNodeList, out FullNodeDirectDependencies);
+		GetFullNodeListAndDirectDependencies(out FullNodeList, out FullNodeDirectDependencies);
 
 		Dictionary<string, int> FullNodeListSortKey = GetDisplayOrder(FullNodeList.Keys.ToList(), FullNodeDirectDependencies, GUBPNodes);
 
@@ -1603,7 +1576,7 @@ public partial class GUBP : BuildCommand
 			HashSet<NodeInfo> NewNodesToDo = new HashSet<NodeInfo>();
 			foreach (NodeInfo NodeToDo in NodesToDo)
 			{
-				string TriggerDot = GetControllingTriggerDotName(NodeToDo, GUBPNodesControllingTriggerDotName, GUBPNodesControllingTrigger);
+				string TriggerDot = NodeToDo.ControllingTriggerDotName;
 				if (TriggerDot == "" && !NodeToDo.Node.TriggerNode())
 				{
 					LogVerbose("  Keeping {0}", NodeToDo.Name);
@@ -1627,7 +1600,6 @@ public partial class GUBP : BuildCommand
             {
                 LogVerbose("** {0}", NodeToDo.Name);
                 NodeIsAlreadyComplete(GUBPNodesCompleted, NodeToDo, LocalOnly); // cache these now to avoid spam later
-                GetControllingTriggerDotName(NodeToDo, GUBPNodesControllingTriggerDotName, GUBPNodesControllingTrigger);
             }
             double BuildDuration = (DateTime.UtcNow - StartTime).TotalMilliseconds;
 			LogVerbose("Took {0}s to cache completion for {1} nodes", BuildDuration / 1000, NodesToDo.Count);
@@ -1647,19 +1619,19 @@ public partial class GUBP : BuildCommand
             Log("Took {0}s to get history for {1} nodes", BuildDuration / 1000, NodesToDo.Count);
         }*/
 
-        List<NodeInfo> OrdereredToDo = TopologicalSort(NodesToDo, GUBPNodesCompleted, GUBPNodesControllingTriggerDotName, GUBPNodesControllingTrigger, ExplicitTrigger, LocalOnly);
+        List<NodeInfo> OrdereredToDo = TopologicalSort(NodesToDo, GUBPNodesCompleted, ExplicitTrigger, LocalOnly);
 
 		List<NodeInfo> UnfinishedTriggers = FindUnfinishedTriggers(bSkipTriggers, LocalOnly, ExplicitTrigger, GUBPNodesCompleted, OrdereredToDo);
 
         LogVerbose("*********** Desired And Dependent Nodes, in order.");
-        PrintNodes(this, OrdereredToDo, GUBPNodesCompleted, GUBPNodesHistory, GUBPNodesControllingTriggerDotName, GUBPNodesControllingTrigger, LocalOnly, UnfinishedTriggers);		
+        PrintNodes(this, OrdereredToDo, GUBPNodesCompleted, GUBPNodesHistory,  LocalOnly, UnfinishedTriggers);		
         //check sorting
 		CheckSortOrder(LocalOnly, GUBPNodesCompleted, OrdereredToDo);
 
         string FakeFail = ParseParamValue("FakeFail");
         if (CommanderSetup)
         {
-			DoCommanderSetup(TimeIndex, bSkipTriggers, bFakeEC, LocalOnly, CLString, ExplicitTrigger, GUBPNodesControllingTrigger, GUBPNodesControllingTriggerDotName, FullNodeList, FullNodeDirectDependencies, FullNodeDependedOnBy, FullNodeDependentPromotions, SeparatePromotables, FullNodeListSortKey, GUBPNodesCompleted, GUBPNodesHistory, OrdereredToDo, UnfinishedTriggers, FakeFail);
+			DoCommanderSetup(TimeIndex, bSkipTriggers, bFakeEC, LocalOnly, CLString, ExplicitTrigger, FullNodeList, FullNodeDirectDependencies, FullNodeDependedOnBy, FullNodeDependentPromotions, SeparatePromotables, FullNodeListSortKey, GUBPNodesCompleted, GUBPNodesHistory, OrdereredToDo, UnfinishedTriggers, FakeFail);
             Log("Commander setup only, done.");            
             PrintRunTime();
             return;
@@ -2063,17 +2035,17 @@ public partial class GUBP : BuildCommand
 		}
 	}
 
-	private void GetFullNodeListAndDirectDependencies(Dictionary<NodeInfo, NodeInfo> GUBPNodesControllingTrigger, Dictionary<NodeInfo, string> GUBPNodesControllingTriggerDotName, out Dictionary<string, string> FullNodeList, out Dictionary<string, string> FullNodeDirectDependencies)
+	private void GetFullNodeListAndDirectDependencies(out Dictionary<string, string> FullNodeList, out Dictionary<string, string> FullNodeDirectDependencies)
 	{
 		FullNodeList = new Dictionary<string,string>();
 		FullNodeDirectDependencies = new Dictionary<string,string>();
 
 		Log("******* {0} GUBP Nodes", GUBPNodes.Count);
-		List<NodeInfo> SortedNodes = TopologicalSort(new HashSet<NodeInfo>(GUBPNodes.Values), null/*GUBPNodesCompleted*/, GUBPNodesControllingTriggerDotName, GUBPNodesControllingTrigger, LocalOnly: true, DoNotConsiderCompletion: true);
+		List<NodeInfo> SortedNodes = TopologicalSort(new HashSet<NodeInfo>(GUBPNodes.Values), null/*GUBPNodesCompleted*/, LocalOnly: true, DoNotConsiderCompletion: true);
 		string DependencyStart = DateTime.Now.ToString();
 		foreach (NodeInfo Node in SortedNodes)
 		{
-			string Note = GetControllingTriggerDotName(Node, GUBPNodesControllingTriggerDotName, GUBPNodesControllingTrigger);
+			string Note = Node.ControllingTriggerDotName;
 			if (Note == "")
 			{
 				Note = CISFrequencyQuantumShiftString(Node);
@@ -2277,7 +2249,7 @@ public partial class GUBP : BuildCommand
 		}
 	}
 
-	private void DoCommanderSetup(int TimeIndex, bool bSkipTriggers, bool bFakeEC, bool LocalOnly, string CLString, NodeInfo ExplicitTrigger, Dictionary<NodeInfo, NodeInfo> GUBPNodesControllingTrigger, Dictionary<NodeInfo, string> GUBPNodesControllingTriggerDotName, Dictionary<string, string> FullNodeList, Dictionary<string, string> FullNodeDirectDependencies, Dictionary<string, string> FullNodeDependedOnBy, Dictionary<string, string> FullNodeDependentPromotions, List<string> SeparatePromotables, Dictionary<string, int> FullNodeListSortKey, Dictionary<string, bool> GUBPNodesCompleted, Dictionary<string, NodeHistory> GUBPNodesHistory, List<NodeInfo> OrdereredToDo, List<NodeInfo> UnfinishedTriggers, string FakeFail)
+	private void DoCommanderSetup(int TimeIndex, bool bSkipTriggers, bool bFakeEC, bool LocalOnly, string CLString, NodeInfo ExplicitTrigger, Dictionary<string, string> FullNodeList, Dictionary<string, string> FullNodeDirectDependencies, Dictionary<string, string> FullNodeDependedOnBy, Dictionary<string, string> FullNodeDependentPromotions, List<string> SeparatePromotables, Dictionary<string, int> FullNodeListSortKey, Dictionary<string, bool> GUBPNodesCompleted, Dictionary<string, NodeHistory> GUBPNodesHistory, List<NodeInfo> OrdereredToDo, List<NodeInfo> UnfinishedTriggers, string FakeFail)
 	{
 
 		if (OrdereredToDo.Count == 0)
@@ -2325,7 +2297,7 @@ public partial class GUBP : BuildCommand
 		// remove nodes that have unfinished triggers
 		foreach (NodeInfo NodeToDo in OrdereredToDo)
 		{
-			NodeInfo ControllingTrigger = GetControllingTrigger(NodeToDo, GUBPNodesControllingTrigger);
+			NodeInfo ControllingTrigger = NodeToDo.ControllingTrigger;
 			bool bNoUnfinishedTriggers = !UnfinishedTriggers.Contains(ControllingTrigger);
 
 			if (bNoUnfinishedTriggers)
@@ -2351,7 +2323,7 @@ public partial class GUBP : BuildCommand
 		string FinishFilterTime = DateTime.Now.ToString();
 		PrintCSVFile(String.Format("UAT,FilterNodes,{0},{1}", StartFilterTime, FinishFilterTime));
 		Log("*********** EC Nodes, in order.");
-		PrintNodes(this, OrdereredToDo, GUBPNodesCompleted, GUBPNodesHistory, GUBPNodesControllingTriggerDotName, GUBPNodesControllingTrigger, LocalOnly, UnfinishedTriggers);
+		PrintNodes(this, OrdereredToDo, GUBPNodesCompleted, GUBPNodesHistory, LocalOnly, UnfinishedTriggers);
 		string FinishNodePrint = DateTime.Now.ToString();
 		PrintCSVFile(String.Format("UAT,SetupCommanderPrint,{0},{1}", FinishFilterTime, FinishNodePrint));
 		// here we are just making sure everything before the explicit trigger is completed.
