Add a scoped stopwatch class for writing out telemetry events.

diff --git a/Engine/Source/Programs/AutomationTool/AutomationUtils/Automation.cs b/Engine/Source/Programs/AutomationTool/AutomationUtils/Automation.cs
--- a/Engine/Source/Programs/AutomationTool/AutomationUtils/Automation.cs
+++ b/Engine/Source/Programs/AutomationTool/AutomationUtils/Automation.cs
@@ -320,10 +320,10 @@ AutomationTool.exe [-verbose] [-compileonly] [-p4] Command0 [-Arg0 -Arg1 -Arg2 
 			// Compile scripts.
 			Log.TraceInformation("Compiling scripts.");
 			ScriptCompiler Compiler = new ScriptCompiler();
-            var StartScriptCompile = DateTime.Now.ToString();
-			Compiler.FindAndCompileAllScripts(AdditionalScriptsFolders: AdditionalScriptsFolders);
-            var FinishScriptCompile = DateTime.Now.ToString();
-            CommandUtils.PrintCSVFile(String.Format("UAT,ScriptCompile,{0},{1}", StartScriptCompile, FinishScriptCompile));
+			using(CommandUtils.TelemetryStopwatch ScriptCompileStopwatch = new CommandUtils.TelemetryStopwatch("ScriptCompile"))
+			{
+				Compiler.FindAndCompileAllScripts(AdditionalScriptsFolders: AdditionalScriptsFolders);
+			}
 
 			if (GlobalCommandLine.CompileOnly)
 			{
diff --git a/Engine/Source/Programs/AutomationTool/AutomationUtils/BuildUtils.cs b/Engine/Source/Programs/AutomationTool/AutomationUtils/BuildUtils.cs
--- a/Engine/Source/Programs/AutomationTool/AutomationUtils/BuildUtils.cs
+++ b/Engine/Source/Programs/AutomationTool/AutomationUtils/BuildUtils.cs
@@ -54,10 +54,10 @@ namespace AutomationTool
 				throw new AutomationException(String.Format("Unabled to build Solution {0}. devenv.com not found.", SolutionFile));
 			}
 			string CmdLine = String.Format("\"{0}\" /build \"{1}\"", SolutionFile, BuildConfig);
-            var StartBuildSolution = DateTime.Now.ToString();
-			RunAndLog(Env, Env.MsDevExe, CmdLine, LogName);
-            var FinishBuildSolution = DateTime.Now.ToString();
-            PrintCSVFile(String.Format("UAT,Compile.{0}.{1}.{2},{3},{4}", Path.GetFileName(SolutionFile), "WinC#", BuildConfig, StartBuildSolution, FinishBuildSolution));
+			using(TelemetryStopwatch CompileStopwatch = new TelemetryStopwatch("Compile.{0}.{1}.{2}", Path.GetFileName(SolutionFile), "WinC#", BuildConfig))
+			{
+				RunAndLog(Env, Env.MsDevExe, CmdLine, LogName);
+			}
 		}
 
 		/// <summary>
diff --git a/Engine/Source/Programs/AutomationTool/AutomationUtils/CommandUtils.cs b/Engine/Source/Programs/AutomationTool/AutomationUtils/CommandUtils.cs
--- a/Engine/Source/Programs/AutomationTool/AutomationUtils/CommandUtils.cs
+++ b/Engine/Source/Programs/AutomationTool/AutomationUtils/CommandUtils.cs
@@ -1896,24 +1896,45 @@ namespace AutomationTool
 			}
 		}
 		private static UnrealBuildTool.UnrealTargetPlatform[] UBTTargetPlatforms; 
-        public static void PrintCSVFile(string Input)
-        {
-            // sometimes nul is used as the CSV file so batch scripts can append using output redirectors easily. We have to ignore those.
-            if (IsBuildMachine && CmdEnv.CSVFile != "" && CmdEnv.CSVFile != "nul")
-            {
-                var CSVBuilder = new StringBuilder();
-                var CSVLineToAppend = String.Format("{0}{1}", Input, Environment.NewLine);
-                CSVBuilder.Append(CSVLineToAppend);
-                try
-                {
-                    File.AppendAllText(CmdEnv.CSVFile, CSVBuilder.ToString());
-                }
-                catch (Exception e)
-                {
-                    Log(System.Diagnostics.TraceEventType.Warning, "Could not append to csv file ({0}) : {1}", CmdEnv.CSVFile, e.ToString());
-                }
-            }
-        }
+
+		public class TelemetryStopwatch : IDisposable
+		{
+			string Name;
+			DateTime StartTime;
+			bool bFinished;
+
+			public TelemetryStopwatch(string Format, params object[] Args)
+			{
+				Name = String.Format(Format, Args);
+				StartTime = DateTime.Now;
+			}
+
+			public void Cancel()
+			{
+				bFinished = true;
+			}
+
+			public void Finish()
+			{
+				if(!bFinished && IsBuildMachine && !String.IsNullOrEmpty(CmdEnv.CSVFile) && CmdEnv.CSVFile != "nul")
+				{
+					try
+					{
+						File.AppendAllText(CmdEnv.CSVFile, String.Format("UAT,{0},{1},{2}" + Environment.NewLine, Name, StartTime.ToString(), DateTime.Now.ToString()));
+					}
+					catch (Exception Ex)
+					{
+						Log(System.Diagnostics.TraceEventType.Warning, "Could not append to csv file ({0}) : {1}", CmdEnv.CSVFile, Ex.ToString());
+					}
+					bFinished = true;
+				}
+			}
+
+			public void Dispose()
+			{
+				Finish();
+			}
+		}
 
 		#endregion
 
diff --git a/Engine/Source/Programs/AutomationTool/AutomationUtils/ProjectUtils.cs b/Engine/Source/Programs/AutomationTool/AutomationUtils/ProjectUtils.cs
--- a/Engine/Source/Programs/AutomationTool/AutomationUtils/ProjectUtils.cs
+++ b/Engine/Source/Programs/AutomationTool/AutomationUtils/ProjectUtils.cs
@@ -719,47 +719,47 @@ namespace AutomationTool
             BaseEngineProject = new BranchUProject();
 
             var AllProjects = UnrealBuildTool.UProjectInfo.FilterGameProjects(false, null);
-            var StartTime = DateTime.Now.ToString();
-            foreach (var InfoEntry in AllProjects)
-            {
-                var UProject = new BranchUProject(InfoEntry);
-                if (UProject.Properties.bIsCodeBasedProject)
-                {
-                    CodeProjects.Add(UProject);
-                }
-                else
-                {
-                    NonCodeProjects.Add(UProject);
-                    // the base project uses BlankProject if it really needs a .uproject file
-                    if (String.IsNullOrEmpty(BaseEngineProject.FilePath) && UProject.GameName == "BlankProject")
-                    {
-                        BaseEngineProject.FilePath = UProject.FilePath;
-                    }
-                }
-            }
-            var FinishTime = DateTime.Now.ToString();
-            CommandUtils.PrintCSVFile(String.Format("UAT,SortProjects,{0},{1}", StartTime, FinishTime));            
+			using(CommandUtils.TelemetryStopwatch SortProjectsStopwatch = new CommandUtils.TelemetryStopwatch("SortProjects"))
+			{
+				foreach (var InfoEntry in AllProjects)
+				{
+					var UProject = new BranchUProject(InfoEntry);
+					if (UProject.Properties.bIsCodeBasedProject)
+					{
+						CodeProjects.Add(UProject);
+					}
+					else
+					{
+						NonCodeProjects.Add(UProject);
+						// the base project uses BlankProject if it really needs a .uproject file
+						if (String.IsNullOrEmpty(BaseEngineProject.FilePath) && UProject.GameName == "BlankProject")
+						{
+							BaseEngineProject.FilePath = UProject.FilePath;
+						}
+					}
+				}
+			}
  /*           if (String.IsNullOrEmpty(BaseEngineProject.FilePath))
             {
                 throw new AutomationException("All branches must have the blank project /Samples/Sandbox/BlankProject");
             }*/
 
-            CommandUtils.Log("  Base Engine:");
-            var StartBranchDump = DateTime.Now.ToString();
-            BaseEngineProject.Dump(InHostPlatforms);
+			using(CommandUtils.TelemetryStopwatch ProjectDumpStopwatch = new CommandUtils.TelemetryStopwatch("Project Dump"))
+			{
+				CommandUtils.Log("  Base Engine:");
+				BaseEngineProject.Dump(InHostPlatforms);
 
-            CommandUtils.Log("  {0} Code projects:", CodeProjects.Count);
-            foreach (var Proj in CodeProjects)
-            {
-                Proj.Dump(InHostPlatforms);
-            }
-            CommandUtils.Log("  {0} Non-Code projects:", CodeProjects.Count);
-            foreach (var Proj in NonCodeProjects)
-            {
-                Proj.Dump(InHostPlatforms);
-            }
-            var FinishBranchDump = DateTime.Now.ToString();
-            CommandUtils.PrintCSVFile(String.Format("UAT,Project Dump,{0},{1}", StartBranchDump, FinishBranchDump));
+				CommandUtils.Log("  {0} Code projects:", CodeProjects.Count);
+				foreach (var Proj in CodeProjects)
+				{
+					Proj.Dump(InHostPlatforms);
+				}
+				CommandUtils.Log("  {0} Non-Code projects:", CodeProjects.Count);
+				foreach (var Proj in NonCodeProjects)
+				{
+					Proj.Dump(InHostPlatforms);
+				}
+			}
         }
 
         public BranchUProject FindGame(string GameName)
diff --git a/Engine/Source/Programs/AutomationTool/AutomationUtils/UE4Build.cs b/Engine/Source/Programs/AutomationTool/AutomationUtils/UE4Build.cs
--- a/Engine/Source/Programs/AutomationTool/AutomationUtils/UE4Build.cs
+++ b/Engine/Source/Programs/AutomationTool/AutomationUtils/UE4Build.cs
@@ -189,10 +189,10 @@ namespace AutomationTool
 
 			ClearExportedXGEXML();
 
-            var StartXGEManifestExport = DateTime.Now.ToString();
-            RunUBT(CmdEnv, UBTExecutable: UBTExecutable, Project: ProjectName, Target: TargetName, Platform: Platform.ToString(), Config: Config, AdditionalArgs: "-generatemanifest -nobuilduht -xgeexport" + AddArgs, EnvVars: EnvVars);
-            var FinishXGEManifestExport = DateTime.Now.ToString();
-            PrintCSVFile(String.Format("UAT,GenerateXGEManifest.{0}.{1}.{2},{3},{4}", TargetName, Platform.ToString(), Config, StartXGEManifestExport, FinishXGEManifestExport));
+			using(TelemetryStopwatch GenerateManifestStopwatch = new TelemetryStopwatch("GenerateXGEManifest.{0}.{1}.{2}", TargetName, Platform.ToString(), Config))
+			{
+				RunUBT(CmdEnv, UBTExecutable: UBTExecutable, Project: ProjectName, Target: TargetName, Platform: Platform.ToString(), Config: Config, AdditionalArgs: "-generatemanifest -nobuilduht -xgeexport" + AddArgs, EnvVars: EnvVars);
+			}
 
 			PrepareManifest(UBTManifest, true);
 
@@ -285,10 +285,10 @@ namespace AutomationTool
 			}
 
 			PrepareUBT();
-            var StartCleanUBT = DateTime.Now.ToString();
-            RunUBT(CmdEnv, UBTExecutable: UBTExecutable, Project: ProjectName, Target: TargetName, Platform: Platform.ToString(), Config: Config, AdditionalArgs: "-clean" + AddArgs, EnvVars: EnvVars);
-            var FinishCleanUBT = DateTime.Now.ToString();
-            PrintCSVFile(String.Format("UAT,CleanWithUBT.{0}.{1}.{2},{3},{4}", TargetName, Platform.ToString(), Config, StartCleanUBT, FinishCleanUBT));
+			using(TelemetryStopwatch CleanStopwatch = new TelemetryStopwatch("CleanWithUBT.{0}.{1}.{2}", TargetName, Platform.ToString(), Config))
+			{
+				RunUBT(CmdEnv, UBTExecutable: UBTExecutable, Project: ProjectName, Target: TargetName, Platform: Platform.ToString(), Config: Config, AdditionalArgs: "-clean" + AddArgs, EnvVars: EnvVars);
+			}
         }
 
 		void BuildWithUBT(string ProjectName, string TargetName, UnrealBuildTool.UnrealTargetPlatform TargetPlatform, string Config, string UprojectPath, bool ForceMonolithic = false, bool ForceNonUnity = false, bool ForceDebugInfo = false, bool ForceFlushMac = false, bool DisableXGE = false, string InAddArgs = "", bool ForceUnity = false, Dictionary<string, string> EnvVars = null)
@@ -334,16 +334,17 @@ namespace AutomationTool
                 string UBTManifest = GetUBTManifest(UprojectPath, AddArgs);
 
 				DeleteFile(UBTManifest);
-                var StartGenerateManifest = DateTime.Now.ToString();
-                RunUBT(CmdEnv, UBTExecutable: UBTExecutable, Project: ProjectName, Target: TargetName, Platform: TargetPlatform.ToString(), Config: Config, AdditionalArgs: AddArgs + " -generatemanifest", EnvVars: EnvVars);
-                var FinishGenerateManifest = DateTime.Now.ToString();
-                PrintCSVFile(String.Format("UAT,PrepareUBTManifest.{0}.{1}.{2},{3},{4}", TargetName, TargetPlatform.ToString(), Config, StartGenerateManifest, FinishGenerateManifest));
+				using(TelemetryStopwatch PrepareManifestStopwatch = new TelemetryStopwatch("PrepareUBTManifest.{0}.{1}.{2}", TargetName, TargetPlatform.ToString(), Config))
+				{
+					RunUBT(CmdEnv, UBTExecutable: UBTExecutable, Project: ProjectName, Target: TargetName, Platform: TargetPlatform.ToString(), Config: Config, AdditionalArgs: AddArgs + " -generatemanifest", EnvVars: EnvVars);
+				}
                 PrepareManifest(UBTManifest, false);
 			}
-            var StartCompile = DateTime.Now.ToString();
-            RunUBT(CmdEnv, UBTExecutable: UBTExecutable, Project: ProjectName, Target: TargetName, Platform: TargetPlatform.ToString(), Config: Config, AdditionalArgs: AddArgs, EnvVars: EnvVars);
-            var FinishCompile = DateTime.Now.ToString();
-            PrintCSVFile(String.Format("UAT,Compile.{0}.{1}.{2},{3},{4}", TargetName, TargetPlatform.ToString(), Config, StartCompile, FinishCompile));
+
+			using(TelemetryStopwatch CompileStopwatch = new TelemetryStopwatch("Compile.{0}.{1}.{2}", TargetName, TargetPlatform.ToString(), Config))
+			{
+				RunUBT(CmdEnv, UBTExecutable: UBTExecutable, Project: ProjectName, Target: TargetName, Platform: TargetPlatform.ToString(), Config: Config, AdditionalArgs: AddArgs, EnvVars: EnvVars);
+			}
             // allow the platform to perform any actions after building a target (seems almost like this should be done in UBT)
 			Platform.Platforms[TargetPlatform].PostBuildTarget(this, string.IsNullOrEmpty(ProjectName) ? TargetName : ProjectName, UprojectPath, Config);
 
@@ -785,12 +786,13 @@ namespace AutomationTool
 
 		public bool ProcessXGEItems(List<XGEItem> Actions, string XGETool, string Args, string TaskFilePath, bool DoRetries, bool SpecialTestFlag, bool ShowProgress)
 		{
-            var StartProcessTime = DateTime.Now.ToString();
             if (string.IsNullOrEmpty(XGETool))
 			{
 				throw new AutomationException("Unable to find xge tool: " + XGETool);
 			}
 
+			TelemetryStopwatch CombineXGEStopwatch = new TelemetryStopwatch("CombineXGEItemFiles.{0}", Path.GetFileNameWithoutExtension(XGETool));
+
 			XmlDocument XGETaskDocument;	
 			if (!CombineXGEItemFiles(Actions, TaskFilePath, out XGETaskDocument))
 			{
@@ -843,68 +845,67 @@ namespace AutomationTool
 					throw new AutomationException("Unable to find xge xml: " + TaskFilePath);
 				}
 
-                var FinishCombineXGEItemFilesTime = DateTime.Now.ToString();
-                PrintCSVFile(String.Format("UAT,CombineXGEItemFiles.{0},{1},{2}", Path.GetFileNameWithoutExtension(XGETool), StartProcessTime, FinishCombineXGEItemFilesTime));
-                
-                int Retries = DoRetries ? 2 : 1;
-                int ConnectionRetries = 4;
-                for (int i = 0; i < Retries; i++)
+				CombineXGEStopwatch.Finish();
+               
+				using(TelemetryStopwatch ProcessXGEStopwatch = new TelemetryStopwatch("ProcessXGE.{0}", Path.GetFileNameWithoutExtension(XGETool)))
 				{
-					try
+					int Retries = DoRetries ? 2 : 1;
+					int ConnectionRetries = 4;
+					for (int i = 0; i < Retries; i++)
 					{
-                        while (true)
-                        {
-                            Log("Running {0} *******", XGETool);
-                            PushDir(CombinePaths(CmdEnv.LocalRoot, @"\Engine\Source"));
-                            int SuccesCode;
-                            string LogFile = GetRunAndLogLogName(CmdEnv, XGETool);
-                            string Output = RunAndLog(XGETool, Args, out SuccesCode, LogFile);
-                            PopDir();
-                            if (ConnectionRetries > 0 && (SuccesCode == 4 || SuccesCode == 2) && !Output.Contains("------Project:"))
-                            {
-                                Log(System.Diagnostics.TraceEventType.Warning, String.Format("{0} failure on the local connection timeout", XGETool));
-                                if (ConnectionRetries < 2)
-                                {
-                                    System.Threading.Thread.Sleep(60000);
-                                }
-                                ConnectionRetries--;
-                                continue;
-                            }
-                            else if (SuccesCode != 0)
-                            {
-                                Log("{0} did not succeed *******", XGETool);
-                                throw new AutomationException("Command failed (Result:{3}): {0} {1}. See logfile for details: '{2}' ",
-                                                                XGETool, Args, Path.GetFileName(LogFile), SuccesCode);
-                            }
-                            Log("{0} {1} Done *******", XGETool, Args);
-                            break;
-                        }
-                        break;
-					}
-					catch (Exception Ex)
-					{
-						Log("{0} failed on try {1}: {2}", XGETool, i + 1, Ex.ToString());
-						if (i + 1 >= Retries)
+						try
 						{
-							return false;
+							while (true)
+							{
+								Log("Running {0} *******", XGETool);
+								PushDir(CombinePaths(CmdEnv.LocalRoot, @"\Engine\Source"));
+								int SuccesCode;
+								string LogFile = GetRunAndLogLogName(CmdEnv, XGETool);
+								string Output = RunAndLog(XGETool, Args, out SuccesCode, LogFile);
+								PopDir();
+								if (ConnectionRetries > 0 && (SuccesCode == 4 || SuccesCode == 2) && !Output.Contains("------Project:"))
+								{
+									Log(System.Diagnostics.TraceEventType.Warning, String.Format("{0} failure on the local connection timeout", XGETool));
+									if (ConnectionRetries < 2)
+									{
+										System.Threading.Thread.Sleep(60000);
+									}
+									ConnectionRetries--;
+									continue;
+								}
+								else if (SuccesCode != 0)
+								{
+									Log("{0} did not succeed *******", XGETool);
+									throw new AutomationException("Command failed (Result:{3}): {0} {1}. See logfile for details: '{2}' ",
+																	XGETool, Args, Path.GetFileName(LogFile), SuccesCode);
+								}
+								Log("{0} {1} Done *******", XGETool, Args);
+								break;
+							}
+							break;
 						}
-						Log("Deleting build products to force a rebuild.");
-						foreach (var Item in Actions)
+						catch (Exception Ex)
 						{
-							XGEDeleteBuildProducts(Item.Manifest);
+							Log("{0} failed on try {1}: {2}", XGETool, i + 1, Ex.ToString());
+							if (i + 1 >= Retries)
+							{
+								return false;
+							}
+							Log("Deleting build products to force a rebuild.");
+							foreach (var Item in Actions)
+							{
+								XGEDeleteBuildProducts(Item.Manifest);
+							}
 						}
 					}
 				}
-                var FinishXGEToolTime = DateTime.Now.ToString();
-                PrintCSVFile(String.Format("UAT,ProcessXGE.{0},{1},{2}", Path.GetFileNameWithoutExtension(XGETool), FinishCombineXGEItemFilesTime, FinishXGEToolTime));
-
 			}
 			foreach (var Item in Actions)
 			{
 				XGEFinishBuildWithUBT(Item);
 			}
 			return true;
-			}
+		}
 
         private static bool CombineXGEItemFiles(List<XGEItem> Actions, string TaskFilePath, out XmlDocument XGETaskDocument)
         {
diff --git a/Engine/Source/Programs/AutomationTool/GUBP/GUBP.cs b/Engine/Source/Programs/AutomationTool/GUBP/GUBP.cs
--- a/Engine/Source/Programs/AutomationTool/GUBP/GUBP.cs
+++ b/Engine/Source/Programs/AutomationTool/GUBP/GUBP.cs
@@ -1473,7 +1473,6 @@ public partial class GUBP : BuildCommand
             TempStorage.DeleteLocalTempStorageManifests(CmdEnv);
         }
 
-		Dictionary<string, string> FullNodeDependedOnBy = new Dictionary<string, string>();
 		Dictionary<string, string> FullNodeDependentPromotions = new Dictionary<string, string>();		
 		List<string> SeparatePromotables = new List<string>();
         {
@@ -1532,6 +1531,7 @@ public partial class GUBP : BuildCommand
 
 		bool bOnlyNode;
 		HashSet<NodeInfo> NodesToDo;
+		Dictionary<string, string> FullNodeDependedOnBy = new Dictionary<string, string>();
 		ParseNodesToDo(WithoutLinux, TimeIndex, CommanderSetup, FullNodeDependedOnBy, out bOnlyNode, out NodesToDo);
 
 		NodeInfo ExplicitTrigger = null;
@@ -2037,46 +2037,46 @@ public partial class GUBP : BuildCommand
 
 	private void GetFullNodeListAndDirectDependencies(out Dictionary<string, string> FullNodeList, out Dictionary<string, string> FullNodeDirectDependencies)
 	{
-		FullNodeList = new Dictionary<string,string>();
-		FullNodeDirectDependencies = new Dictionary<string,string>();
-
-		Log("******* {0} GUBP Nodes", GUBPNodes.Count);
-		List<NodeInfo> SortedNodes = TopologicalSort(new HashSet<NodeInfo>(GUBPNodes.Values), null/*GUBPNodesCompleted*/, LocalOnly: true, DoNotConsiderCompletion: true);
-		string DependencyStart = DateTime.Now.ToString();
-		foreach (NodeInfo Node in SortedNodes)
+		using(TelemetryStopwatch Timer = new TelemetryStopwatch("GetDependencies"))
 		{
-			string Note = Node.ControllingTriggerDotName;
-			if (Note == "")
-			{
-				Note = CISFrequencyQuantumShiftString(Node);
-			}
-			if (Note == "")
-			{
-				Note = "always";
-			}
-			if (GUBPNodes[Node.Name].Node.RunInEC())
+			FullNodeList = new Dictionary<string,string>();
+			FullNodeDirectDependencies = new Dictionary<string,string>();
+
+			Log("******* {0} GUBP Nodes", GUBPNodes.Count);
+			List<NodeInfo> SortedNodes = TopologicalSort(new HashSet<NodeInfo>(GUBPNodes.Values), null/*GUBPNodesCompleted*/, LocalOnly: true, DoNotConsiderCompletion: true);
+			foreach (NodeInfo Node in SortedNodes)
 			{
-				List<NodeInfo> Deps = GetECDependencies(Node);
-				string All = "";
-				foreach (NodeInfo Dep in Deps)
+				string Note = Node.ControllingTriggerDotName;
+				if (Note == "")
 				{
-					if (All != "")
+					Note = CISFrequencyQuantumShiftString(Node);
+				}
+				if (Note == "")
+				{
+					Note = "always";
+				}
+				if (GUBPNodes[Node.Name].Node.RunInEC())
+				{
+					List<NodeInfo> Deps = GetECDependencies(Node);
+					string All = "";
+					foreach (NodeInfo Dep in Deps)
 					{
-						All += " ";
+						if (All != "")
+						{
+							All += " ";
+						}
+						All += Dep.Name;
 					}
-					All += Dep.Name;
+					LogVerbose("  {0}: {1}      {2}", Node.Name, Note, All);
+					FullNodeList.Add(Node.Name, Note);
+					FullNodeDirectDependencies.Add(Node.Name, All);
+				}
+				else
+				{
+					LogVerbose("  {0}: {1} [Aggregate]", Node.Name, Note);
 				}
-				LogVerbose("  {0}: {1}      {2}", Node.Name, Note, All);
-				FullNodeList.Add(Node.Name, Note);
-				FullNodeDirectDependencies.Add(Node.Name, All);
-			}
-			else
-			{
-				LogVerbose("  {0}: {1} [Aggregate]", Node.Name, Note);
 			}
 		}
-		string DependencyFinish = DateTime.Now.ToString();
-		PrintCSVFile(String.Format("UAT,GetDependencies,{0},{1}", DependencyStart, DependencyFinish));
 	}
 
 	private int UpdateCISCounter()
@@ -2292,44 +2292,44 @@ public partial class GUBP : BuildCommand
 			ECJobProps.Add("IsRoot=1");
 		}
 
-		List<NodeInfo> FilteredOrdereredToDo = new List<NodeInfo>();
-		string StartFilterTime = DateTime.Now.ToString();
-		// remove nodes that have unfinished triggers
-		foreach (NodeInfo NodeToDo in OrdereredToDo)
+		using(TelemetryStopwatch StartFilterTimer = new TelemetryStopwatch("FilterNodes"))
 		{
-			NodeInfo ControllingTrigger = NodeToDo.ControllingTrigger;
-			bool bNoUnfinishedTriggers = !UnfinishedTriggers.Contains(ControllingTrigger);
-
-			if (bNoUnfinishedTriggers)
+			List<NodeInfo> FilteredOrdereredToDo = new List<NodeInfo>();
+			// remove nodes that have unfinished triggers
+			foreach (NodeInfo NodeToDo in OrdereredToDo)
 			{
-				// if we are triggering, then remove nodes that are not controlled by the trigger or are dependencies of this trigger
-				if (ExplicitTrigger != null)
+				NodeInfo ControllingTrigger = NodeToDo.ControllingTrigger;
+				bool bNoUnfinishedTriggers = !UnfinishedTriggers.Contains(ControllingTrigger);
+
+				if (bNoUnfinishedTriggers)
 				{
-					if (ExplicitTrigger != NodeToDo && !NodeDependsOn(NodeToDo, ExplicitTrigger) && !NodeDependsOn(ExplicitTrigger, NodeToDo))
+					// if we are triggering, then remove nodes that are not controlled by the trigger or are dependencies of this trigger
+					if (ExplicitTrigger != null)
 					{
-						continue; // this wasn't on the chain related to the trigger we are triggering, so it is not relevant
+						if (ExplicitTrigger != NodeToDo && !NodeDependsOn(NodeToDo, ExplicitTrigger) && !NodeDependsOn(ExplicitTrigger, NodeToDo))
+						{
+							continue; // this wasn't on the chain related to the trigger we are triggering, so it is not relevant
+						}
 					}
+					if (bPreflightBuild && !bSkipTriggers && NodeToDo.Node.TriggerNode())
+					{
+						// in preflight builds, we are either skipping triggers (and running things downstream) or we just stop at triggers and don't make them available for triggering.
+						continue;
+					}
+					FilteredOrdereredToDo.Add(NodeToDo);
 				}
-				if (bPreflightBuild && !bSkipTriggers && NodeToDo.Node.TriggerNode())
-				{
-					// in preflight builds, we are either skipping triggers (and running things downstream) or we just stop at triggers and don't make them available for triggering.
-					continue;
-				}
-				FilteredOrdereredToDo.Add(NodeToDo);
 			}
+			OrdereredToDo = FilteredOrdereredToDo;
+		}
+		using(TelemetryStopwatch PrintNodesTimer = new TelemetryStopwatch("SetupCommanderPrint"))
+		{
+			Log("*********** EC Nodes, in order.");
+			PrintNodes(this, OrdereredToDo, GUBPNodesCompleted, GUBPNodesHistory, LocalOnly, UnfinishedTriggers);
 		}
-
-		OrdereredToDo = FilteredOrdereredToDo;
-		string FinishFilterTime = DateTime.Now.ToString();
-		PrintCSVFile(String.Format("UAT,FilterNodes,{0},{1}", StartFilterTime, FinishFilterTime));
-		Log("*********** EC Nodes, in order.");
-		PrintNodes(this, OrdereredToDo, GUBPNodesCompleted, GUBPNodesHistory, LocalOnly, UnfinishedTriggers);
-		string FinishNodePrint = DateTime.Now.ToString();
-		PrintCSVFile(String.Format("UAT,SetupCommanderPrint,{0},{1}", FinishFilterTime, FinishNodePrint));
 		// here we are just making sure everything before the explicit trigger is completed.
 		if (ExplicitTrigger != null)
 		{
-			foreach (NodeInfo NodeToDo in FilteredOrdereredToDo)
+			foreach (NodeInfo NodeToDo in OrdereredToDo)
 			{
 				if (NodeToDo.Node.RunInEC() && !NodeIsAlreadyComplete(GUBPNodesCompleted, NodeToDo, LocalOnly) && NodeToDo != ExplicitTrigger && !NodeDependsOn(ExplicitTrigger, NodeToDo)) // if something is already finished, we don't put it into EC
 				{
@@ -2369,337 +2369,337 @@ public partial class GUBP : BuildCommand
 			}
 		}
 
-		string ParentPath = ParseParamValue("ParentPath");
-		string StartPerlOutput = DateTime.Now.ToString();
-		string BaseArgs = String.Format("$batch->createJobStep({{parentPath => '{0}'", ParentPath);
-
-		bool bHasNoop = false;
-		if (LastSticky == null && bHaveECNodes)
-		{
-			// if we don't have any sticky nodes and we have other nodes, we run a fake noop just to release the resource 
-			string Args = String.Format("{0}, subprocedure => 'GUBP_UAT_Node', parallel => '0', jobStepName => 'Noop', actualParameter => [{{actualParameterName => 'NodeName', value => 'Noop'}}, {{actualParameterName => 'Sticky', value =>'1' }}], releaseMode => 'release'}});", BaseArgs);
-			StepList.Add(Args);
-			bHasNoop = true;
-		}
-
 		List<string> FakeECArgs = new List<string>();
-		Dictionary<string, List<NodeInfo>> AgentGroupChains = new Dictionary<string, List<NodeInfo>>();
-		List<NodeInfo> StickyChain = new List<NodeInfo>();
-		foreach (NodeInfo NodeToDo in OrdereredToDo)
+		using(TelemetryStopwatch PerlOutputStopwatch = new TelemetryStopwatch("PerlOutput"))
 		{
-			if (NodeToDo.Node.RunInEC() && !NodeIsAlreadyComplete(GUBPNodesCompleted, NodeToDo, LocalOnly)) // if something is already finished, we don't put it into EC  
+			string ParentPath = ParseParamValue("ParentPath");
+			string BaseArgs = String.Format("$batch->createJobStep({{parentPath => '{0}'", ParentPath);
+
+			bool bHasNoop = false;
+			if (LastSticky == null && bHaveECNodes)
 			{
-				string MyAgentGroup = NodeToDo.Node.AgentSharingGroup;
-				if (MyAgentGroup != "")
+				// if we don't have any sticky nodes and we have other nodes, we run a fake noop just to release the resource 
+				string Args = String.Format("{0}, subprocedure => 'GUBP_UAT_Node', parallel => '0', jobStepName => 'Noop', actualParameter => [{{actualParameterName => 'NodeName', value => 'Noop'}}, {{actualParameterName => 'Sticky', value =>'1' }}], releaseMode => 'release'}});", BaseArgs);
+				StepList.Add(Args);
+				bHasNoop = true;
+			}
+
+			Dictionary<string, List<NodeInfo>> AgentGroupChains = new Dictionary<string, List<NodeInfo>>();
+			List<NodeInfo> StickyChain = new List<NodeInfo>();
+			foreach (NodeInfo NodeToDo in OrdereredToDo)
+			{
+				if (NodeToDo.Node.RunInEC() && !NodeIsAlreadyComplete(GUBPNodesCompleted, NodeToDo, LocalOnly)) // if something is already finished, we don't put it into EC  
 				{
-					if (!AgentGroupChains.ContainsKey(MyAgentGroup))
+					string MyAgentGroup = NodeToDo.Node.AgentSharingGroup;
+					if (MyAgentGroup != "")
 					{
-						AgentGroupChains.Add(MyAgentGroup, new List<NodeInfo> { NodeToDo });
+						if (!AgentGroupChains.ContainsKey(MyAgentGroup))
+						{
+							AgentGroupChains.Add(MyAgentGroup, new List<NodeInfo> { NodeToDo });
+						}
+						else
+						{
+							AgentGroupChains[MyAgentGroup].Add(NodeToDo);
+						}
 					}
-					else
+				}
+				if (NodeToDo.Node.IsSticky())
+				{
+					if (!StickyChain.Contains(NodeToDo))
 					{
-						AgentGroupChains[MyAgentGroup].Add(NodeToDo);
+						StickyChain.Add(NodeToDo);
 					}
 				}
 			}
-			if (NodeToDo.Node.IsSticky())
+			foreach (NodeInfo NodeToDo in OrdereredToDo)
 			{
-				if (!StickyChain.Contains(NodeToDo))
+				if (NodeToDo.Node.RunInEC() && !NodeIsAlreadyComplete(GUBPNodesCompleted, NodeToDo, LocalOnly)) // if something is already finished, we don't put it into EC  
 				{
-					StickyChain.Add(NodeToDo);
-				}
-			}
-		}
-		foreach (NodeInfo NodeToDo in OrdereredToDo)
-		{
-			if (NodeToDo.Node.RunInEC() && !NodeIsAlreadyComplete(GUBPNodesCompleted, NodeToDo, LocalOnly)) // if something is already finished, we don't put it into EC  
-			{
-				string EMails;
-				List<string> NodeProps = GetECPropsForNode(NodeToDo, CLString, out EMails);
-				ECProps.AddRange(NodeProps);
+					string EMails;
+					List<string> NodeProps = GetECPropsForNode(NodeToDo, CLString, out EMails);
+					ECProps.AddRange(NodeProps);
 
-				bool Sticky = NodeToDo.Node.IsSticky();
-				bool DoParallel = !Sticky;
-				if (NodeToDo.Node.ECProcedure() == "GUBP_UAT_Node_Parallel_AgentShare_Editor")
-				{
-					DoParallel = true;
-				}
-				if (Sticky && NodeToDo.Node.ECAgentString() != "")
-				{
-					throw new AutomationException("Node {1} is sticky but has agent requirements.", NodeToDo.Name);
-				}
-				string Procedure = NodeToDo.Node.ECProcedure();
-				if (NodeToDo.Node.IsSticky() && NodeToDo == LastSticky)
-				{
-					Procedure = Procedure + "_Release";
-				}
-				string Args = String.Format("{0}, subprocedure => '{1}', parallel => '{2}', jobStepName => '{3}', actualParameter => [{{actualParameterName => 'NodeName', value =>'{4}'}}",
-					BaseArgs, Procedure, DoParallel ? 1 : 0, NodeToDo.Name, NodeToDo.Name);
-				string ProcedureParams = NodeToDo.Node.ECProcedureParams();
-				if (!String.IsNullOrEmpty(ProcedureParams))
-				{
-					Args = Args + ProcedureParams;
-				}
+					bool Sticky = NodeToDo.Node.IsSticky();
+					bool DoParallel = !Sticky;
+					if (NodeToDo.Node.ECProcedure() == "GUBP_UAT_Node_Parallel_AgentShare_Editor")
+					{
+						DoParallel = true;
+					}
+					if (Sticky && NodeToDo.Node.ECAgentString() != "")
+					{
+						throw new AutomationException("Node {1} is sticky but has agent requirements.", NodeToDo.Name);
+					}
+					string Procedure = NodeToDo.Node.ECProcedure();
+					if (NodeToDo.Node.IsSticky() && NodeToDo == LastSticky)
+					{
+						Procedure = Procedure + "_Release";
+					}
+					string Args = String.Format("{0}, subprocedure => '{1}', parallel => '{2}', jobStepName => '{3}', actualParameter => [{{actualParameterName => 'NodeName', value =>'{4}'}}",
+						BaseArgs, Procedure, DoParallel ? 1 : 0, NodeToDo.Name, NodeToDo.Name);
+					string ProcedureParams = NodeToDo.Node.ECProcedureParams();
+					if (!String.IsNullOrEmpty(ProcedureParams))
+					{
+						Args = Args + ProcedureParams;
+					}
 
-				if ((Procedure == "GUBP_UAT_Trigger" || Procedure == "GUBP_Hardcoded_Trigger") && !String.IsNullOrEmpty(EMails))
-				{
-					Args = Args + ", {actualParameterName => 'EmailsForTrigger', value => \'" + EMails + "\'}";
-				}
-				Args = Args + "]";
-				string PreCondition = "";
-				string RunCondition = "";
-				List<string> UncompletedEcDeps = new List<string>();
-				List<string> PreConditionUncompletedEcDeps = new List<string>();
-				string MyAgentGroup = NodeToDo.Node.AgentSharingGroup;
-				bool bDoNestedJobstep = false;
-				bool bDoFirstNestedJobstep = false;
+					if ((Procedure == "GUBP_UAT_Trigger" || Procedure == "GUBP_Hardcoded_Trigger") && !String.IsNullOrEmpty(EMails))
+					{
+						Args = Args + ", {actualParameterName => 'EmailsForTrigger', value => \'" + EMails + "\'}";
+					}
+					Args = Args + "]";
+					string PreCondition = "";
+					string RunCondition = "";
+					List<string> UncompletedEcDeps = new List<string>();
+					List<string> PreConditionUncompletedEcDeps = new List<string>();
+					string MyAgentGroup = NodeToDo.Node.AgentSharingGroup;
+					bool bDoNestedJobstep = false;
+					bool bDoFirstNestedJobstep = false;
 
 
-				string NodeParentPath = ParentPath;
-				string PreconditionParentPath;
+					string NodeParentPath = ParentPath;
+					string PreconditionParentPath;
 
-				PreconditionParentPath = ParentPath;
-				{
-					List<NodeInfo> EcDeps = GetECDependencies(NodeToDo);
-					foreach (NodeInfo Dep in EcDeps)
+					PreconditionParentPath = ParentPath;
 					{
-						if (Dep.Node.RunInEC() && !NodeIsAlreadyComplete(GUBPNodesCompleted, Dep, LocalOnly) && OrdereredToDo.Contains(Dep)) // if something is already finished, we don't put it into EC
+						List<NodeInfo> EcDeps = GetECDependencies(NodeToDo);
+						foreach (NodeInfo Dep in EcDeps)
 						{
-							if (OrdereredToDo.IndexOf(Dep) > OrdereredToDo.IndexOf(NodeToDo))
+							if (Dep.Node.RunInEC() && !NodeIsAlreadyComplete(GUBPNodesCompleted, Dep, LocalOnly) && OrdereredToDo.Contains(Dep)) // if something is already finished, we don't put it into EC
 							{
-								throw new AutomationException("Topological sort error, node {0} has a dependency of {1} which sorted after it.", NodeToDo.Name, Dep.Name);
+								if (OrdereredToDo.IndexOf(Dep) > OrdereredToDo.IndexOf(NodeToDo))
+								{
+									throw new AutomationException("Topological sort error, node {0} has a dependency of {1} which sorted after it.", NodeToDo.Name, Dep.Name);
+								}
+								UncompletedEcDeps.Add(Dep.Name);
 							}
-							UncompletedEcDeps.Add(Dep.Name);
 						}
 					}
-				}
-
-				foreach (string Dep in UncompletedEcDeps)
-				{
-					PreConditionUncompletedEcDeps.Add(Dep);
-				}
-				if (MyAgentGroup != "")
-				{
-					bDoNestedJobstep = true;
-					NodeParentPath = ParentPath + "/jobSteps[" + MyAgentGroup + "]";
 
-					PreConditionUncompletedEcDeps = new List<string>();
-
-					List<NodeInfo> MyChain = AgentGroupChains[MyAgentGroup];
-					int MyIndex = MyChain.IndexOf(NodeToDo);
-					if (MyIndex > 0)
+					foreach (string Dep in UncompletedEcDeps)
 					{
-						PreConditionUncompletedEcDeps.Add(MyChain[MyIndex - 1].Name);
+						PreConditionUncompletedEcDeps.Add(Dep);
 					}
-					else
+					if (MyAgentGroup != "")
 					{
-						bDoFirstNestedJobstep = bDoNestedJobstep;
-						// to avoid idle agents (and also EC doesn't actually reserve our agent!), we promote all dependencies to the first one
-						foreach (NodeInfo Chain in MyChain)
+						bDoNestedJobstep = true;
+						NodeParentPath = ParentPath + "/jobSteps[" + MyAgentGroup + "]";
+
+						PreConditionUncompletedEcDeps = new List<string>();
+
+						List<NodeInfo> MyChain = AgentGroupChains[MyAgentGroup];
+						int MyIndex = MyChain.IndexOf(NodeToDo);
+						if (MyIndex > 0)
 						{
-							List<NodeInfo> EcDeps = GetECDependencies(Chain);
-							foreach (NodeInfo Dep in EcDeps)
+							PreConditionUncompletedEcDeps.Add(MyChain[MyIndex - 1].Name);
+						}
+						else
+						{
+							bDoFirstNestedJobstep = bDoNestedJobstep;
+							// to avoid idle agents (and also EC doesn't actually reserve our agent!), we promote all dependencies to the first one
+							foreach (NodeInfo Chain in MyChain)
 							{
-								if (GUBPNodes[Dep.Name].Node.RunInEC() && !NodeIsAlreadyComplete(GUBPNodesCompleted, Dep, LocalOnly) && OrdereredToDo.Contains(Dep)) // if something is already finished, we don't put it into EC
+								List<NodeInfo> EcDeps = GetECDependencies(Chain);
+								foreach (NodeInfo Dep in EcDeps)
 								{
-									if (OrdereredToDo.IndexOf(Dep) > OrdereredToDo.IndexOf(Chain))
+									if (GUBPNodes[Dep.Name].Node.RunInEC() && !NodeIsAlreadyComplete(GUBPNodesCompleted, Dep, LocalOnly) && OrdereredToDo.Contains(Dep)) // if something is already finished, we don't put it into EC
 									{
-										throw new AutomationException("Topological sort error, node {0} has a dependency of {1} which sorted after it.", Chain.Name, Dep.Name);
-									}
-									if (!MyChain.Contains(Dep) && !PreConditionUncompletedEcDeps.Contains(Dep.Name))
-									{
-										PreConditionUncompletedEcDeps.Add(Dep.Name);
+										if (OrdereredToDo.IndexOf(Dep) > OrdereredToDo.IndexOf(Chain))
+										{
+											throw new AutomationException("Topological sort error, node {0} has a dependency of {1} which sorted after it.", Chain.Name, Dep.Name);
+										}
+										if (!MyChain.Contains(Dep) && !PreConditionUncompletedEcDeps.Contains(Dep.Name))
+										{
+											PreConditionUncompletedEcDeps.Add(Dep.Name);
+										}
 									}
 								}
 							}
 						}
 					}
-				}
-				if (NodeToDo.Node.IsSticky())
-				{
-					List<NodeInfo> MyChain = StickyChain;
-					int MyIndex = MyChain.IndexOf(NodeToDo);
-					if (MyIndex > 0)
+					if (NodeToDo.Node.IsSticky())
 					{
-						if (!PreConditionUncompletedEcDeps.Contains(MyChain[MyIndex - 1].Name) && !NodeIsAlreadyComplete(GUBPNodesCompleted, MyChain[MyIndex - 1], LocalOnly))
+						List<NodeInfo> MyChain = StickyChain;
+						int MyIndex = MyChain.IndexOf(NodeToDo);
+						if (MyIndex > 0)
 						{
-							PreConditionUncompletedEcDeps.Add(MyChain[MyIndex - 1].Name);
+							if (!PreConditionUncompletedEcDeps.Contains(MyChain[MyIndex - 1].Name) && !NodeIsAlreadyComplete(GUBPNodesCompleted, MyChain[MyIndex - 1], LocalOnly))
+							{
+								PreConditionUncompletedEcDeps.Add(MyChain[MyIndex - 1].Name);
+							}
 						}
-					}
-					else
-					{
-						List<NodeInfo> EcDeps = GetECDependencies(NodeToDo);
-						foreach (NodeInfo Dep in EcDeps)
+						else
 						{
-							if (Dep.Node.RunInEC() && !NodeIsAlreadyComplete(GUBPNodesCompleted, Dep, LocalOnly) && OrdereredToDo.Contains(Dep)) // if something is already finished, we don't put it into EC
+							List<NodeInfo> EcDeps = GetECDependencies(NodeToDo);
+							foreach (NodeInfo Dep in EcDeps)
 							{
-								if (OrdereredToDo.IndexOf(Dep) > OrdereredToDo.IndexOf(NodeToDo))
+								if (Dep.Node.RunInEC() && !NodeIsAlreadyComplete(GUBPNodesCompleted, Dep, LocalOnly) && OrdereredToDo.Contains(Dep)) // if something is already finished, we don't put it into EC
 								{
-									throw new AutomationException("Topological sort error, node {0} has a dependency of {1} which sorted after it.", NodeToDo.Name, Dep.Name);
-								}
-								if (!MyChain.Contains(Dep) && !PreConditionUncompletedEcDeps.Contains(Dep.Name))
-								{
-									PreConditionUncompletedEcDeps.Add(Dep.Name);
+									if (OrdereredToDo.IndexOf(Dep) > OrdereredToDo.IndexOf(NodeToDo))
+									{
+										throw new AutomationException("Topological sort error, node {0} has a dependency of {1} which sorted after it.", NodeToDo.Name, Dep.Name);
+									}
+									if (!MyChain.Contains(Dep) && !PreConditionUncompletedEcDeps.Contains(Dep.Name))
+									{
+										PreConditionUncompletedEcDeps.Add(Dep.Name);
+									}
 								}
 							}
 						}
 					}
-				}
-				if (bHasNoop && PreConditionUncompletedEcDeps.Count == 0)
-				{
-					PreConditionUncompletedEcDeps.Add("Noop");
-				}
-				if (PreConditionUncompletedEcDeps.Count > 0)
-				{
-					PreCondition = "\"\\$\" . \"[/javascript if(";
-					// these run "parallel", but we add preconditions to serialize them
-					int Index = 0;
-					foreach (string Dep in PreConditionUncompletedEcDeps)
+					if (bHasNoop && PreConditionUncompletedEcDeps.Count == 0)
+					{
+						PreConditionUncompletedEcDeps.Add("Noop");
+					}
+					if (PreConditionUncompletedEcDeps.Count > 0)
 					{
-						PreCondition = PreCondition + "getProperty('" + GetJobStep(PreconditionParentPath, Dep) + "/status\') == \'completed\'";
-						Index++;
-						if (Index != PreConditionUncompletedEcDeps.Count)
+						PreCondition = "\"\\$\" . \"[/javascript if(";
+						// these run "parallel", but we add preconditions to serialize them
+						int Index = 0;
+						foreach (string Dep in PreConditionUncompletedEcDeps)
 						{
-							PreCondition = PreCondition + " && ";
+							PreCondition = PreCondition + "getProperty('" + GetJobStep(PreconditionParentPath, Dep) + "/status\') == \'completed\'";
+							Index++;
+							if (Index != PreConditionUncompletedEcDeps.Count)
+							{
+								PreCondition = PreCondition + " && ";
+							}
 						}
+						PreCondition = PreCondition + ") true;]\"";
 					}
-					PreCondition = PreCondition + ") true;]\"";
-				}
 
-				if (UncompletedEcDeps.Count > 0)
-				{
-					RunCondition = "\"\\$\" . \"[/javascript if(";
-					int Index = 0;
-					foreach (string Dep in UncompletedEcDeps)
+					if (UncompletedEcDeps.Count > 0)
 					{
-						RunCondition = RunCondition + "((\'\\$\" . \"[" + GetJobStep(PreconditionParentPath, Dep) + "/outcome]\' == \'success\') || ";
-						RunCondition = RunCondition + "(\'\\$\" . \"[" + GetJobStep(PreconditionParentPath, Dep) + "/outcome]\' == \'warning\'))";
-
-						Index++;
-						if (Index != UncompletedEcDeps.Count)
+						RunCondition = "\"\\$\" . \"[/javascript if(";
+						int Index = 0;
+						foreach (string Dep in UncompletedEcDeps)
 						{
-							RunCondition = RunCondition + " && ";
+							RunCondition = RunCondition + "((\'\\$\" . \"[" + GetJobStep(PreconditionParentPath, Dep) + "/outcome]\' == \'success\') || ";
+							RunCondition = RunCondition + "(\'\\$\" . \"[" + GetJobStep(PreconditionParentPath, Dep) + "/outcome]\' == \'warning\'))";
+
+							Index++;
+							if (Index != UncompletedEcDeps.Count)
+							{
+								RunCondition = RunCondition + " && ";
+							}
 						}
+						RunCondition = RunCondition + ")true; else false;]\"";
 					}
-					RunCondition = RunCondition + ")true; else false;]\"";
-				}
 
-				if (bDoNestedJobstep)
-				{
-					if (bDoFirstNestedJobstep)
+					if (bDoNestedJobstep)
 					{
+						if (bDoFirstNestedJobstep)
 						{
-							string NestArgs = String.Format("$batch->createJobStep({{parentPath => '{0}', jobStepName => '{1}', parallel => '1'",
-								ParentPath, MyAgentGroup);
-							if (!String.IsNullOrEmpty(PreCondition))
 							{
-								NestArgs = NestArgs + ", precondition => " + PreCondition;
+								string NestArgs = String.Format("$batch->createJobStep({{parentPath => '{0}', jobStepName => '{1}', parallel => '1'",
+									ParentPath, MyAgentGroup);
+								if (!String.IsNullOrEmpty(PreCondition))
+								{
+									NestArgs = NestArgs + ", precondition => " + PreCondition;
+								}
+								NestArgs = NestArgs + "});";
+								StepList.Add(NestArgs);
 							}
-							NestArgs = NestArgs + "});";
-							StepList.Add(NestArgs);
-						}
-						{
-							string NestArgs = String.Format("$batch->createJobStep({{parentPath => '{0}/jobSteps[{1}]', jobStepName => '{2}_GetPool', subprocedure => 'GUBP{3}_AgentShare_GetPool', parallel => '1', actualParameter => [{{actualParameterName => 'AgentSharingGroup', value => '{4}'}}, {{actualParameterName => 'NodeName', value => '{5}'}}]",
-								ParentPath, MyAgentGroup, MyAgentGroup, NodeToDo.Node.ECProcedureInfix(), MyAgentGroup, NodeToDo.Name);
-							if (!String.IsNullOrEmpty(PreCondition))
 							{
-								NestArgs = NestArgs + ", precondition => " + PreCondition;
+								string NestArgs = String.Format("$batch->createJobStep({{parentPath => '{0}/jobSteps[{1}]', jobStepName => '{2}_GetPool', subprocedure => 'GUBP{3}_AgentShare_GetPool', parallel => '1', actualParameter => [{{actualParameterName => 'AgentSharingGroup', value => '{4}'}}, {{actualParameterName => 'NodeName', value => '{5}'}}]",
+									ParentPath, MyAgentGroup, MyAgentGroup, NodeToDo.Node.ECProcedureInfix(), MyAgentGroup, NodeToDo.Name);
+								if (!String.IsNullOrEmpty(PreCondition))
+								{
+									NestArgs = NestArgs + ", precondition => " + PreCondition;
+								}
+								NestArgs = NestArgs + "});";
+								StepList.Add(NestArgs);
 							}
-							NestArgs = NestArgs + "});";
-							StepList.Add(NestArgs);
-						}
-						{
-							string NestArgs = String.Format("$batch->createJobStep({{parentPath => '{0}/jobSteps[{1}]', jobStepName => '{2}_GetAgent', subprocedure => 'GUBP{3}_AgentShare_GetAgent', parallel => '1', exclusiveMode => 'call', resourceName => '{4}', actualParameter => [{{actualParameterName => 'AgentSharingGroup', value => '{5}'}}, {{actualParameterName => 'NodeName', value=> '{6}'}}]",
-								ParentPath, MyAgentGroup, MyAgentGroup, NodeToDo.Node.ECProcedureInfix(),
-								String.Format("$[/myJob/jobSteps[{0}]/ResourcePool]", MyAgentGroup),
-								MyAgentGroup, NodeToDo.Name);
 							{
-								NestArgs = NestArgs + ", precondition  => ";
-								NestArgs = NestArgs + "\"\\$\" . \"[/javascript if(";
-								NestArgs = NestArgs + "getProperty('" + PreconditionParentPath + "/jobSteps[" + MyAgentGroup + "]/jobSteps[" + MyAgentGroup + "_GetPool]/status') == 'completed'";
-								NestArgs = NestArgs + ") true;]\"";
+								string NestArgs = String.Format("$batch->createJobStep({{parentPath => '{0}/jobSteps[{1}]', jobStepName => '{2}_GetAgent', subprocedure => 'GUBP{3}_AgentShare_GetAgent', parallel => '1', exclusiveMode => 'call', resourceName => '{4}', actualParameter => [{{actualParameterName => 'AgentSharingGroup', value => '{5}'}}, {{actualParameterName => 'NodeName', value=> '{6}'}}]",
+									ParentPath, MyAgentGroup, MyAgentGroup, NodeToDo.Node.ECProcedureInfix(),
+									String.Format("$[/myJob/jobSteps[{0}]/ResourcePool]", MyAgentGroup),
+									MyAgentGroup, NodeToDo.Name);
+								{
+									NestArgs = NestArgs + ", precondition  => ";
+									NestArgs = NestArgs + "\"\\$\" . \"[/javascript if(";
+									NestArgs = NestArgs + "getProperty('" + PreconditionParentPath + "/jobSteps[" + MyAgentGroup + "]/jobSteps[" + MyAgentGroup + "_GetPool]/status') == 'completed'";
+									NestArgs = NestArgs + ") true;]\"";
+								}
+								NestArgs = NestArgs + "});";
+								StepList.Add(NestArgs);
+							}
+							{
+								PreCondition = "\"\\$\" . \"[/javascript if(";
+								PreCondition = PreCondition + "getProperty('" + PreconditionParentPath + "/jobSteps[" + MyAgentGroup + "]/jobSteps[" + MyAgentGroup + "_GetAgent]/status') == 'completed'";
+								PreCondition = PreCondition + ") true;]\"";
 							}
-							NestArgs = NestArgs + "});";
-							StepList.Add(NestArgs);
-						}
-						{
-							PreCondition = "\"\\$\" . \"[/javascript if(";
-							PreCondition = PreCondition + "getProperty('" + PreconditionParentPath + "/jobSteps[" + MyAgentGroup + "]/jobSteps[" + MyAgentGroup + "_GetAgent]/status') == 'completed'";
-							PreCondition = PreCondition + ") true;]\"";
 						}
+						Args = Args.Replace(String.Format("parentPath => '{0}'", ParentPath), String.Format("parentPath => '{0}'", NodeParentPath));
+						Args = Args.Replace("UAT_Node_Parallel_AgentShare", "UAT_Node_Parallel_AgentShare3");
 					}
-					Args = Args.Replace(String.Format("parentPath => '{0}'", ParentPath), String.Format("parentPath => '{0}'", NodeParentPath));
-					Args = Args.Replace("UAT_Node_Parallel_AgentShare", "UAT_Node_Parallel_AgentShare3");
-				}
 
-				if (!String.IsNullOrEmpty(PreCondition))
-				{
-					Args = Args + ", precondition => " + PreCondition;
-				}
-				if (!String.IsNullOrEmpty(RunCondition))
-				{
-					Args = Args + ", condition => " + RunCondition;
-				}
-#if false
-                    // this doesn't work because it includes precondition time
-                    if (GUBPNodes[NodeToDo].TimeoutInMinutes() > 0)
-                    {
-                        Args = Args + String.Format(" --timeLimitUnits minutes --timeLimit {0}", GUBPNodes[NodeToDo].TimeoutInMinutes());
-                    }
-#endif
-				if (Sticky && NodeToDo == LastSticky)
-				{
-					Args = Args + ", releaseMode => 'release'";
-				}
-				Args = Args + "});";
-				StepList.Add(Args);
-				if (bFakeEC &&
-					!UnfinishedTriggers.Contains(NodeToDo) &&
-					(NodeToDo.Node.ECProcedure().StartsWith("GUBP_UAT_Node") || NodeToDo.Node.ECProcedure().StartsWith("GUBP_Mac_UAT_Node")) // other things we really can't test
-					) // unfinished triggers are never run directly by EC, rather it does another job setup
-				{
-					string Arg = String.Format("gubp -Node={0} -FakeEC {1} {2} {3} {4} {5}",
-						NodeToDo,
-						bFake ? "-Fake" : "",
-						ParseParam("AllPlatforms") ? "-AllPlatforms" : "",
-						ParseParam("UnfinishedTriggersFirst") ? "-UnfinishedTriggersFirst" : "",
-						ParseParam("UnfinishedTriggersParallel") ? "-UnfinishedTriggersParallel" : "",
-						ParseParam("WithMac") ? "-WithMac" : ""
-						);
-
-					string Node = ParseParamValue("-Node");
-					if (!String.IsNullOrEmpty(Node))
+					if (!String.IsNullOrEmpty(PreCondition))
 					{
-						Arg = Arg + " -Node=" + Node;
+						Args = Args + ", precondition => " + PreCondition;
 					}
-					if (!String.IsNullOrEmpty(FakeFail))
+					if (!String.IsNullOrEmpty(RunCondition))
 					{
-						Arg = Arg + " -FakeFail=" + FakeFail;
+						Args = Args + ", condition => " + RunCondition;
+					}
+	#if false
+						// this doesn't work because it includes precondition time
+						if (GUBPNodes[NodeToDo].TimeoutInMinutes() > 0)
+						{
+							Args = Args + String.Format(" --timeLimitUnits minutes --timeLimit {0}", GUBPNodes[NodeToDo].TimeoutInMinutes());
+						}
+	#endif
+					if (Sticky && NodeToDo == LastSticky)
+					{
+						Args = Args + ", releaseMode => 'release'";
+					}
+					Args = Args + "});";
+					StepList.Add(Args);
+					if (bFakeEC &&
+						!UnfinishedTriggers.Contains(NodeToDo) &&
+						(NodeToDo.Node.ECProcedure().StartsWith("GUBP_UAT_Node") || NodeToDo.Node.ECProcedure().StartsWith("GUBP_Mac_UAT_Node")) // other things we really can't test
+						) // unfinished triggers are never run directly by EC, rather it does another job setup
+					{
+						string Arg = String.Format("gubp -Node={0} -FakeEC {1} {2} {3} {4} {5}",
+							NodeToDo,
+							bFake ? "-Fake" : "",
+							ParseParam("AllPlatforms") ? "-AllPlatforms" : "",
+							ParseParam("UnfinishedTriggersFirst") ? "-UnfinishedTriggersFirst" : "",
+							ParseParam("UnfinishedTriggersParallel") ? "-UnfinishedTriggersParallel" : "",
+							ParseParam("WithMac") ? "-WithMac" : ""
+							);
+
+						string Node = ParseParamValue("-Node");
+						if (!String.IsNullOrEmpty(Node))
+						{
+							Arg = Arg + " -Node=" + Node;
+						}
+						if (!String.IsNullOrEmpty(FakeFail))
+						{
+							Arg = Arg + " -FakeFail=" + FakeFail;
+						}
+						FakeECArgs.Add(Arg);
 					}
-					FakeECArgs.Add(Arg);
-				}
 
-				if (MyAgentGroup != "" && !bDoNestedJobstep)
-				{
-					List<NodeInfo> MyChain = AgentGroupChains[MyAgentGroup];
-					int MyIndex = MyChain.IndexOf(NodeToDo);
-					if (MyIndex == MyChain.Count - 1)
+					if (MyAgentGroup != "" && !bDoNestedJobstep)
 					{
-						string RelPreCondition = "\"\\$\" . \"[/javascript if(";
-						// this runs "parallel", but we a precondition to serialize it
-						RelPreCondition = RelPreCondition + "getProperty('" + PreconditionParentPath + "/jobSteps[" + NodeToDo.Name + "]/status') == 'completed'";
-						RelPreCondition = RelPreCondition + ") true;]\"";
-						// we need to release the resource
-						string RelArgs = String.Format("{0}, subprocedure => 'GUBP_Release_AgentShare', parallel => '1', jobStepName => 'Release_{1}', actualParameter => [{{actualParameterName => 'AgentSharingGroup', valued => '{2}'}}], releaseMode => 'release', precondition => '{3}'",
-							BaseArgs, MyAgentGroup, MyAgentGroup, RelPreCondition);
-						StepList.Add(RelArgs);
+						List<NodeInfo> MyChain = AgentGroupChains[MyAgentGroup];
+						int MyIndex = MyChain.IndexOf(NodeToDo);
+						if (MyIndex == MyChain.Count - 1)
+						{
+							string RelPreCondition = "\"\\$\" . \"[/javascript if(";
+							// this runs "parallel", but we a precondition to serialize it
+							RelPreCondition = RelPreCondition + "getProperty('" + PreconditionParentPath + "/jobSteps[" + NodeToDo.Name + "]/status') == 'completed'";
+							RelPreCondition = RelPreCondition + ") true;]\"";
+							// we need to release the resource
+							string RelArgs = String.Format("{0}, subprocedure => 'GUBP_Release_AgentShare', parallel => '1', jobStepName => 'Release_{1}', actualParameter => [{{actualParameterName => 'AgentSharingGroup', valued => '{2}'}}], releaseMode => 'release', precondition => '{3}'",
+								BaseArgs, MyAgentGroup, MyAgentGroup, RelPreCondition);
+							StepList.Add(RelArgs);
+						}
 					}
 				}
 			}
+			WriteECPerl(StepList);
 		}
-		WriteECPerl(StepList);
-		string FinishPerlOutput = DateTime.Now.ToString();
-		PrintCSVFile(String.Format("UAT,PerlOutput,{0},{1}", StartPerlOutput, FinishPerlOutput));
 		RunECTool(String.Format("setProperty \"/myWorkflow/HasTests\" \"{0}\"", bHasTests));
 		{
 			ECProps.Add("GUBP_LoadedProps=1");
@@ -2847,13 +2847,10 @@ public partial class GUBP : BuildCommand
                     {                        
 						Log("***** Building GUBP Node {0} for {1}", NodeToDo.Name, NodeStoreName);
 						DateTime StartTime = DateTime.UtcNow;
-                        string StartBuild = DateTime.Now.ToString();
-                        NodeToDo.Node.DoBuild(this);
-                        string FinishBuild = DateTime.Now.ToString();
-                        if (IsBuildMachine)
-                        {
-                            PrintCSVFile(String.Format("UAT,DoBuild.{0},{1},{2}", NodeToDo.Name, StartBuild, FinishBuild));
-                        }
+						using(TelemetryStopwatch DoBuildStopwatch = new TelemetryStopwatch("DoBuild.{0}", NodeToDo.Name))
+						{
+							NodeToDo.Node.DoBuild(this);
+						}
 						BuildDuration = (DateTime.UtcNow - StartTime).TotalMilliseconds / 1000;
 						
                     }
@@ -2861,14 +2858,14 @@ public partial class GUBP : BuildCommand
                     {
 						double StoreDuration = 0.0;
 						DateTime StartTime = DateTime.UtcNow;
-                        string StartStore = DateTime.Now.ToString();
-                        TempStorage.StoreToTempStorage(CmdEnv, NodeStoreName, NodeToDo.Node.BuildProducts, !bSaveSharedTempStorage, GameNameIfAny, StorageRootIfAny);
+						using(TelemetryStopwatch StoreBuildProductsStopwatch = new TelemetryStopwatch("StoreBuildProducts"))
+						{
+							TempStorage.StoreToTempStorage(CmdEnv, NodeStoreName, NodeToDo.Node.BuildProducts, !bSaveSharedTempStorage, GameNameIfAny, StorageRootIfAny);
+						}
 						StoreDuration = (DateTime.UtcNow - StartTime).TotalMilliseconds / 1000;
-                        string FinishStore = DateTime.Now.ToString();                        
 						Log("Took {0} seconds to store build products", StoreDuration);
 						if (IsBuildMachine)
 						{
-                            PrintCSVFile(String.Format("UAT,StoreBuildProducts,{0},{1}", StartStore, FinishStore));
 							RunECTool(String.Format("setProperty \"/myJobStep/StoreDuration\" \"{0}\"", StoreDuration.ToString()));
 						}
                         if (ParseParam("StompCheck"))
@@ -2878,13 +2875,10 @@ public partial class GUBP : BuildCommand
                                 try
                                 {
                                     bool WasLocal;
-                                    string StartRetrieve = DateTime.Now.ToString();
-									TempStorage.RetrieveFromTempStorage(CmdEnv, NodeStoreName, out WasLocal, GameNameIfAny, StorageRootIfAny);
-                                    string FinishRetrieve = DateTime.Now.ToString();
-                                    if (IsBuildMachine)
-                                    {
-                                        PrintCSVFile(String.Format("UAT,RetrieveBuildProducts,{0},{1}", StartRetrieve, FinishRetrieve));
-                                    }
+									using(TelemetryStopwatch RetrieveBuildProductsStopwatch = new TelemetryStopwatch("RetrieveBuildProducts"))
+									{
+										TempStorage.RetrieveFromTempStorage(CmdEnv, NodeStoreName, out WasLocal, GameNameIfAny, StorageRootIfAny);
+									}
 									if (!WasLocal)
 									{
 										throw new AutomationException("Retrieve was not local?");
@@ -2904,22 +2898,25 @@ public partial class GUBP : BuildCommand
                 {
                     if (SaveSuccessRecords)
                     {
-                        string StartNodeHistory = DateTime.Now.ToString();
-                        UpdateNodeHistory(GUBPNodesHistory, NodeToDo, CLString);
-                        string FinishNodeHistory = DateTime.Now.ToString();                        
-                        SaveStatus(NodeToDo, FailedTempStorageSuffix, NodeStoreName, bSaveSharedTempStorage, GameNameIfAny, ParseParamValue("MyJobStepId"));
-                        string FinishSaveStatus = DateTime.Now.ToString();                        
-                        UpdateECProps(NodeToDo, CLString);
-                        string FinishECPropUpdate = DateTime.Now.ToString();
+						using(TelemetryStopwatch UpdateNodeHistoryStopwatch = new TelemetryStopwatch("UpdateNodeHistory"))
+						{
+							UpdateNodeHistory(GUBPNodesHistory, NodeToDo, CLString);
+						}
+						using(TelemetryStopwatch SaveNodeStatusStopwatch = new TelemetryStopwatch("SaveNodeStatus"))
+						{
+							SaveStatus(NodeToDo, FailedTempStorageSuffix, NodeStoreName, bSaveSharedTempStorage, GameNameIfAny, ParseParamValue("MyJobStepId"));
+						}
+						using(TelemetryStopwatch UpdateECPropsStopwatch = new TelemetryStopwatch("UpdateECProps"))
+						{
+							UpdateECProps(NodeToDo, CLString);
+						}
                         
 						if (IsBuildMachine)
 						{
-							GetFailureEmails(GUBPNodesHistory, NodeToDo, CLString);
-                            string FinishFailEmails = DateTime.Now.ToString();
-                            PrintCSVFile(String.Format("UAT,UpdateNodeHistory,{0},{1}", StartNodeHistory, FinishNodeHistory));
-                            PrintCSVFile(String.Format("UAT,SaveNodeStatus,{0},{1}", FinishNodeHistory, FinishSaveStatus));
-                            PrintCSVFile(String.Format("UAT,UpdateECProps,{0},{1}", FinishSaveStatus, FinishECPropUpdate));
-                            PrintCSVFile(String.Format("UAT,GetFailEmails,{0},{1}", FinishECPropUpdate, FinishFailEmails));
+							using(TelemetryStopwatch GetFailEmailsStopwatch = new TelemetryStopwatch("GetFailEmails"))
+							{
+								GetFailureEmails(GUBPNodesHistory, NodeToDo, CLString);
+							}
 						}
 						UpdateECBuildTime(NodeToDo, BuildDuration);
                     }
@@ -2971,22 +2968,25 @@ public partial class GUBP : BuildCommand
                 }
                 if (SaveSuccessRecords) 
                 {
-                    string StartNodeHistory = DateTime.Now.ToString();
-                    UpdateNodeHistory(GUBPNodesHistory, NodeToDo, CLString);
-                    string FinishNodeHistory = DateTime.Now.ToString();
-                    SaveStatus(NodeToDo, SucceededTempStorageSuffix, NodeStoreName, bSaveSharedTempStorage, GameNameIfAny);
-                    string FinishSaveStatus = DateTime.Now.ToString();                    
-                    UpdateECProps(NodeToDo, CLString);
-                    string FinishECPropUpdate = DateTime.Now.ToString();
+					using(TelemetryStopwatch UpdateNodeHistoryStopwatch = new TelemetryStopwatch("UpdateNodeHistory"))
+					{
+						UpdateNodeHistory(GUBPNodesHistory, NodeToDo, CLString);
+					}
+					using(TelemetryStopwatch SaveNodeStatusStopwatch = new TelemetryStopwatch("SaveNodeStatus"))
+					{
+						SaveStatus(NodeToDo, SucceededTempStorageSuffix, NodeStoreName, bSaveSharedTempStorage, GameNameIfAny);
+					}
+					using(TelemetryStopwatch UpdateECPropsStopwatch = new TelemetryStopwatch("UpdateECProps"))
+					{
+						UpdateECProps(NodeToDo, CLString);
+					}
                     
 					if (IsBuildMachine)
 					{
-						GetFailureEmails(GUBPNodesHistory, NodeToDo, CLString);
-                        string FinishFailEmails = DateTime.Now.ToString();
-                        PrintCSVFile(String.Format("UAT,UpdateNodeHistory,{0},{1}", StartNodeHistory, FinishNodeHistory));
-                        PrintCSVFile(String.Format("UAT,SaveNodeStatus,{0},{1}", FinishNodeHistory, FinishSaveStatus));
-                        PrintCSVFile(String.Format("UAT,UpdateECProps,{0},{1}", FinishSaveStatus, FinishECPropUpdate));
-                        PrintCSVFile(String.Format("UAT,GetFailEmails,{0},{1}", FinishECPropUpdate, FinishFailEmails));
+						using(TelemetryStopwatch GetFailEmailsStopwatch = new TelemetryStopwatch("GetFailEmails"))
+						{
+							GetFailureEmails(GUBPNodesHistory, NodeToDo, CLString);
+						}
 					}
 					UpdateECBuildTime(NodeToDo, BuildDuration);
                 }
diff --git a/Engine/Source/Programs/AutomationTool/GUBP/LegacyNodes.cs b/Engine/Source/Programs/AutomationTool/GUBP/LegacyNodes.cs
--- a/Engine/Source/Programs/AutomationTool/GUBP/LegacyNodes.cs
+++ b/Engine/Source/Programs/AutomationTool/GUBP/LegacyNodes.cs
@@ -379,10 +379,10 @@ partial class GUBP
                 Agenda.DoRetries = false; // these would delete build products
 				bool UseParallelExecutor = bDependentOnCompileTools && (HostPlatform == UnrealTargetPlatform.Win64);
 				UE4Build.Build(Agenda, InDeleteBuildProducts: ReallyDeleteBuildProducts, InUpdateVersionFiles: false, InForceNoXGE: true, InForceUnity: true, InUseParallelExecutor: UseParallelExecutor);
-                var StartPostBuild = DateTime.Now.ToString();
-                PostBuild(bp, UE4Build);
-                var FinishPostBuild = DateTime.Now.ToString();
-                PrintCSVFile(String.Format("UAT,PostBuild,{0},{1}", StartPostBuild, FinishPostBuild));
+				using(TelemetryStopwatch PostBuildStopwatch = new TelemetryStopwatch("PostBuild"))
+				{
+					PostBuild(bp, UE4Build);
+				}
                 UE4Build.CheckBuildProducts(UE4Build.BuildProductFiles);
                 foreach (var Product in UE4Build.BuildProductFiles)
                 {
@@ -392,10 +392,10 @@ partial class GUBP
                 if (bp.bSignBuildProducts)
                 {
                     // Sign everything we built
-                    var StartSign = DateTime.Now.ToString();
-                    CodeSign.SignMultipleIfEXEOrDLL(bp, BuildProducts);
-                    var FinishSign = DateTime.Now.ToString();
-                    PrintCSVFile(String.Format("UAT,SignBuildProducts,{0},{1}", StartSign, FinishSign));
+					using(TelemetryStopwatch SignStopwatch = new TelemetryStopwatch("SignBuildProducts"))
+					{
+						CodeSign.SignMultipleIfEXEOrDLL(bp, BuildProducts);
+					}
                 }
                 PostBuildProducts(bp);
             }
@@ -2389,10 +2389,10 @@ partial class GUBP
             {
                 // not sure if we need something here or if the cook commandlet will automatically convert the exe name
             }
-            var StartCook = DateTime.Now.ToString();
-			CommandUtils.CookCommandlet(GameProj.FilePath, "UE4Editor-Cmd.exe", null, null, null, null, CookPlatform);
-            var FinishCook = DateTime.Now.ToString();
-            PrintCSVFile(String.Format("UAT,Cook.{0}.{1},{2},{3}", GameProj.GameName, CookPlatform, StartCook, FinishCook));
+			using(TelemetryStopwatch CookStopwatch = new TelemetryStopwatch("Cook.{0}.{1}", GameProj.GameName, CookPlatform))
+			{
+				CommandUtils.CookCommandlet(GameProj.FilePath, "UE4Editor-Cmd.exe", null, null, null, null, CookPlatform);
+			}
             var CookedPath = RootIfAnyForTempStorage();
             var CookedFiles = CommandUtils.FindFiles("*", true, CookedPath);
             if (CookedFiles.GetLength(0) < 1)
