Change GUBPCompletedNodes dictionary into a member on each NodeInfo instead.

diff --git a/Engine/Source/Programs/AutomationTool/GUBP/GUBP.cs b/Engine/Source/Programs/AutomationTool/GUBP/GUBP.cs
--- a/Engine/Source/Programs/AutomationTool/GUBP/GUBP.cs
+++ b/Engine/Source/Programs/AutomationTool/GUBP/GUBP.cs
@@ -38,6 +38,7 @@ public partial class GUBP : BuildCommand
 		public List<NodeInfo> PseudoDependencies;
 		public NodeInfo ControllingTrigger;
 		public string ControllingTriggerDotName;
+		public bool IsComplete;
 
 		public override string ToString()
 		{
@@ -163,6 +164,14 @@ public partial class GUBP : BuildCommand
         return Deps.Contains(Rootward);
     }
 
+	static void FindControllingTriggers(IEnumerable<NodeInfo> NodesToDo)
+	{
+		foreach(NodeInfo NodeToDo in NodesToDo)
+		{
+			FindControllingTriggers(NodeToDo);
+		}
+	}
+
     static void FindControllingTriggers(NodeInfo NodeToDo)
     {
 		if(NodeToDo.ControllingTriggerDotName == null)
@@ -271,38 +280,38 @@ public partial class GUBP : BuildCommand
         return Result;
     }
 
-    bool NodeIsAlreadyComplete(Dictionary<string, bool> GUBPNodesCompleted, NodeInfo NodeToDo, bool LocalOnly)
-    {
-        if (GUBPNodesCompleted.ContainsKey(NodeToDo.Name))
-        {
-            return GUBPNodesCompleted[NodeToDo.Name];
-        }
-        string NodeStoreName = StoreName + "-" + NodeToDo.Node.GetFullName();
-        string GameNameIfAny = NodeToDo.Node.GameNameIfAnyForTempStorage();
-        bool Result;
-        if (LocalOnly)
-        {
-            Result = TempStorage.LocalTempStorageExists(CmdEnv, NodeStoreName, bQuiet : true);
-        }
-        else
-        {
-            Result = TempStorage.TempStorageExists(CmdEnv, NodeStoreName, GameNameIfAny, bQuiet: true);
-			if(GameNameIfAny != "" && Result == false)
+	static void FindCompletionState(IEnumerable<NodeInfo> NodesToDo, string StoreName, bool LocalOnly)
+	{
+		foreach(NodeInfo NodeToDo in NodesToDo)
+		{
+			string NodeStoreName = StoreName + "-" + NodeToDo.Name;
+			string GameNameIfAny = NodeToDo.Node.GameNameIfAnyForTempStorage();
+
+			if (LocalOnly)
 			{
-				Result = TempStorage.TempStorageExists(CmdEnv, NodeStoreName, "", bQuiet: true);
+				NodeToDo.IsComplete = TempStorage.LocalTempStorageExists(CmdEnv, NodeStoreName, bQuiet : true);
 			}
-        }
-        if (Result)
-        {
-            LogVerbose("***** GUBP Trigger Node was already triggered {0} -> {1} : {2}", NodeToDo.Name, GameNameIfAny, NodeStoreName);
-        }
-        else
-        {
-            LogVerbose("***** GUBP Trigger Node was NOT yet triggered {0} -> {1} : {2}", NodeToDo.Name, GameNameIfAny, NodeStoreName);
-        }
-        GUBPNodesCompleted.Add(NodeToDo.Name, Result);
-        return Result;
+			else
+			{
+				NodeToDo.IsComplete = TempStorage.TempStorageExists(CmdEnv, NodeStoreName, GameNameIfAny, bQuiet: true);
+				if(GameNameIfAny != "" && !NodeToDo.IsComplete)
+				{
+					NodeToDo.IsComplete = TempStorage.TempStorageExists(CmdEnv, NodeStoreName, "", bQuiet: true);
+				}
+			}
+
+            LogVerbose("** {0}", NodeToDo.Name);
+			if (!NodeToDo.IsComplete)
+			{
+				LogVerbose("***** GUBP Trigger Node was already triggered {0} -> {1} : {2}", NodeToDo.Name, GameNameIfAny, NodeStoreName);
+			}
+			else
+			{
+				LogVerbose("***** GUBP Trigger Node was NOT yet triggered {0} -> {1} : {2}", NodeToDo.Name, GameNameIfAny, NodeStoreName);
+			}
+		}
     }
+
     string RunECTool(string Args, bool bQuiet = false)
     {
         if (ParseParam("FakeEC"))
@@ -437,7 +446,7 @@ public partial class GUBP : BuildCommand
         Log("Took {0}s to get P4 history", BuildDuration / 1000);
 
     }
-    void PrintNodes(GUBP bp, List<NodeInfo> Nodes, Dictionary<string, bool> GUBPNodesCompleted, Dictionary<string, NodeHistory> GUBPNodesHistory, bool LocalOnly, List<NodeInfo> UnfinishedTriggers = null)
+    void PrintNodes(GUBP bp, List<NodeInfo> Nodes, Dictionary<string, NodeHistory> GUBPNodesHistory, List<NodeInfo> UnfinishedTriggers = null)
     {
         bool bShowAllChanges = bp.ParseParam("AllChanges") && GUBPNodesHistory != null;
         bool bShowChanges = (bp.ParseParam("Changes") && GUBPNodesHistory != null) || bShowAllChanges;
@@ -515,7 +524,7 @@ public partial class GUBP : BuildCommand
                 (LastAgentGroup != "" ? "  " : ""),
                 NodeToDo.Name,
                 FrequencyString,
-                NodeIsAlreadyComplete(GUBPNodesCompleted, NodeToDo, LocalOnly) ? " - (Completed)" : "",
+                NodeToDo.IsComplete ? " - (Completed)" : "",
                 NodeToDo.Node.TriggerNode() ? " - (TriggerNode)" : "",
                 NodeToDo.Node.IsSticky() ? " - (Sticky)" : "",				
                 Agent,
@@ -719,7 +728,7 @@ public partial class GUBP : BuildCommand
         return Result;
     }
 
-    List<NodeInfo> TopologicalSort(HashSet<NodeInfo> NodesToDo, Dictionary<string, bool> GUBPNodesCompleted, NodeInfo ExplicitTrigger = null, bool LocalOnly = false, bool SubSort = false, bool DoNotConsiderCompletion = false)
+    List<NodeInfo> TopologicalSort(HashSet<NodeInfo> NodesToDo, NodeInfo ExplicitTrigger, bool SubSort, bool DoNotConsiderCompletion)
     {
         DateTime StartTime = DateTime.UtcNow;
 
@@ -746,7 +755,7 @@ public partial class GUBP : BuildCommand
             }
             foreach (KeyValuePair<string, List<NodeInfo>> Chain in AgentGroupChains)
             {
-                SortedAgentGroupChains.Add(Chain.Key, TopologicalSort(new HashSet<NodeInfo>(Chain.Value), GUBPNodesCompleted, ExplicitTrigger, LocalOnly, true, DoNotConsiderCompletion));
+                SortedAgentGroupChains.Add(Chain.Key, TopologicalSort(new HashSet<NodeInfo>(Chain.Value), ExplicitTrigger, true, DoNotConsiderCompletion));
             }
             Log("***************Done with recursion");
         }
@@ -844,8 +853,8 @@ public partial class GUBP : BuildCommand
                             }
                             else if (BestPseudoReady == bPseudoReady)
                             {
-                                bool IamLateTrigger = !DoNotConsiderCompletion && NodeToDo.Node.TriggerNode() && NodeToDo != ExplicitTrigger && !NodeIsAlreadyComplete(GUBPNodesCompleted, NodeToDo, LocalOnly);
-                                bool BestIsLateTrigger = !DoNotConsiderCompletion && GUBPNodes[BestNode.Name].Node.TriggerNode() && BestNode != ExplicitTrigger && !NodeIsAlreadyComplete(GUBPNodesCompleted, BestNode, LocalOnly);
+                                bool IamLateTrigger = !DoNotConsiderCompletion && NodeToDo.Node.TriggerNode() && NodeToDo != ExplicitTrigger && !NodeToDo.IsComplete;
+                                bool BestIsLateTrigger = !DoNotConsiderCompletion && GUBPNodes[BestNode.Name].Node.TriggerNode() && BestNode != ExplicitTrigger && !BestNode.IsComplete;
                                 if (BestIsLateTrigger && !IamLateTrigger)
                                 {
                                     bReady = false;
@@ -1462,11 +1471,8 @@ public partial class GUBP : BuildCommand
 		AddNodesForBranch(TimeIndex, bNoAutomatedTesting);
 
 		LinkGraph(GUBPNodes);
-
-		foreach(NodeInfo NodeToDo in GUBPNodes.Values)
-		{
-			FindControllingTriggers(NodeToDo);
-		}
+		FindControllingTriggers(GUBPNodes.Values);
+		FindCompletionState(GUBPNodes.Values, StoreName, LocalOnly);
 
         if (bCleanLocalTempStorage)  // shared temp storage can never be wiped
         {
@@ -1590,20 +1596,8 @@ public partial class GUBP : BuildCommand
 			NodesToDo = NewNodesToDo;
 		}
 
-        Dictionary<string, bool> GUBPNodesCompleted = new Dictionary<string, bool>();
         Dictionary<string, NodeHistory> GUBPNodesHistory = new Dictionary<string, NodeHistory>();
 
-        LogVerbose("******* Caching completion");
-        {
-            DateTime StartTime = DateTime.UtcNow;
-            foreach (NodeInfo NodeToDo in NodesToDo)
-            {
-                LogVerbose("** {0}", NodeToDo.Name);
-                NodeIsAlreadyComplete(GUBPNodesCompleted, NodeToDo, LocalOnly); // cache these now to avoid spam later
-            }
-            double BuildDuration = (DateTime.UtcNow - StartTime).TotalMilliseconds;
-			LogVerbose("Took {0}s to cache completion for {1} nodes", BuildDuration / 1000, NodesToDo.Count);
-        }
         /*if (CLString != "" && StoreName.Contains(CLString) && !ParseParam("NoHistory"))
         {
             Log("******* Updating history");
@@ -1619,20 +1613,20 @@ public partial class GUBP : BuildCommand
             Log("Took {0}s to get history for {1} nodes", BuildDuration / 1000, NodesToDo.Count);
         }*/
 
-        List<NodeInfo> OrdereredToDo = TopologicalSort(NodesToDo, GUBPNodesCompleted, ExplicitTrigger, LocalOnly);
+        List<NodeInfo> OrdereredToDo = TopologicalSort(NodesToDo, ExplicitTrigger, false, false);
 
-		List<NodeInfo> UnfinishedTriggers = FindUnfinishedTriggers(bSkipTriggers, LocalOnly, ExplicitTrigger, GUBPNodesCompleted, OrdereredToDo);
+		List<NodeInfo> UnfinishedTriggers = FindUnfinishedTriggers(bSkipTriggers, ExplicitTrigger, OrdereredToDo);
 
         LogVerbose("*********** Desired And Dependent Nodes, in order.");
-        PrintNodes(this, OrdereredToDo, GUBPNodesCompleted, GUBPNodesHistory,  LocalOnly, UnfinishedTriggers);		
+        PrintNodes(this, OrdereredToDo, GUBPNodesHistory,  UnfinishedTriggers);		
         //check sorting
-		CheckSortOrder(LocalOnly, GUBPNodesCompleted, OrdereredToDo);
+		CheckSortOrder(OrdereredToDo);
 
         string FakeFail = ParseParamValue("FakeFail");
         if (CommanderSetup)
         {
-			DoCommanderSetup(TimeIndex, bSkipTriggers, bFakeEC, LocalOnly, CLString, ExplicitTrigger, FullNodeList, FullNodeDirectDependencies, FullNodeDependedOnBy, FullNodeDependentPromotions, SeparatePromotables, FullNodeListSortKey, GUBPNodesCompleted, GUBPNodesHistory, OrdereredToDo, UnfinishedTriggers, FakeFail);
-            Log("Commander setup only, done.");            
+			DoCommanderSetup(TimeIndex, bSkipTriggers, bFakeEC, CLString, ExplicitTrigger, FullNodeList, FullNodeDirectDependencies, FullNodeDependedOnBy, FullNodeDependentPromotions, SeparatePromotables, FullNodeListSortKey, GUBPNodesHistory, OrdereredToDo, UnfinishedTriggers, FakeFail);
+            Log("Commander setup only, done.");
             PrintRunTime();
             return;
 
@@ -1646,7 +1640,7 @@ public partial class GUBP : BuildCommand
             Log("List only, done.");
             return;
         }    
-		ExecuteNodes(OrdereredToDo, bOnlyNode, bFakeEC, LocalOnly, bSaveSharedTempStorage, GUBPNodesCompleted, GUBPNodesHistory, CLString, FakeFail);
+		ExecuteNodes(OrdereredToDo, bOnlyNode, bFakeEC, bSaveSharedTempStorage, GUBPNodesHistory, CLString, FakeFail);
 	}
 
 	private static void LinkGraph(Dictionary<string, NodeInfo> GUBPNodes)
@@ -2043,7 +2037,7 @@ public partial class GUBP : BuildCommand
 			FullNodeDirectDependencies = new Dictionary<string,string>();
 
 			Log("******* {0} GUBP Nodes", GUBPNodes.Count);
-			List<NodeInfo> SortedNodes = TopologicalSort(new HashSet<NodeInfo>(GUBPNodes.Values), null/*GUBPNodesCompleted*/, LocalOnly: true, DoNotConsiderCompletion: true);
+			List<NodeInfo> SortedNodes = TopologicalSort(new HashSet<NodeInfo>(GUBPNodes.Values), null, SubSort: false, DoNotConsiderCompletion: true);
 			foreach (NodeInfo Node in SortedNodes)
 			{
 				string Note = Node.ControllingTriggerDotName;
@@ -2197,7 +2191,7 @@ public partial class GUBP : BuildCommand
 		throw new AutomationException("Failed to update the CIS counter after 20 tries.");
 	}
 
-	private List<NodeInfo> FindUnfinishedTriggers(bool bSkipTriggers, bool LocalOnly, NodeInfo ExplicitTrigger, Dictionary<string, bool> GUBPNodesCompleted, List<NodeInfo> OrdereredToDo)
+	private List<NodeInfo> FindUnfinishedTriggers(bool bSkipTriggers, NodeInfo ExplicitTrigger, List<NodeInfo> OrdereredToDo)
 	{
 		// find all unfinished triggers, excepting the one we are triggering right now
 		List<NodeInfo> UnfinishedTriggers = new List<NodeInfo>();
@@ -2205,7 +2199,7 @@ public partial class GUBP : BuildCommand
 		{
 			foreach (NodeInfo NodeToDo in OrdereredToDo)
 			{
-				if (NodeToDo.Node.TriggerNode() && !NodeIsAlreadyComplete(GUBPNodesCompleted, NodeToDo, LocalOnly))
+				if (NodeToDo.Node.TriggerNode() && !NodeToDo.IsComplete)
 				{
 					if (ExplicitTrigger != NodeToDo)
 					{
@@ -2217,11 +2211,11 @@ public partial class GUBP : BuildCommand
 		return UnfinishedTriggers;
 	}
 
-	private void CheckSortOrder(bool LocalOnly, Dictionary<string, bool> GUBPNodesCompleted, List<NodeInfo> OrdereredToDo)
+	private void CheckSortOrder(List<NodeInfo> OrdereredToDo)
 	{
 		foreach (NodeInfo NodeToDo in OrdereredToDo)
 		{
-			if (NodeToDo.Node.TriggerNode() && (NodeToDo.Node.IsSticky() || NodeIsAlreadyComplete(GUBPNodesCompleted, NodeToDo, LocalOnly))) // these sticky triggers are ok, everything is already completed anyway
+			if (NodeToDo.Node.TriggerNode() && (NodeToDo.Node.IsSticky() || NodeToDo.IsComplete)) // these sticky triggers are ok, everything is already completed anyway
 			{
 				continue;
 			}
@@ -2249,7 +2243,7 @@ public partial class GUBP : BuildCommand
 		}
 	}
 
-	private void DoCommanderSetup(int TimeIndex, bool bSkipTriggers, bool bFakeEC, bool LocalOnly, string CLString, NodeInfo ExplicitTrigger, Dictionary<string, string> FullNodeList, Dictionary<string, string> FullNodeDirectDependencies, Dictionary<string, string> FullNodeDependedOnBy, Dictionary<string, string> FullNodeDependentPromotions, List<string> SeparatePromotables, Dictionary<string, int> FullNodeListSortKey, Dictionary<string, bool> GUBPNodesCompleted, Dictionary<string, NodeHistory> GUBPNodesHistory, List<NodeInfo> OrdereredToDo, List<NodeInfo> UnfinishedTriggers, string FakeFail)
+	private void DoCommanderSetup(int TimeIndex, bool bSkipTriggers, bool bFakeEC, string CLString, NodeInfo ExplicitTrigger, Dictionary<string, string> FullNodeList, Dictionary<string, string> FullNodeDirectDependencies, Dictionary<string, string> FullNodeDependedOnBy, Dictionary<string, string> FullNodeDependentPromotions, List<string> SeparatePromotables, Dictionary<string, int> FullNodeListSortKey, Dictionary<string, NodeHistory> GUBPNodesHistory, List<NodeInfo> OrdereredToDo, List<NodeInfo> UnfinishedTriggers, string FakeFail)
 	{
 
 		if (OrdereredToDo.Count == 0)
@@ -2324,14 +2318,14 @@ public partial class GUBP : BuildCommand
 		using(TelemetryStopwatch PrintNodesTimer = new TelemetryStopwatch("SetupCommanderPrint"))
 		{
 			Log("*********** EC Nodes, in order.");
-			PrintNodes(this, OrdereredToDo, GUBPNodesCompleted, GUBPNodesHistory, LocalOnly, UnfinishedTriggers);
+			PrintNodes(this, OrdereredToDo, GUBPNodesHistory, UnfinishedTriggers);
 		}
 		// here we are just making sure everything before the explicit trigger is completed.
 		if (ExplicitTrigger != null)
 		{
 			foreach (NodeInfo NodeToDo in OrdereredToDo)
 			{
-				if (NodeToDo.Node.RunInEC() && !NodeIsAlreadyComplete(GUBPNodesCompleted, NodeToDo, LocalOnly) && NodeToDo != ExplicitTrigger && !NodeDependsOn(ExplicitTrigger, NodeToDo)) // if something is already finished, we don't put it into EC
+				if (NodeToDo.Node.RunInEC() && !NodeToDo.IsComplete && NodeToDo != ExplicitTrigger && !NodeDependsOn(ExplicitTrigger, NodeToDo)) // if something is already finished, we don't put it into EC
 				{
 					throw new AutomationException("We are being asked to process node {0}, however, this is an explicit trigger {1}, so everything before it should already be handled. It seems likely that you waited too long to run the trigger. You will have to do a new build from scratch.", NodeToDo.Name, ExplicitTrigger.Name);
 				}
@@ -2351,7 +2345,7 @@ public partial class GUBP : BuildCommand
 		// sticky nodes are ones that we run on the main agent. We run then first and they must not be intermixed with parallel jobs
 		foreach (NodeInfo NodeToDo in OrdereredToDo)
 		{
-			if (NodeToDo.Node.RunInEC() && !NodeIsAlreadyComplete(GUBPNodesCompleted, NodeToDo, LocalOnly)) // if something is already finished, we don't put it into EC
+			if (NodeToDo.Node.RunInEC() && !NodeToDo.IsComplete) // if something is already finished, we don't put it into EC
 			{
 				bHaveECNodes = true;
 				if (NodeToDo.Node.IsSticky())
@@ -2388,7 +2382,7 @@ public partial class GUBP : BuildCommand
 			List<NodeInfo> StickyChain = new List<NodeInfo>();
 			foreach (NodeInfo NodeToDo in OrdereredToDo)
 			{
-				if (NodeToDo.Node.RunInEC() && !NodeIsAlreadyComplete(GUBPNodesCompleted, NodeToDo, LocalOnly)) // if something is already finished, we don't put it into EC  
+				if (NodeToDo.Node.RunInEC() && !NodeToDo.IsComplete) // if something is already finished, we don't put it into EC  
 				{
 					string MyAgentGroup = NodeToDo.Node.AgentSharingGroup;
 					if (MyAgentGroup != "")
@@ -2413,7 +2407,7 @@ public partial class GUBP : BuildCommand
 			}
 			foreach (NodeInfo NodeToDo in OrdereredToDo)
 			{
-				if (NodeToDo.Node.RunInEC() && !NodeIsAlreadyComplete(GUBPNodesCompleted, NodeToDo, LocalOnly)) // if something is already finished, we don't put it into EC  
+				if (NodeToDo.Node.RunInEC() && !NodeToDo.IsComplete) // if something is already finished, we don't put it into EC  
 				{
 					string EMails;
 					List<string> NodeProps = GetECPropsForNode(NodeToDo, CLString, out EMails);
@@ -2464,7 +2458,7 @@ public partial class GUBP : BuildCommand
 						List<NodeInfo> EcDeps = GetECDependencies(NodeToDo);
 						foreach (NodeInfo Dep in EcDeps)
 						{
-							if (Dep.Node.RunInEC() && !NodeIsAlreadyComplete(GUBPNodesCompleted, Dep, LocalOnly) && OrdereredToDo.Contains(Dep)) // if something is already finished, we don't put it into EC
+							if (Dep.Node.RunInEC() && !Dep.IsComplete && OrdereredToDo.Contains(Dep)) // if something is already finished, we don't put it into EC
 							{
 								if (OrdereredToDo.IndexOf(Dep) > OrdereredToDo.IndexOf(NodeToDo))
 								{
@@ -2501,7 +2495,7 @@ public partial class GUBP : BuildCommand
 								List<NodeInfo> EcDeps = GetECDependencies(Chain);
 								foreach (NodeInfo Dep in EcDeps)
 								{
-									if (GUBPNodes[Dep.Name].Node.RunInEC() && !NodeIsAlreadyComplete(GUBPNodesCompleted, Dep, LocalOnly) && OrdereredToDo.Contains(Dep)) // if something is already finished, we don't put it into EC
+									if (GUBPNodes[Dep.Name].Node.RunInEC() && !Dep.IsComplete && OrdereredToDo.Contains(Dep)) // if something is already finished, we don't put it into EC
 									{
 										if (OrdereredToDo.IndexOf(Dep) > OrdereredToDo.IndexOf(Chain))
 										{
@@ -2522,7 +2516,7 @@ public partial class GUBP : BuildCommand
 						int MyIndex = MyChain.IndexOf(NodeToDo);
 						if (MyIndex > 0)
 						{
-							if (!PreConditionUncompletedEcDeps.Contains(MyChain[MyIndex - 1].Name) && !NodeIsAlreadyComplete(GUBPNodesCompleted, MyChain[MyIndex - 1], LocalOnly))
+							if (!PreConditionUncompletedEcDeps.Contains(MyChain[MyIndex - 1].Name) && !MyChain[MyIndex - 1].IsComplete)
 							{
 								PreConditionUncompletedEcDeps.Add(MyChain[MyIndex - 1].Name);
 							}
@@ -2532,7 +2526,7 @@ public partial class GUBP : BuildCommand
 							List<NodeInfo> EcDeps = GetECDependencies(NodeToDo);
 							foreach (NodeInfo Dep in EcDeps)
 							{
-								if (Dep.Node.RunInEC() && !NodeIsAlreadyComplete(GUBPNodesCompleted, Dep, LocalOnly) && OrdereredToDo.Contains(Dep)) // if something is already finished, we don't put it into EC
+								if (Dep.Node.RunInEC() && !Dep.IsComplete && OrdereredToDo.Contains(Dep)) // if something is already finished, we don't put it into EC
 								{
 									if (OrdereredToDo.IndexOf(Dep) > OrdereredToDo.IndexOf(NodeToDo))
 									{
@@ -2722,7 +2716,7 @@ public partial class GUBP : BuildCommand
 		}
 	}
 
-	void ExecuteNodes(List<NodeInfo> OrdereredToDo, bool bOnlyNode, bool bFakeEC, bool LocalOnly, bool bSaveSharedTempStorage, Dictionary<string, bool> GUBPNodesCompleted, Dictionary<string, NodeHistory> GUBPNodesHistory, string CLString, string FakeFail)
+	void ExecuteNodes(List<NodeInfo> OrdereredToDo, bool bOnlyNode, bool bFakeEC, bool bSaveSharedTempStorage, Dictionary<string, NodeHistory> GUBPNodesHistory, string CLString, string FakeFail)
 	{
         Dictionary<string, NodeInfo> BuildProductToNodeMap = new Dictionary<string, NodeInfo>();
 		foreach (NodeInfo NodeToDo in OrdereredToDo)
@@ -2792,7 +2786,7 @@ public partial class GUBP : BuildCommand
                 && (NodeToDo.Node.RunInEC() || !NodeToDo.Node.IsAggregate()); //aggregates not in EC can be "run" multiple times, so we can't track those
 
             Log("***** Running GUBP Node {0} -> {1} : {2}", NodeToDo.Node.GetFullName(), GameNameIfAny, NodeStoreName);
-            if (NodeIsAlreadyComplete(GUBPNodesCompleted, NodeToDo, LocalOnly))
+            if (NodeToDo.IsComplete)
             {
                 if (NodeToDo.Name == VersionFilesNode.StaticGetFullName() && !IsBuildMachine)
                 {
