Remove unneccessary indirection on the GUBPNodes dictionary.

diff --git a/Engine/Source/Programs/AutomationTool/GUBP/GUBP.cs b/Engine/Source/Programs/AutomationTool/GUBP/GUBP.cs
--- a/Engine/Source/Programs/AutomationTool/GUBP/GUBP.cs
+++ b/Engine/Source/Programs/AutomationTool/GUBP/GUBP.cs
@@ -841,11 +841,11 @@ public partial class GUBP : BuildCommand
                     }
                     else if (String.Compare(BestNode.ControllingTriggerDotName, NodeToDo.ControllingTriggerDotName) == 0) //sorted by controlling trigger
                     {
-                        if (GUBPNodes[BestNode.Name].Node.IsSticky() && !NodeToDo.Node.IsSticky()) //sticky nodes first
+                        if (BestNode.Node.IsSticky() && !NodeToDo.Node.IsSticky()) //sticky nodes first
                         {
                             bReady = false;
                         }
-                        else if (GUBPNodes[BestNode.Name].Node.IsSticky() == NodeToDo.Node.IsSticky())
+                        else if (BestNode.Node.IsSticky() == NodeToDo.Node.IsSticky())
                         {
                             if (BestPseudoReady && !bPseudoReady)
                             {
@@ -854,7 +854,7 @@ public partial class GUBP : BuildCommand
                             else if (BestPseudoReady == bPseudoReady)
                             {
                                 bool IamLateTrigger = !DoNotConsiderCompletion && NodeToDo.Node.TriggerNode() && NodeToDo != ExplicitTrigger && !NodeToDo.IsComplete;
-                                bool BestIsLateTrigger = !DoNotConsiderCompletion && GUBPNodes[BestNode.Name].Node.TriggerNode() && BestNode != ExplicitTrigger && !BestNode.IsComplete;
+                                bool BestIsLateTrigger = !DoNotConsiderCompletion && BestNode.Node.TriggerNode() && BestNode != ExplicitTrigger && !BestNode.IsComplete;
                                 if (BestIsLateTrigger && !IamLateTrigger)
                                 {
                                     bReady = false;
@@ -892,9 +892,9 @@ public partial class GUBP : BuildCommand
             }
             if (bProgressMade)
             {
-                if (!SubSort && GUBPNodes[BestNode.Name].Node.AgentSharingGroup != "")
+                if (!SubSort && BestNode.Node.AgentSharingGroup != "")
                 {
-                    foreach (NodeInfo ChainNode in SortedAgentGroupChains[GUBPNodes[BestNode.Name].Node.AgentSharingGroup])
+                    foreach (NodeInfo ChainNode in SortedAgentGroupChains[BestNode.Node.AgentSharingGroup])
                     {
                         OrdereredToDo.Add(ChainNode);
                         NodesToDo.Remove(ChainNode);
@@ -917,7 +917,7 @@ public partial class GUBP : BuildCommand
                     {
                         foreach (NodeInfo ChainNode in SortedAgentGroupChains[NodeToDo.Node.AgentSharingGroup])
                         {
-                            foreach (NodeInfo Dep in GUBPNodes[ChainNode.Name].Dependencies)
+                            foreach (NodeInfo Dep in ChainNode.Dependencies)
                             {
                                 if (!SortedAgentGroupChains[NodeToDo.Node.AgentSharingGroup].Contains(Dep) && NodesToDo.Contains(Dep))
                                 {
