Minor clean-ups.

diff --git a/Engine/Source/Programs/AutomationTool/GUBP/GUBP.cs b/Engine/Source/Programs/AutomationTool/GUBP/GUBP.cs
--- a/Engine/Source/Programs/AutomationTool/GUBP/GUBP.cs
+++ b/Engine/Source/Programs/AutomationTool/GUBP/GUBP.cs
@@ -1479,6 +1479,13 @@ public partial class GUBP : BuildCommand
             TempStorage.DeleteLocalTempStorageManifests(CmdEnv);
         }
 
+		// Compute all the frequencies
+		Dictionary<NodeInfo, int> FrequencyOverrides = ApplyFrequencyBarriers(GUBPNodes, BranchOptions.FrequencyBarriers);
+        foreach (KeyValuePair<string, NodeInfo> NodeToDo in GUBPNodes)
+        {
+            ComputeDependentCISFrequencyQuantumShift(NodeToDo.Value, FrequencyOverrides);
+        }
+
 		Dictionary<string, string> FullNodeDependentPromotions = new Dictionary<string, string>();		
 		List<string> SeparatePromotables = new List<string>();
         {
@@ -1506,13 +1513,6 @@ public partial class GUBP : BuildCommand
                 }
             }
 
-			// Compute all the frequencies
-			Dictionary<NodeInfo, int> FrequencyOverrides = ApplyFrequencyBarriers();
-            foreach (KeyValuePair<string, NodeInfo> NodeToDo in GUBPNodes)
-            {
-                ComputeDependentCISFrequencyQuantumShift(NodeToDo.Value, FrequencyOverrides);
-            }
-
             foreach (KeyValuePair<string, NodeInfo> NodeToDo in GUBPNodes)
             {
                 List<string> Deps = GUBPNodes[NodeToDo.Key].Node.DependentPromotions;
@@ -1618,29 +1618,24 @@ public partial class GUBP : BuildCommand
 		List<NodeInfo> UnfinishedTriggers = FindUnfinishedTriggers(bSkipTriggers, ExplicitTrigger, OrdereredToDo);
 
         LogVerbose("*********** Desired And Dependent Nodes, in order.");
-        PrintNodes(this, OrdereredToDo, GUBPNodesHistory,  UnfinishedTriggers);		
+        PrintNodes(this, OrdereredToDo, GUBPNodesHistory, UnfinishedTriggers);		
         //check sorting
 		CheckSortOrder(OrdereredToDo);
 
         string FakeFail = ParseParamValue("FakeFail");
-        if (CommanderSetup)
+        if(CommanderSetup)
         {
 			DoCommanderSetup(TimeIndex, bSkipTriggers, bFakeEC, CLString, ExplicitTrigger, FullNodeList, FullNodeDirectDependencies, FullNodeDependedOnBy, FullNodeDependentPromotions, SeparatePromotables, FullNodeListSortKey, GUBPNodesHistory, OrdereredToDo, UnfinishedTriggers, FakeFail);
-            Log("Commander setup only, done.");
-            PrintRunTime();
-            return;
-
         }
-        if (ParseParam("SaveGraph"))
-        {
-            SaveGraphVisualization(OrdereredToDo);
-        }
-        if (bListOnly)
-        {
-            Log("List only, done.");
-            return;
-        }    
-		ExecuteNodes(OrdereredToDo, bOnlyNode, bFakeEC, bSaveSharedTempStorage, GUBPNodesHistory, CLString, FakeFail);
+		else if(ParseParam("SaveGraph"))
+		{
+			SaveGraphVisualization(OrdereredToDo);
+		}
+		else if(!bListOnly)
+		{
+			ExecuteNodes(OrdereredToDo, bOnlyNode, bFakeEC, bSaveSharedTempStorage, GUBPNodesHistory, CLString, FakeFail);
+		}
+        PrintRunTime();
 	}
 
 	private static void LinkGraph(Dictionary<string, NodeInfo> GUBPNodes)
@@ -1696,11 +1691,11 @@ public partial class GUBP : BuildCommand
 		}
 	}
 
-	private Dictionary<NodeInfo, int> ApplyFrequencyBarriers()
+	private static Dictionary<NodeInfo, int> ApplyFrequencyBarriers(Dictionary<string, NodeInfo> GUBPNodes, Dictionary<string, sbyte> FrequencyBarriers)
 	{
 		// Make sure that everything that's listed as a frequency barrier is completed with the given interval
 		Dictionary<NodeInfo, int> FrequencyOverrides = new Dictionary<NodeInfo, int>();
-		foreach (KeyValuePair<string, sbyte> Barrier in BranchOptions.FrequencyBarriers)
+		foreach (KeyValuePair<string, sbyte> Barrier in FrequencyBarriers)
 		{
 			// Find the barrier node
 			NodeInfo BarrierNode;
