Move EC setup for promotables and dependency names into their own functions, and remove the temporary variables from GUBPNode.

diff --git a/Engine/Source/Programs/AutomationTool/GUBP/GUBP.cs b/Engine/Source/Programs/AutomationTool/GUBP/GUBP.cs
--- a/Engine/Source/Programs/AutomationTool/GUBP/GUBP.cs
+++ b/Engine/Source/Programs/AutomationTool/GUBP/GUBP.cs
@@ -106,7 +106,7 @@ public partial class GUBP : BuildCommand
         GUBPNodes[Node].Node.FullNamesOfPseudosependencies.Clear();
     }
 
-    List<NodeInfo> GetDependencies(NodeInfo NodeToDo, bool bFlat = false, bool ECOnly = false)
+    static List<NodeInfo> GetDependencies(NodeInfo NodeToDo, bool bFlat = false, bool ECOnly = false)
     {
         List<NodeInfo> Result = new List<NodeInfo>();
         foreach (NodeInfo Dependency in NodeToDo.Dependencies)
@@ -154,7 +154,7 @@ public partial class GUBP : BuildCommand
         return Result;
     }
 
-	List<NodeInfo> GetECDependencies(NodeInfo NodeToDo, bool bFlat = false)
+	static List<NodeInfo> GetECDependencies(NodeInfo NodeToDo, bool bFlat = false)
     {
         return GetDependencies(NodeToDo, bFlat, true);
     }
@@ -452,8 +452,6 @@ public partial class GUBP : BuildCommand
         bool bShowChanges = (bp.ParseParam("Changes") && GUBPNodesHistory != null) || bShowAllChanges;
         bool bShowDetailedHistory = (bp.ParseParam("History") && GUBPNodesHistory != null) || bShowChanges;
         bool bShowDependencies = bp.ParseParam("ShowDependencies");
-		bool bShowDependednOn = bp.ParseParam("ShowDependedOn");
-		bool bShowDependentPromotions = bp.ParseParam("ShowDependentPromotions");
         bool bShowECDependencies = bp.ParseParam("ShowECDependencies");
         bool bShowHistory = !bp.ParseParam("NoHistory") && GUBPNodesHistory != null;
         bool AddEmailProps = bp.ParseParam("ShowEmails");
@@ -566,21 +564,6 @@ public partial class GUBP : BuildCommand
                     Log("           {0}", Dep.Name);
                 }
             }
-
-			if(bShowDependednOn)
-			{
-				foreach (string Dep in NodeToDo.Node.FullNamesOfDependedOn)
-				{
-					Log("            depOn> {0}", Dep);
-				}
-			}
-			if (bShowDependentPromotions)
-			{
-				foreach (string Dep in NodeToDo.Node.DependentPromotions)
-				{
-					Log("            depPro> {0}", Dep);
-				}
-			}			
         }
     }
     void SaveGraphVisualization(List<NodeInfo> Nodes)
@@ -1486,49 +1469,6 @@ public partial class GUBP : BuildCommand
             ComputeDependentCISFrequencyQuantumShift(NodeToDo.Value, FrequencyOverrides);
         }
 
-		Dictionary<string, string> FullNodeDependentPromotions = new Dictionary<string, string>();		
-		List<string> SeparatePromotables = new List<string>();
-        {
-            foreach (KeyValuePair<string, NodeInfo> NodeToDo in GUBPNodes)
-            {
-                if (GUBPNodes[NodeToDo.Key].Node.IsSeparatePromotable())
-                {
-                    SeparatePromotables.Add(GUBPNodes[NodeToDo.Key].Node.GetFullName());
-                    List<NodeInfo> Dependencies = new List<NodeInfo>();
-                    Dependencies = GetECDependencies(NodeToDo.Value);
-                    foreach (NodeInfo Dep in Dependencies)
-                    {
-                        if (!GUBPNodes.ContainsKey(Dep.Name))
-                        {
-                            throw new AutomationException("Node {0} is not in the graph.  It is a dependency of {1}.", Dep.Name, NodeToDo.Key);
-                        }
-                        if (!GUBPNodes[Dep.Name].Node.IsPromotableAggregate())
-                        {
-                            if (!GUBPNodes[Dep.Name].Node.DependentPromotions.Contains(NodeToDo.Key))
-                            {
-                                GUBPNodes[Dep.Name].Node.DependentPromotions.Add(NodeToDo.Key);
-                            }
-                        }
-                    }
-                }
-            }
-
-            foreach (KeyValuePair<string, NodeInfo> NodeToDo in GUBPNodes)
-            {
-                List<string> Deps = GUBPNodes[NodeToDo.Key].Node.DependentPromotions;
-                string All = "";
-                foreach (string Dep in Deps)
-                {
-                    if (All != "")
-                    {
-                        All += " ";
-                    }
-                    All += Dep;
-                }
-                FullNodeDependentPromotions.Add(NodeToDo.Key, All);
-            }
-        }	
-
 		bool bOnlyNode;
 		HashSet<NodeInfo> NodesToDo;
 		Dictionary<string, string> FullNodeDependedOnBy = new Dictionary<string, string>();
@@ -1619,13 +1559,16 @@ public partial class GUBP : BuildCommand
         string FakeFail = ParseParamValue("FakeFail");
         if(CommanderSetup)
         {
+			List<NodeInfo> SeparatePromotables = FindPromotables(GUBPNodes.Values);
+			Dictionary<NodeInfo, List<NodeInfo>> DependentPromotions = FindDependentPromotables(GUBPNodes.Values, SeparatePromotables);
+
 			Dictionary<string, string> FullNodeList;
 			Dictionary<string, string> FullNodeDirectDependencies;
 			GetFullNodeListAndDirectDependencies(out FullNodeList, out FullNodeDirectDependencies);
 
 			Dictionary<string, int> FullNodeListSortKey = GetDisplayOrder(FullNodeList.Keys.ToList(), FullNodeDirectDependencies, GUBPNodes);
 
-			DoCommanderSetup(TimeIndex, bSkipTriggers, bFakeEC, CLString, ExplicitTrigger, FullNodeList, FullNodeDirectDependencies, FullNodeDependedOnBy, FullNodeDependentPromotions, SeparatePromotables, FullNodeListSortKey, GUBPNodesHistory, OrdereredToDo, UnfinishedTriggers, FakeFail);
+			DoCommanderSetup(TimeIndex, bSkipTriggers, bFakeEC, CLString, ExplicitTrigger, FullNodeList, FullNodeDirectDependencies, FullNodeDependedOnBy, DependentPromotions, SeparatePromotables, FullNodeListSortKey, GUBPNodesHistory, OrdereredToDo, UnfinishedTriggers, FakeFail);
         }
 		else if(ParseParam("SaveGraph"))
 		{
@@ -1638,6 +1581,36 @@ public partial class GUBP : BuildCommand
         PrintRunTime();
 	}
 
+	private static List<NodeInfo> FindPromotables(IEnumerable<NodeInfo> NodesToDo)
+	{
+		List<NodeInfo> SeparatePromotables = new List<NodeInfo>();
+		foreach (NodeInfo NodeToDo in NodesToDo)
+		{
+			if (NodeToDo.Node.IsSeparatePromotable())
+			{
+				SeparatePromotables.Add(NodeToDo);
+			}
+		}
+		return SeparatePromotables;
+	}
+
+	private static Dictionary<NodeInfo, List<NodeInfo>> FindDependentPromotables(IEnumerable<NodeInfo> NodesToDo, IEnumerable<NodeInfo> SeparatePromotions)
+	{
+		Dictionary<NodeInfo, List<NodeInfo>> DependentPromotions = NodesToDo.ToDictionary(x => x, x => new List<NodeInfo>());
+		foreach (NodeInfo NodeToDo in SeparatePromotions)
+		{
+			List<NodeInfo> Dependencies = GetECDependencies(NodeToDo);
+			foreach (NodeInfo Dep in Dependencies)
+			{
+				if (!Dep.Node.IsPromotableAggregate())
+				{
+					DependentPromotions[Dep].Add(NodeToDo);
+				}
+			}
+		}
+		return DependentPromotions;
+	}
+
 	private static void LinkGraph(Dictionary<string, NodeInfo> GUBPNodes)
 	{
 		int NumErrors = 0;
@@ -1988,39 +1961,28 @@ public partial class GUBP : BuildCommand
 		//find things that depend on our nodes and setup commander dictionary
 		if (!bOnlyNode)
 		{
+			foreach (NodeInfo NodeToDo in NodesToDo)
+			{
+				FullNodeDependedOnBy[NodeToDo.Name] = "";
+			}
 			foreach (NodeInfo NodeToDo in NodesToDo)
 			{
 				if (!NodeToDo.Node.IsAggregate() && !NodeToDo.Node.IsTest())
 				{
-					List<NodeInfo> ECDependencies = new List<NodeInfo>();
-					ECDependencies = GetECDependencies(NodeToDo);
+					List<NodeInfo> ECDependencies = GetECDependencies(NodeToDo);
 					foreach (NodeInfo Dep in ECDependencies)
 					{
-						if (!GUBPNodes.ContainsKey(Dep.Name))
+						if(FullNodeDependedOnBy[Dep.Name].Length == 0)
 						{
-							throw new AutomationException("Node {0} is not in the graph. It is a dependency of {1}.", Dep.Name, NodeToDo.Name);
+							FullNodeDependedOnBy[Dep.Name] = NodeToDo.Name;
 						}
-						if (!GUBPNodes[Dep.Name].Node.FullNamesOfDependedOn.Contains(NodeToDo.Name))
+						else
 						{
-							GUBPNodes[Dep.Name].Node.FullNamesOfDependedOn.Add(NodeToDo.Name);
+							FullNodeDependedOnBy[Dep.Name] += " " + NodeToDo.Name;
 						}
 					}
 				}
 			}
-			foreach (NodeInfo NodeToDo in NodesToDo)
-			{
-				List<string> Deps = NodeToDo.Node.FullNamesOfDependedOn;
-				string All = "";
-				foreach (string Dep in Deps)
-				{
-					if (All != "")
-					{
-						All += " ";
-					}
-					All += Dep;
-				}
-				FullNodeDependedOnBy.Add(NodeToDo.Name, All);
-			}
 		}
 	}
 
@@ -2238,7 +2200,7 @@ public partial class GUBP : BuildCommand
 		}
 	}
 
-	private void DoCommanderSetup(int TimeIndex, bool bSkipTriggers, bool bFakeEC, string CLString, NodeInfo ExplicitTrigger, Dictionary<string, string> FullNodeList, Dictionary<string, string> FullNodeDirectDependencies, Dictionary<string, string> FullNodeDependedOnBy, Dictionary<string, string> FullNodeDependentPromotions, List<string> SeparatePromotables, Dictionary<string, int> FullNodeListSortKey, Dictionary<string, NodeHistory> GUBPNodesHistory, List<NodeInfo> OrdereredToDo, List<NodeInfo> UnfinishedTriggers, string FakeFail)
+	private void DoCommanderSetup(int TimeIndex, bool bSkipTriggers, bool bFakeEC, string CLString, NodeInfo ExplicitTrigger, Dictionary<string, string> FullNodeList, Dictionary<string, string> FullNodeDirectDependencies, Dictionary<string, string> FullNodeDependedOnBy, Dictionary<NodeInfo, List<NodeInfo>> DependentPromotions, List<NodeInfo> SeparatePromotables, Dictionary<string, int> FullNodeListSortKey, Dictionary<string, NodeHistory> GUBPNodesHistory, List<NodeInfo> OrdereredToDo, List<NodeInfo> UnfinishedTriggers, string FakeFail)
 	{
 
 		if (OrdereredToDo.Count == 0)
@@ -2263,13 +2225,13 @@ public partial class GUBP : BuildCommand
 		{
 			ECProps.Add(string.Format("DependedOnBy/{0}={1}", NodePair.Key, NodePair.Value));
 		}
-		foreach (KeyValuePair<string, string> NodePair in FullNodeDependentPromotions)
+		foreach (KeyValuePair<NodeInfo, List<NodeInfo>> NodePair in DependentPromotions)
 		{
-			ECProps.Add(string.Format("DependentPromotions/{0}={1}", NodePair.Key, NodePair.Value));
+			ECProps.Add(string.Format("DependentPromotions/{0}={1}", NodePair.Key.Name, String.Join(" ", NodePair.Value.Select(x => x.Name))));
 		}
-		foreach (string Node in SeparatePromotables)
+		foreach (NodeInfo Node in SeparatePromotables)
 		{
-			ECProps.Add(string.Format("PossiblePromotables/{0}={1}", Node, ""));
+			ECProps.Add(string.Format("PossiblePromotables/{0}={1}", Node.Name, ""));
 		}
 		List<string> ECJobProps = new List<string>();
 		if (ExplicitTrigger != null)
diff --git a/Engine/Source/Programs/AutomationTool/GUBP/LegacyNodes.cs b/Engine/Source/Programs/AutomationTool/GUBP/LegacyNodes.cs
--- a/Engine/Source/Programs/AutomationTool/GUBP/LegacyNodes.cs
+++ b/Engine/Source/Programs/AutomationTool/GUBP/LegacyNodes.cs
@@ -16,9 +16,7 @@ partial class GUBP
     {
         public List<string> FullNamesOfDependencies = new List<string>();
         public List<string> FullNamesOfPseudosependencies = new List<string>(); //these are really only used for sorting. We want the editor to fail before the monolithics. Think of it as "can't possibly be useful without".
-		public List<string> FullNamesOfDependedOn = new List<string>();
         public List<string> BuildProducts = null;
-		public List<string> DependentPromotions = new List<string>();
         public List<string> AllDependencyBuildProducts = null;
         public List<string> AllDependencies = null;		
         public string AgentSharingGroup = "";
