Fix exception if node is not already in the dictionary (may be an aggregate node).

diff --git a/Engine/Source/Programs/AutomationTool/GUBP/GUBP.cs b/Engine/Source/Programs/AutomationTool/GUBP/GUBP.cs
--- a/Engine/Source/Programs/AutomationTool/GUBP/GUBP.cs
+++ b/Engine/Source/Programs/AutomationTool/GUBP/GUBP.cs
@@ -1972,7 +1972,8 @@ public partial class GUBP : BuildCommand
 					List<NodeInfo> ECDependencies = GetECDependencies(NodeToDo);
 					foreach (NodeInfo Dep in ECDependencies)
 					{
-						if(FullNodeDependedOnBy[Dep.Name].Length == 0)
+						string CurrentValue;
+						if(!FullNodeDependedOnBy.TryGetValue(Dep.Name, out CurrentValue) || CurrentValue.Length == 0)
 						{
 							FullNodeDependedOnBy[Dep.Name] = NodeToDo.Name;
 						}
