Turn some GUBP globals into locals that are passed around by value.

diff --git a/Engine/Source/Programs/AutomationTool/GUBP/GUBP.cs b/Engine/Source/Programs/AutomationTool/GUBP/GUBP.cs
--- a/Engine/Source/Programs/AutomationTool/GUBP/GUBP.cs
+++ b/Engine/Source/Programs/AutomationTool/GUBP/GUBP.cs
@@ -11,16 +11,11 @@ using System.Linq;
 
 public partial class GUBP : BuildCommand
 {
-    public string StoreName = null;
 	public string BranchName;
     public int CL = 0;
     public bool bSignBuildProducts = false;
-    public bool bHasTests = false;
     public BranchInfo Branch = null;
-    public bool bOrthogonalizeEditorPlatforms = false;
     public List<UnrealTargetPlatform> HostPlatforms;
-    public bool bFake = false;
-    public static bool bNoIOSOnPC = false;
     public static bool bForceIncrementalCompile = false;
     public string EmailHint;
     static public bool bPreflightBuild = false;
@@ -950,7 +945,7 @@ public partial class GUBP : BuildCommand
         return ParentPath + "/" + GetJobStepPath(Dep);
     }
 
-    void UpdateNodeHistory(Dictionary<string, NodeHistory> GUBPNodesHistory, NodeInfo NodeToDo, string CLString)
+    void UpdateNodeHistory(Dictionary<string, NodeHistory> GUBPNodesHistory, NodeInfo NodeToDo, string CLString, string StoreName)
     {
         if (NodeToDo.Node.RunInEC() && !NodeToDo.Node.TriggerNode() && CLString != "")
         {
@@ -994,7 +989,7 @@ public partial class GUBP : BuildCommand
 			GUBPNodesHistory.Add(NodeToDo.Name, History);
         }
     }
-	void GetFailureEmails(Dictionary<string, NodeHistory> GUBPNodesHistory, NodeInfo NodeToDo, string CLString, bool OnlyLateUpdates = false)
+	void GetFailureEmails(Dictionary<string, NodeHistory> GUBPNodesHistory, NodeInfo NodeToDo, string CLString, string StoreName, bool OnlyLateUpdates = false)
 	{
         string EMails;
         string FailCauserEMails = "";
@@ -1014,7 +1009,7 @@ public partial class GUBP : BuildCommand
                 {
                     if (OnlyLateUpdates)
                     {
-                        LastNonDuplicateFail = FindLastNonDuplicateFail(GUBPNodesHistory, NodeToDo, CLString);
+                        LastNonDuplicateFail = FindLastNonDuplicateFail(GUBPNodesHistory, NodeToDo, CLString, StoreName);
                         if (LastNonDuplicateFail < P4Env.Changelist)
                         {
                             Log("*** Red-after-red spam reduction, changed CL {0} to CL {1} because the errors didn't change.", P4Env.Changelist, LastNonDuplicateFail);
@@ -1121,7 +1116,7 @@ public partial class GUBP : BuildCommand
         return true;
     }
 
-    int FindLastNonDuplicateFail(Dictionary<string, NodeHistory> GUBPNodesHistory, NodeInfo NodeToDo, string CLString)
+    int FindLastNonDuplicateFail(Dictionary<string, NodeHistory> GUBPNodesHistory, NodeInfo NodeToDo, string CLString, string StoreName)
     {
         NodeHistory History = GUBPNodesHistory[NodeToDo.Name];
         int Result = P4Env.Changelist;
@@ -1343,7 +1338,7 @@ public partial class GUBP : BuildCommand
 
         bForceIncrementalCompile = ParseParam("ForceIncrementalCompile");
         bool bNoAutomatedTesting = ParseParam("NoAutomatedTesting") || BranchOptions.bNoAutomatedTesting;		
-        StoreName = ParseParamValue("Store");
+        string StoreName = ParseParamValue("Store");
         string StoreSuffix = ParseParamValue("StoreSuffix", "");
 
         if (bPreflightBuild)
@@ -1358,7 +1353,7 @@ public partial class GUBP : BuildCommand
         bool bHistory = ParseParam("History") || bChanges;
         bool bListOnly = ParseParam("ListOnly") || bHistory;
         bool bSkipTriggers = ParseParam("SkipTriggers");
-        bFake = ParseParam("fake");
+        bool bFake = ParseParam("fake");
         bool bFakeEC = ParseParam("FakeEC");
         int TimeIndex = ParseParamInt("TimeIndex", 0);
         if (TimeIndex == 0)
@@ -1366,8 +1361,6 @@ public partial class GUBP : BuildCommand
             TimeIndex = ParseParamInt("UserTimeIndex", 0);
         }
 
-        bNoIOSOnPC = HostPlatforms.Contains(UnrealTargetPlatform.Mac);
-
         bool bSaveSharedTempStorage = false;
 
         if (bHistory && !P4Enabled)
@@ -1568,7 +1561,7 @@ public partial class GUBP : BuildCommand
 
 			Dictionary<string, int> FullNodeListSortKey = GetDisplayOrder(FullNodeList.Keys.ToList(), FullNodeDirectDependencies, GUBPNodes);
 
-			DoCommanderSetup(TimeIndex, bSkipTriggers, bFakeEC, CLString, ExplicitTrigger, FullNodeList, FullNodeDirectDependencies, FullNodeDependedOnBy, DependentPromotions, SeparatePromotables, FullNodeListSortKey, GUBPNodesHistory, OrdereredToDo, UnfinishedTriggers, FakeFail);
+			DoCommanderSetup(TimeIndex, bSkipTriggers, bFake, bFakeEC, CLString, ExplicitTrigger, FullNodeList, FullNodeDirectDependencies, FullNodeDependedOnBy, DependentPromotions, SeparatePromotables, FullNodeListSortKey, GUBPNodesHistory, OrdereredToDo, UnfinishedTriggers, FakeFail);
         }
 		else if(ParseParam("SaveGraph"))
 		{
@@ -1576,7 +1569,7 @@ public partial class GUBP : BuildCommand
 		}
 		else if(!bListOnly)
 		{
-			ExecuteNodes(OrdereredToDo, bOnlyNode, bFakeEC, bSaveSharedTempStorage, GUBPNodesHistory, CLString, FakeFail);
+			ExecuteNodes(OrdereredToDo, bOnlyNode, bFake, bFakeEC, bSaveSharedTempStorage, GUBPNodesHistory, CLString, StoreName, FakeFail);
 		}
         PrintRunTime();
 	}
@@ -2177,10 +2170,6 @@ public partial class GUBP : BuildCommand
 			{
 				continue;
 			}
-			if (NodeToDo.Node.IsTest())
-			{
-				bHasTests = true;
-			}
 			int MyIndex = OrdereredToDo.IndexOf(NodeToDo);
 			foreach (NodeInfo Dep in NodeToDo.Dependencies)
 			{
@@ -2201,7 +2190,7 @@ public partial class GUBP : BuildCommand
 		}
 	}
 
-	private void DoCommanderSetup(int TimeIndex, bool bSkipTriggers, bool bFakeEC, string CLString, NodeInfo ExplicitTrigger, Dictionary<string, string> FullNodeList, Dictionary<string, string> FullNodeDirectDependencies, Dictionary<string, string> FullNodeDependedOnBy, Dictionary<NodeInfo, List<NodeInfo>> DependentPromotions, List<NodeInfo> SeparatePromotables, Dictionary<string, int> FullNodeListSortKey, Dictionary<string, NodeHistory> GUBPNodesHistory, List<NodeInfo> OrdereredToDo, List<NodeInfo> UnfinishedTriggers, string FakeFail)
+	private void DoCommanderSetup(int TimeIndex, bool bSkipTriggers, bool bFake, bool bFakeEC, string CLString, NodeInfo ExplicitTrigger, Dictionary<string, string> FullNodeList, Dictionary<string, string> FullNodeDirectDependencies, Dictionary<string, string> FullNodeDependedOnBy, Dictionary<NodeInfo, List<NodeInfo>> DependentPromotions, List<NodeInfo> SeparatePromotables, Dictionary<string, int> FullNodeListSortKey, Dictionary<string, NodeHistory> GUBPNodesHistory, List<NodeInfo> OrdereredToDo, List<NodeInfo> UnfinishedTriggers, string FakeFail)
 	{
 
 		if (OrdereredToDo.Count == 0)
@@ -2652,6 +2641,7 @@ public partial class GUBP : BuildCommand
 			}
 			WriteECPerl(StepList);
 		}
+		bool bHasTests = OrdereredToDo.Any(x => x.Node.IsTest());
 		RunECTool(String.Format("setProperty \"/myWorkflow/HasTests\" \"{0}\"", bHasTests));
 		{
 			ECProps.Add("GUBP_LoadedProps=1");
@@ -2674,7 +2664,7 @@ public partial class GUBP : BuildCommand
 		}
 	}
 
-	void ExecuteNodes(List<NodeInfo> OrdereredToDo, bool bOnlyNode, bool bFakeEC, bool bSaveSharedTempStorage, Dictionary<string, NodeHistory> GUBPNodesHistory, string CLString, string FakeFail)
+	void ExecuteNodes(List<NodeInfo> OrdereredToDo, bool bOnlyNode, bool bFake, bool bFakeEC, bool bSaveSharedTempStorage, Dictionary<string, NodeHistory> GUBPNodesHistory, string CLString, string StoreName, string FakeFail)
 	{
         Dictionary<string, NodeInfo> BuildProductToNodeMap = new Dictionary<string, NodeInfo>();
 		foreach (NodeInfo NodeToDo in OrdereredToDo)
@@ -2852,7 +2842,7 @@ public partial class GUBP : BuildCommand
                     {
 						using(TelemetryStopwatch UpdateNodeHistoryStopwatch = new TelemetryStopwatch("UpdateNodeHistory"))
 						{
-							UpdateNodeHistory(GUBPNodesHistory, NodeToDo, CLString);
+							UpdateNodeHistory(GUBPNodesHistory, NodeToDo, CLString, StoreName);
 						}
 						using(TelemetryStopwatch SaveNodeStatusStopwatch = new TelemetryStopwatch("SaveNodeStatus"))
 						{
@@ -2867,7 +2857,7 @@ public partial class GUBP : BuildCommand
 						{
 							using(TelemetryStopwatch GetFailEmailsStopwatch = new TelemetryStopwatch("GetFailEmails"))
 							{
-								GetFailureEmails(GUBPNodesHistory, NodeToDo, CLString);
+								GetFailureEmails(GUBPNodesHistory, NodeToDo, CLString, StoreName);
 							}
 						}
 						UpdateECBuildTime(NodeToDo, BuildDuration);
@@ -2922,7 +2912,7 @@ public partial class GUBP : BuildCommand
                 {
 					using(TelemetryStopwatch UpdateNodeHistoryStopwatch = new TelemetryStopwatch("UpdateNodeHistory"))
 					{
-						UpdateNodeHistory(GUBPNodesHistory, NodeToDo, CLString);
+						UpdateNodeHistory(GUBPNodesHistory, NodeToDo, CLString, StoreName);
 					}
 					using(TelemetryStopwatch SaveNodeStatusStopwatch = new TelemetryStopwatch("SaveNodeStatus"))
 					{
@@ -2937,7 +2927,7 @@ public partial class GUBP : BuildCommand
 					{
 						using(TelemetryStopwatch GetFailEmailsStopwatch = new TelemetryStopwatch("GetFailEmails"))
 						{
-							GetFailureEmails(GUBPNodesHistory, NodeToDo, CLString);
+							GetFailureEmails(GUBPNodesHistory, NodeToDo, CLString, StoreName);
 						}
 					}
 					UpdateECBuildTime(NodeToDo, BuildDuration);
diff --git a/Engine/Source/Programs/AutomationTool/GUBP/LegacyBranchSetup.cs b/Engine/Source/Programs/AutomationTool/GUBP/LegacyBranchSetup.cs
--- a/Engine/Source/Programs/AutomationTool/GUBP/LegacyBranchSetup.cs
+++ b/Engine/Source/Programs/AutomationTool/GUBP/LegacyBranchSetup.cs
@@ -293,6 +293,8 @@ partial class GUBP
             LogVerbose("Active Platform: {0}", Plat.ToString());
         }
 
+		bool bNoIOSOnPC = HostPlatforms.Contains(UnrealTargetPlatform.Mac);
+
         if (HostPlatforms.Count >= 2)
         {
             // make sure each project is set up with the right assumptions on monolithics that prefer a platform.
@@ -301,8 +303,8 @@ partial class GUBP
                 var OptionsMac = CodeProj.Options(UnrealTargetPlatform.Mac);
                 var OptionsPC = CodeProj.Options(UnrealTargetPlatform.Win64);
 
-                var MacMonos = GamePlatformMonolithicsNode.GetMonolithicPlatformsForUProject(UnrealTargetPlatform.Mac, ActivePlatforms, CodeProj, false);
-                var PCMonos = GamePlatformMonolithicsNode.GetMonolithicPlatformsForUProject(UnrealTargetPlatform.Win64, ActivePlatforms, CodeProj, false);
+				var MacMonos = GamePlatformMonolithicsNode.GetMonolithicPlatformsForUProject(UnrealTargetPlatform.Mac, ActivePlatforms, CodeProj, false, bNoIOSOnPC);
+				var PCMonos = GamePlatformMonolithicsNode.GetMonolithicPlatformsForUProject(UnrealTargetPlatform.Win64, ActivePlatforms, CodeProj, false, bNoIOSOnPC);
 
                 if (!OptionsMac.bIsPromotable && OptionsPC.bIsPromotable && 
                     (MacMonos.Contains(UnrealTargetPlatform.IOS) || PCMonos.Contains(UnrealTargetPlatform.IOS)))
@@ -392,20 +394,6 @@ partial class GUBP
 				}
 
 				AddNode(new EditorAndToolsNode(this, HostPlatform));
-
-				if (bOrthogonalizeEditorPlatforms)
-				{
-					foreach (var Plat in ActivePlatforms)
-					{
-						if (Plat != HostPlatform && Plat != GetAltHostPlatform(HostPlatform))
-						{
-							if (Platform.Platforms[HostPlatform].CanHostPlatform(Plat))
-							{
-								AddNode(new EditorPlatformNode(HostPlatform, Plat));
-							}
-						}
-					}
-				}
 			}
 
             bool DoASharedPromotable = false;
@@ -448,7 +436,7 @@ partial class GUBP
                 }
 
                 var TempNonCodeFormalBuilds = Target.Rules.GUBP_GetNonCodeFormalBuilds_BaseEditorTypeOnly();
-                var HostMonos = GamePlatformMonolithicsNode.GetMonolithicPlatformsForUProject(HostPlatform, ActivePlatforms, Branch.BaseEngineProject, true);
+				var HostMonos = GamePlatformMonolithicsNode.GetMonolithicPlatformsForUProject(HostPlatform, ActivePlatforms, Branch.BaseEngineProject, true, bNoIOSOnPC);
 
                 foreach (var Codeless in TempNonCodeFormalBuilds)
                 {
@@ -1010,7 +998,7 @@ partial class GUBP
                 }
                 if (PromotedHosts.Count > 0)
                 {
-                    AddNode(new GameAggregatePromotableNode(this, PromotedHosts, ActivePlatforms, CodeProj, true));
+                    AddNode(new GameAggregatePromotableNode(this, PromotedHosts, ActivePlatforms, CodeProj, true, bNoIOSOnPC));
                     if (AnySeparate)
                     {
                         AddNode(new WaitForGamePromotionUserInput(this, CodeProj, false));
@@ -1023,7 +1011,7 @@ partial class GUBP
         }
         if (NumSharedAllHosts > 0)
         {
-            AddNode(new GameAggregatePromotableNode(this, HostPlatforms, ActivePlatforms, Branch.BaseEngineProject, false));
+            AddNode(new GameAggregatePromotableNode(this, HostPlatforms, ActivePlatforms, Branch.BaseEngineProject, false, bNoIOSOnPC));
 
             AddNode(new SharedAggregatePromotableNode(this, HostPlatforms, ActivePlatforms));
             AddNode(new WaitForSharedPromotionUserInput(this, false));
diff --git a/Engine/Source/Programs/AutomationTool/GUBP/LegacyNodes.cs b/Engine/Source/Programs/AutomationTool/GUBP/LegacyNodes.cs
--- a/Engine/Source/Programs/AutomationTool/GUBP/LegacyNodes.cs
+++ b/Engine/Source/Programs/AutomationTool/GUBP/LegacyNodes.cs
@@ -556,10 +556,6 @@ partial class GUBP
 			{
 				AddArgs += " -precompile";
 			}
-            if (bp.bOrthogonalizeEditorPlatforms)
-            {
-                AddArgs += " -skipnonhostplatforms";
-            }
             Agenda.AddTargets(
                 new string[] { bp.Branch.BaseEngineProject.Properties.Targets[TargetRules.TargetType.Editor].TargetName },
                 HostPlatform, UnrealTargetConfiguration.Development, InAddArgs: AddArgs);
@@ -1098,79 +1094,6 @@ partial class GUBP
         }
     }
 
-
-    public class EditorPlatformNode : CompileNode
-    {
-        UnrealTargetPlatform EditorPlatform;
-
-        public EditorPlatformNode(UnrealTargetPlatform InHostPlatform, UnrealTargetPlatform Plat)
-            : base(InHostPlatform)
-        {
-            if (InHostPlatform != UnrealTargetPlatform.Win64)
-            {
-            AgentSharingGroup = "Editor" + StaticGetHostPlatformSuffix(InHostPlatform);
-            }
-            EditorPlatform = Plat;
-            AddDependency(RootEditorNode.StaticGetFullName(InHostPlatform));
-        }
-
-        public static string StaticGetFullName(UnrealTargetPlatform InHostPlatform, UnrealTargetPlatform Plat)
-        {
-            return Plat.ToString() + "_EditorPlatform" + StaticGetHostPlatformSuffix(InHostPlatform);
-        }
-        public override string GetFullName()
-        {
-            return StaticGetFullName(HostPlatform, EditorPlatform);
-        }
-        public override float Priority()
-        {
-            return base.Priority() + 1;
-        }
-        public override bool IsSticky()
-        {
-            if (HostPlatform == UnrealTargetPlatform.Win64)
-            {
-                return true;
-            }
-            return false;
-        }
-        public override string ECProcedure()
-        {
-            if (HostPlatform == UnrealTargetPlatform.Win64)
-            {
-                return String.Format("GUBP_UAT_Node_Parallel_AgentShare_Editor");
-            }
-            return base.ECProcedure();
-        }
-        public override int CISFrequencyQuantumShift(GUBP bp)
-        {
-            int Result = base.CISFrequencyQuantumShift(bp);            
-            return Result;
-        }
-
-        public override UE4Build.BuildAgenda GetAgenda(GUBP bp)
-        {
-            if (!bp.bOrthogonalizeEditorPlatforms)
-            {
-                throw new AutomationException("EditorPlatformNode node should not be used unless we are orthogonalizing editor platforms.");
-            }
-            var Agenda = new UE4Build.BuildAgenda();
-
-            Agenda.AddTargets(
-                new string[] { bp.Branch.BaseEngineProject.Properties.Targets[TargetRules.TargetType.Editor].TargetName },
-                HostPlatform, UnrealTargetConfiguration.Development, InAddArgs: "-nobuilduht -skipactionhistory -CopyAppBundleBackToDevice -onlyplatformspecificfor=" + EditorPlatform.ToString());
-            foreach (var ProgramTarget in bp.Branch.BaseEngineProject.Properties.Programs)
-            {
-                if (ProgramTarget.Rules.GUBP_AlwaysBuildWithBaseEditor() && ProgramTarget.Rules.SupportsPlatform(HostPlatform) && ProgramTarget.Rules.GUBP_NeedsPlatformSpecificDLLs())
-                {
-                    Agenda.AddTargets(new string[] { ProgramTarget.TargetName }, HostPlatform, UnrealTargetConfiguration.Development, InAddArgs: "-nobuilduht -skipactionhistory -CopyAppBundleBackToDevice -onlyplatformspecificfor=" + EditorPlatform.ToString());
-                }
-            }
-            return Agenda;
-        }
-
-    }
-
     public class EditorGameNode : CompileNode
     {
 		List<BranchInfo.BranchUProject> GameProjects = new List<BranchInfo.BranchUProject>();
@@ -1403,7 +1326,7 @@ partial class GUBP
 			return GameProj.GameName + "_Monolithics" + (Precompiled? "_Precompiled" : "");
 		}
 
-		public static List<UnrealTargetPlatform> GetMonolithicPlatformsForUProject(UnrealTargetPlatform HostPlatform, List<UnrealTargetPlatform> ActivePlatforms, BranchInfo.BranchUProject GameProj, bool bIncludeHostPlatform)
+		public static List<UnrealTargetPlatform> GetMonolithicPlatformsForUProject(UnrealTargetPlatform HostPlatform, List<UnrealTargetPlatform> ActivePlatforms, BranchInfo.BranchUProject GameProj, bool bIncludeHostPlatform, bool bNoIOSOnPC)
 		{
 			UnrealTargetPlatform AltHostPlatform = GUBP.GetAltHostPlatform(HostPlatform);
 			var Result = new List<UnrealTargetPlatform>();
@@ -1417,7 +1340,7 @@ partial class GUBP
 					var AllPlatforms = Platforms.Union(AdditionalPlatforms);
 					foreach (var Plat in AllPlatforms)
 					{
-						if (GUBP.bNoIOSOnPC && Plat == UnrealTargetPlatform.IOS && HostPlatform == UnrealTargetPlatform.Win64)
+						if (bNoIOSOnPC && Plat == UnrealTargetPlatform.IOS && HostPlatform == UnrealTargetPlatform.Win64)
 						{
 							continue;
 						}
@@ -1756,7 +1679,7 @@ partial class GUBP
     {
         BranchInfo.BranchUProject GameProj;		
 
-        public GameAggregatePromotableNode(GUBP bp, List<UnrealTargetPlatform> InHostPlatforms, List<UnrealTargetPlatform> InActivePlatforms, BranchInfo.BranchUProject InGameProj, bool IsSeparate)
+        public GameAggregatePromotableNode(GUBP bp, List<UnrealTargetPlatform> InHostPlatforms, List<UnrealTargetPlatform> InActivePlatforms, BranchInfo.BranchUProject InGameProj, bool IsSeparate, bool bNoIOSOnPC)
             : base(InHostPlatforms, InGameProj.GameName)
         {
             GameProj = InGameProj;
@@ -1774,20 +1697,10 @@ partial class GUBP
                     AddDependency(EditorGameNode.StaticGetFullName(HostPlatform, GameProj));
                 }				
                 // add all of the platforms I use
-                {
-                    var Platforms = GamePlatformMonolithicsNode.GetMonolithicPlatformsForUProject(HostPlatform, InActivePlatforms, InGameProj, false);
-                    if (bp.bOrthogonalizeEditorPlatforms)
-                    {
-                        foreach (var Plat in Platforms)
-                        {
-                            AddDependency(EditorPlatformNode.StaticGetFullName(HostPlatform, Plat));
-                        }
-                    }
-                }
                 {
                     if (!GameProj.Options(HostPlatform).bPromoteEditorOnly)
                     {
-						var Platforms = GamePlatformMonolithicsNode.GetMonolithicPlatformsForUProject(HostPlatform, InActivePlatforms, InGameProj, true);
+						var Platforms = GamePlatformMonolithicsNode.GetMonolithicPlatformsForUProject(HostPlatform, InActivePlatforms, InGameProj, true, bNoIOSOnPC);
 						foreach (var Plat in Platforms)
 						{
 							AddDependency(GamePlatformMonolithicsNode.StaticGetFullName(HostPlatform, GameProj, Plat));
@@ -2291,13 +2204,6 @@ partial class GUBP
             CookPlatform = InCookPlatform;
             bIsMassive = false;
             AddDependency(EditorAndToolsNode.StaticGetFullName(HostPlatform));
-            if (bp.bOrthogonalizeEditorPlatforms)
-            {
-                if (TargetPlatform != HostPlatform && TargetPlatform != GUBP.GetAltHostPlatform(HostPlatform))
-                {
-                    AddDependency(EditorPlatformNode.StaticGetFullName(HostPlatform, TargetPlatform));
-                }
-            }			
             bool bIsShared = false;
             // is this the "base game" or a non code project?
             if (InGameProj.GameName != bp.Branch.BaseEngineProject.GameName && GameProj.Properties.Targets.ContainsKey(TargetRules.TargetType.Editor))
@@ -2422,13 +2328,6 @@ partial class GUBP
             CookPlatform = InCookPlatform;
             bIsMassive = false;
             AddDependency(RootEditorNode.StaticGetFullName(HostPlatform));
-            if (bp.bOrthogonalizeEditorPlatforms)
-            {
-                if (TargetPlatform != HostPlatform && TargetPlatform != GUBP.GetAltHostPlatform(HostPlatform))
-                {
-                    AddDependency(EditorPlatformNode.StaticGetFullName(HostPlatform, TargetPlatform));
-                }
-            }
 
             bool bIsShared = false;
             // is this the "base game" or a non code project?
@@ -3265,10 +3164,6 @@ partial class GUBP
 			var Agenda = new UE4Build.BuildAgenda();
 
 			string AddArgs = "-nobuilduht";
-			if (bp.bOrthogonalizeEditorPlatforms)
-			{
-				AddArgs += " -skipnonhostplatforms";
-			}
 
 			Agenda.AddTargets(
 				new string[] { bp.Branch.BaseEngineProject.Properties.Targets[TargetRules.TargetType.Editor].TargetName },
