Fix typo in variable name.

diff --git a/Engine/Source/Programs/AutomationTool/GUBP/GUBP.cs b/Engine/Source/Programs/AutomationTool/GUBP/GUBP.cs
--- a/Engine/Source/Programs/AutomationTool/GUBP/GUBP.cs
+++ b/Engine/Source/Programs/AutomationTool/GUBP/GUBP.cs
@@ -1540,14 +1540,14 @@ public partial class GUBP : BuildCommand
             Log("Took {0}s to get history for {1} nodes", BuildDuration / 1000, NodesToDo.Count);
         }*/
 
-        List<NodeInfo> OrdereredToDo = TopologicalSort(NodesToDo, ExplicitTrigger, false, false);
+        List<NodeInfo> OrderedToDo = TopologicalSort(NodesToDo, ExplicitTrigger, false, false);
 
-		List<NodeInfo> UnfinishedTriggers = FindUnfinishedTriggers(bSkipTriggers, ExplicitTrigger, OrdereredToDo);
+		List<NodeInfo> UnfinishedTriggers = FindUnfinishedTriggers(bSkipTriggers, ExplicitTrigger, OrderedToDo);
 
         LogVerbose("*********** Desired And Dependent Nodes, in order.");
-        PrintNodes(this, OrdereredToDo, GUBPNodesHistory, UnfinishedTriggers);		
+        PrintNodes(this, OrderedToDo, GUBPNodesHistory, UnfinishedTriggers);		
         //check sorting
-		CheckSortOrder(OrdereredToDo);
+		CheckSortOrder(OrderedToDo);
 
         string FakeFail = ParseParamValue("FakeFail");
         if(CommanderSetup)
@@ -1561,15 +1561,15 @@ public partial class GUBP : BuildCommand
 
 			Dictionary<string, int> FullNodeListSortKey = GetDisplayOrder(FullNodeList.Keys.ToList(), FullNodeDirectDependencies, GUBPNodes);
 
-			DoCommanderSetup(TimeIndex, bSkipTriggers, bFake, bFakeEC, CLString, ExplicitTrigger, FullNodeList, FullNodeDirectDependencies, FullNodeDependedOnBy, DependentPromotions, SeparatePromotables, FullNodeListSortKey, GUBPNodesHistory, OrdereredToDo, UnfinishedTriggers, FakeFail);
+			DoCommanderSetup(TimeIndex, bSkipTriggers, bFake, bFakeEC, CLString, ExplicitTrigger, FullNodeList, FullNodeDirectDependencies, FullNodeDependedOnBy, DependentPromotions, SeparatePromotables, FullNodeListSortKey, GUBPNodesHistory, OrderedToDo, UnfinishedTriggers, FakeFail);
         }
 		else if(ParseParam("SaveGraph"))
 		{
-			SaveGraphVisualization(OrdereredToDo);
+			SaveGraphVisualization(OrderedToDo);
 		}
 		else if(!bListOnly)
 		{
-			ExecuteNodes(OrdereredToDo, bOnlyNode, bFake, bFakeEC, bSaveSharedTempStorage, GUBPNodesHistory, CLString, StoreName, FakeFail);
+			ExecuteNodes(OrderedToDo, bOnlyNode, bFake, bFakeEC, bSaveSharedTempStorage, GUBPNodesHistory, CLString, StoreName, FakeFail);
 		}
         PrintRunTime();
 	}
