Remove unused node causing warning which is treated as an error on Mac.

diff --git a/Engine/Source/Programs/AutomationTool/GUBP/LegacyNodes.cs b/Engine/Source/Programs/AutomationTool/GUBP/LegacyNodes.cs
--- a/Engine/Source/Programs/AutomationTool/GUBP/LegacyNodes.cs
+++ b/Engine/Source/Programs/AutomationTool/GUBP/LegacyNodes.cs
@@ -2312,113 +2312,6 @@ partial class GUBP
         }
     }
 
-
-    public class DDCNode : HostPlatformNode
-    {
-        BranchInfo.BranchUProject GameProj;
-        UnrealTargetPlatform TargetPlatform;
-        string CookPlatform;
-        bool bIsMassive;
-
-        public DDCNode(GUBP bp, UnrealTargetPlatform InHostPlatform, BranchInfo.BranchUProject InGameProj, UnrealTargetPlatform InTargetPlatform, string InCookPlatform)
-            : base(InHostPlatform)
-        {
-            GameProj = InGameProj;
-            TargetPlatform = InTargetPlatform;
-            CookPlatform = InCookPlatform;
-            bIsMassive = false;
-            AddDependency(RootEditorNode.StaticGetFullName(HostPlatform));
-
-            bool bIsShared = false;
-            // is this the "base game" or a non code project?
-            if (InGameProj.GameName != bp.Branch.BaseEngineProject.GameName && GameProj.Properties.Targets.ContainsKey(TargetRules.TargetType.Editor))
-            {
-                var Options = InGameProj.Options(HostPlatform);
-                bIsMassive = Options.bIsMassive;
-
-                AddDependency(EditorGameNode.StaticGetFullName(HostPlatform, GameProj));
-                // add an arc to prevent DDCNode from running until promotable is labeled
-                if (Options.bIsPromotable)
-                {
-                    if (Options.bSeparateGamePromotion)
-                    {
-                       // AddPseudodependency(GameLabelPromotableNode.StaticGetFullName(GameProj, false));
-                    }
-                    else
-                    {
-                        bIsShared = true;
-                    }
-                }
-                else if (Options.bTestWithShared)
-                {
-                    bIsShared = true;
-                }
-                //AddPseudodependency(GamePlatformMonolithicsNode.StaticGetFullName(HostPlatform, GameProj, TargetPlatform));
-            }
-            else
-            {
-                bIsShared = true;
-                //AddPseudodependency(GamePlatformMonolithicsNode.StaticGetFullName(HostPlatform, bp.Branch.BaseEngineProject, TargetPlatform));
-            }
-            if (bIsShared)
-            {
-                // add an arc to prevent cooks from running until promotable is labeled
-                //AddPseudodependency(WaitForTestShared.StaticGetFullName());
-                //AgentSharingGroup = "SharedCooks" + StaticGetHostPlatformSuffix(HostPlatform);
-
-                // If the cook fails for the base engine, don't bother trying
-                if (InGameProj.GameName != bp.Branch.BaseEngineProject.GameName && bp.HasNode(DDCNode.StaticGetFullName(HostPlatform, bp.Branch.BaseEngineProject, CookPlatform)))
-                {
-                    //AddPseudodependency(DDCNode.StaticGetFullName(HostPlatform, bp.Branch.BaseEngineProject, CookPlatform));
-                }
-
-                // If the base cook platform fails, don't bother trying other ones
-                string BaseCookedPlatform = Platform.Platforms[HostPlatform].GetCookPlatform(false, false, "");
-                if (InGameProj.GameName == bp.Branch.BaseEngineProject.GameName && CookPlatform != BaseCookedPlatform &&
-                    bp.HasNode(DDCNode.StaticGetFullName(HostPlatform, bp.Branch.BaseEngineProject, BaseCookedPlatform)))
-                {
-                    //AddPseudodependency(DDCNode.StaticGetFullName(HostPlatform, bp.Branch.BaseEngineProject, BaseCookedPlatform));
-                }
-            }
-        }
-        public static string StaticGetFullName(UnrealTargetPlatform InHostPlatform, BranchInfo.BranchUProject InGameProj, string InCookPlatform)
-        {
-            return InGameProj.GameName + "_" + InCookPlatform.Replace("+", "_") + "_DDC" + StaticGetHostPlatformSuffix(InHostPlatform);
-        }
-        public override string GetFullName()
-        {
-            return StaticGetFullName(HostPlatform, GameProj, CookPlatform);
-        }
-        public override string GameNameIfAnyForTempStorage()
-        {
-            return GameProj.GameName;
-        }
-        public override float Priority()
-        {
-            return base.Priority() + 10.0f;
-        }
-        public override int AgentMemoryRequirement(GUBP bp)
-        {
-            return bIsMassive ? 32 : 0;
-        }
-        public override int TimeoutInMinutes()
-        {
-            return bIsMassive ? 240 : base.TimeoutInMinutes();
-        }
-
-        public override void DoBuild(GUBP bp)
-        {
-            if (HostPlatform == UnrealTargetPlatform.Mac)
-            {
-                // not sure if we need something here or if the cook commandlet will automatically convert the exe name
-            }
-            CommandUtils.DDCCommandlet(GameProj.FilePath, "UE4Editor-Cmd.exe", null, CookPlatform, "-fill");
-
-            BuildProducts = new List<string>();
-            SaveRecordOfSuccessAndAddToBuildProducts();
-        }
-    }
-
     public class GamePlatformCookedAndCompiledNode : HostPlatformAggregateNode
     {
         BranchInfo.BranchUProject GameProj;
