Simplify the representation of trigger chains (just have a list of triggers). Also fix a bug where nodes cannot be dependent on other nodes behind different triggers in the chain.

diff --git a/Engine/Source/Programs/AutomationTool/GUBP/GUBP.cs b/Engine/Source/Programs/AutomationTool/GUBP/GUBP.cs
--- a/Engine/Source/Programs/AutomationTool/GUBP/GUBP.cs
+++ b/Engine/Source/Programs/AutomationTool/GUBP/GUBP.cs
@@ -31,10 +31,14 @@ public partial class GUBP : BuildCommand
 		public GUBPNode Node;
 		public List<NodeInfo> Dependencies;
 		public List<NodeInfo> PseudoDependencies;
-		public NodeInfo ControllingTrigger;
-		public string ControllingTriggerDotName;
+		public List<NodeInfo> ControllingTriggers;
 		public bool IsComplete;
 
+		public string ControllingTriggerDotName
+		{
+			get { return String.Join(".", ControllingTriggers.Select(x => x.Name)); }
+		}
+
 		public override string ToString()
 		{
 			System.Diagnostics.Trace.TraceWarning("Implicit conversion from NodeInfo to string\n{0}", Environment.StackTrace);
@@ -169,53 +173,51 @@ public partial class GUBP : BuildCommand
 
     static void FindControllingTriggers(NodeInfo NodeToDo)
     {
-		if(NodeToDo.ControllingTriggerDotName == null)
+		if(NodeToDo.ControllingTriggers == null)
 		{
+			NodeToDo.ControllingTriggers = new List<NodeInfo>();
+
 			// Find all the dependencies of this node
 			List<NodeInfo> AllDependencies = new List<NodeInfo>();
 			AllDependencies.AddRange(NodeToDo.Dependencies);
 			AllDependencies.AddRange(NodeToDo.PseudoDependencies);
 
 			// Find the immediate trigger controlling this one
+			List<NodeInfo> PreviousTriggers = new List<NodeInfo>();
 			foreach (NodeInfo Dependency in AllDependencies)
 			{
 				FindControllingTriggers(Dependency);
 
-				NodeInfo Trigger = null;
+				NodeInfo PreviousTrigger = null;
 				if(Dependency.Node.TriggerNode())
 				{
-					Trigger = Dependency;
+					PreviousTrigger = Dependency;
 				}
-				else
+				else if(Dependency.ControllingTriggers.Count > 0)
 				{
-					Trigger = Dependency.ControllingTrigger;
+					PreviousTrigger = Dependency.ControllingTriggers.Last();
 				}
 				
-				if(Trigger != null)
+				if(PreviousTrigger != null && !PreviousTriggers.Contains(PreviousTrigger))
 				{
-					if(NodeToDo.ControllingTrigger == null)
-					{
-						NodeToDo.ControllingTrigger = Trigger;
-					}
-					else if(NodeToDo.ControllingTrigger != Trigger)
-					{
-						throw new AutomationException("Node {0} has two controlling triggers {1} and {2}.", NodeToDo.Name, NodeToDo.ControllingTrigger.Name, Trigger.Name);
-					}
+					PreviousTriggers.Add(PreviousTrigger);
 				}
 			}
 
-			// Set the path of controlling triggers for this node
-			if(NodeToDo.ControllingTrigger == null)
-			{
-				NodeToDo.ControllingTriggerDotName = "";
-			}
-			else if(NodeToDo.ControllingTrigger.ControllingTriggerDotName.Length == 0)
+			// Remove previous triggers from the list that aren't the last in the chain. If there's a trigger chain of X.Y.Z, and a node has dependencies behind all three, the only trigger we care about is Z.
+			PreviousTriggers.RemoveAll(x => PreviousTriggers.Any(y => y.ControllingTriggers.Contains(x)));
+
+			// We only support one direct controlling trigger at the moment (though it may be in a chain with other triggers)
+			if(PreviousTriggers.Count > 1)
 			{
-				NodeToDo.ControllingTriggerDotName = NodeToDo.ControllingTrigger.Name;
+				throw new AutomationException("Node {0} has multiple controlling triggers: {1}", NodeToDo.Name, String.Join(", ", PreviousTriggers.Select(x => x.Name)));
 			}
-			else
+
+			// Update the list of controlling triggers
+			if(PreviousTriggers.Count == 1)
 			{
-				NodeToDo.ControllingTriggerDotName = NodeToDo.ControllingTrigger.ControllingTriggerDotName + "." + NodeToDo.ControllingTrigger.Name;
+				NodeToDo.ControllingTriggers.AddRange(PreviousTriggers[0].ControllingTriggers);
+				NodeToDo.ControllingTriggers.Add(PreviousTriggers[0]);
 			}
 		}
 	}
@@ -477,8 +479,7 @@ public partial class GUBP : BuildCommand
                         string Finished = "";
                         if (UnfinishedTriggers != null)
                         {
-                            NodeInfo MyShortControllingTrigger = NodeToDo.ControllingTrigger;
-                            if (UnfinishedTriggers.Contains(MyShortControllingTrigger))
+                            if (NodeToDo.ControllingTriggers.Count > 0 && UnfinishedTriggers.Contains(NodeToDo.ControllingTriggers.Last()))
                             {
                                 Finished = "(not yet triggered)";
                             }
@@ -2239,7 +2240,7 @@ public partial class GUBP : BuildCommand
 			// remove nodes that have unfinished triggers
 			foreach (NodeInfo NodeToDo in OrdereredToDo)
 			{
-				NodeInfo ControllingTrigger = NodeToDo.ControllingTrigger;
+				NodeInfo ControllingTrigger = (NodeToDo.ControllingTriggers.Count > 0)? NodeToDo.ControllingTriggers.Last() : null;
 				bool bNoUnfinishedTriggers = !UnfinishedTriggers.Contains(ControllingTrigger);
 
 				if (bNoUnfinishedTriggers)
