Move the computed frequency shift out of the public API for GUBP.

diff --git a/Engine/Source/Programs/AutomationTool/GUBP/GUBP.cs b/Engine/Source/Programs/AutomationTool/GUBP/GUBP.cs
--- a/Engine/Source/Programs/AutomationTool/GUBP/GUBP.cs
+++ b/Engine/Source/Programs/AutomationTool/GUBP/GUBP.cs
@@ -32,6 +32,7 @@ public partial class GUBP : BuildCommand
 		public List<NodeInfo> Dependencies;
 		public List<NodeInfo> PseudoDependencies;
 		public List<NodeInfo> ControllingTriggers;
+		public int FrequencyShift = -1;
 		public bool IsComplete;
 
 		public string ControllingTriggerDotName
@@ -225,8 +226,8 @@ public partial class GUBP : BuildCommand
     string CISFrequencyQuantumShiftString(NodeInfo NodeToDo)
     {
         string FrequencyString = "";
-        int Quantum = NodeToDo.Node.DependentCISFrequencyQuantumShift();
-        if (Quantum > 0)
+        int Quantum = NodeToDo.FrequencyShift;
+		if (Quantum > 0)
         {
 			int TimeQuantum = 20;
 			if(BranchOptions.QuantumOverride != 0)
@@ -248,7 +249,7 @@ public partial class GUBP : BuildCommand
 
     int ComputeDependentCISFrequencyQuantumShift(NodeInfo NodeToDo, Dictionary<NodeInfo, int> FrequencyOverrides)
     {
-        int Result = NodeToDo.Node.ComputedDependentCISFrequencyQuantumShift;        
+        int Result = NodeToDo.FrequencyShift;        
         if (Result < 0)
         {
             Result = NodeToDo.Node.CISFrequencyQuantumShift(this);
@@ -272,7 +273,7 @@ public partial class GUBP : BuildCommand
             {
                 throw new AutomationException("Failed to compute shift.");
             }
-            NodeToDo.Node.ComputedDependentCISFrequencyQuantumShift = Result;
+            NodeToDo.FrequencyShift = Result;
         }
         return Result;
     }
@@ -1841,9 +1842,9 @@ public partial class GUBP : BuildCommand
 			LogVerbose("Check to make sure we didn't ask for nodes that will be culled by time index");
 			foreach (NodeInfo NodeToDo in NodesToDo)
 			{
-				if (TimeIndex % (1 << NodeToDo.Node.DependentCISFrequencyQuantumShift()) != 0)
+				if (TimeIndex % (1 << NodeToDo.FrequencyShift) != 0)
 				{
-					throw new AutomationException("You asked specifically for node {0}, but it is culled by the time quantum: TimeIndex = {1}, DependentCISFrequencyQuantumShift = {2}.", NodeToDo.Name, TimeIndex, NodeToDo.Node.DependentCISFrequencyQuantumShift());
+					throw new AutomationException("You asked specifically for node {0}, but it is culled by the time quantum: TimeIndex = {1}, DependentCISFrequencyQuantumShift = {2}.", NodeToDo.Name, TimeIndex, NodeToDo.FrequencyShift);
 				}
 			}
 		}
@@ -1923,7 +1924,7 @@ public partial class GUBP : BuildCommand
 			HashSet<NodeInfo> NewNodesToDo = new HashSet<NodeInfo>();
 			foreach (NodeInfo NodeToDo in NodesToDo)
 			{
-				if (TimeIndex % (1 << NodeToDo.Node.DependentCISFrequencyQuantumShift()) == 0)
+				if (TimeIndex % (1 << NodeToDo.FrequencyShift) == 0)
 				{
 					LogVerbose("  Keeping {0}", NodeToDo.Name);
 					NewNodesToDo.Add(NodeToDo);
@@ -2957,10 +2958,10 @@ public partial class GUBP : BuildCommand
 		foreach(string NodeName in NodeNames)
 		{
 			List<string> NodesByThisFrequency;
-			if(!NodesByFrequency.TryGetValue(GUBPNodes[NodeName].Node.DependentCISFrequencyQuantumShift(), out NodesByThisFrequency))
+			if(!NodesByFrequency.TryGetValue(GUBPNodes[NodeName].FrequencyShift, out NodesByThisFrequency))
 			{
 				NodesByThisFrequency = new List<string>();
-				NodesByFrequency.Add(GUBPNodes[NodeName].Node.DependentCISFrequencyQuantumShift(), NodesByThisFrequency);
+				NodesByFrequency.Add(GUBPNodes[NodeName].FrequencyShift, NodesByThisFrequency);
 			}
 			NodesByThisFrequency.Add(NodeName);
 		}
diff --git a/Engine/Source/Programs/AutomationTool/GUBP/LegacyNodes.cs b/Engine/Source/Programs/AutomationTool/GUBP/LegacyNodes.cs
--- a/Engine/Source/Programs/AutomationTool/GUBP/LegacyNodes.cs
+++ b/Engine/Source/Programs/AutomationTool/GUBP/LegacyNodes.cs
@@ -20,7 +20,6 @@ partial class GUBP
         public List<string> AllDependencyBuildProducts = null;
         public List<string> AllDependencies = null;		
         public string AgentSharingGroup = "";
-        public int ComputedDependentCISFrequencyQuantumShift = -1;
 
         public virtual string GetFullName()
         {
@@ -101,17 +100,6 @@ partial class GUBP
         {
             return 0;
         }
-        /// <summary>
-        /// As above the maximum of all ancestors and pseudoancestors
-        /// </summary>
-        public int DependentCISFrequencyQuantumShift()
-        {
-            if (ComputedDependentCISFrequencyQuantumShift < 0)
-            {
-                throw new AutomationException("Asked for frequency shift before it was computed.");
-            }
-            return ComputedDependentCISFrequencyQuantumShift;
-        }
 
         public virtual float Priority()
         {
