Remove EmailHint instance variable; it's only ever needed in one place, it's just not being passed as a parameter.

diff --git a/Engine/Source/Programs/AutomationTool/GUBP/GUBP.cs b/Engine/Source/Programs/AutomationTool/GUBP/GUBP.cs
--- a/Engine/Source/Programs/AutomationTool/GUBP/GUBP.cs
+++ b/Engine/Source/Programs/AutomationTool/GUBP/GUBP.cs
@@ -17,7 +17,6 @@ public partial class GUBP : BuildCommand
     public BranchInfo Branch = null;
     public List<UnrealTargetPlatform> HostPlatforms;
     public static bool bForceIncrementalCompile = false;
-    public string EmailHint;
     static public bool bPreflightBuild = false;
     public int PreflightShelveCL = 0;
     static public string PreflightMangleSuffix = "";
diff --git a/Engine/Source/Programs/AutomationTool/GUBP/LegacyBranchSetup.cs b/Engine/Source/Programs/AutomationTool/GUBP/LegacyBranchSetup.cs
--- a/Engine/Source/Programs/AutomationTool/GUBP/LegacyBranchSetup.cs
+++ b/Engine/Source/Programs/AutomationTool/GUBP/LegacyBranchSetup.cs
@@ -118,8 +118,8 @@ partial class GUBP
     }
 
     public abstract class GUBPEmailHacker
-    {        
-        public virtual List<string> AddEmails(GUBP bp, string Branch, string NodeName)
+    {
+		public virtual List<string> AddEmails(GUBP bp, string Branch, string NodeName, string EmailHint)
         {
             return new List<string>();
         }
@@ -145,7 +145,7 @@ partial class GUBP
         {
             return OnlyEmail;
         }
-        EmailHint = ParseParamValue("EmailHint");
+        string EmailHint = ParseParamValue("EmailHint");
         if(EmailHint == null)
         {
             EmailHint = "";
@@ -177,7 +177,7 @@ partial class GUBP
 		}
         foreach (var EmailHacker in EmailHackers)
         {
-            Result.AddRange(EmailHacker.AddEmails(this, Branch, NodeInfo.Name));
+            Result.AddRange(EmailHacker.AddEmails(this, Branch, NodeInfo.Name, EmailHint));
         }
         foreach (var EmailHacker in EmailHackers)
         {
