Make temp storage suffixes const.

diff --git a/Engine/Source/Programs/AutomationTool/GUBP/GUBP.cs b/Engine/Source/Programs/AutomationTool/GUBP/GUBP.cs
--- a/Engine/Source/Programs/AutomationTool/GUBP/GUBP.cs
+++ b/Engine/Source/Programs/AutomationTool/GUBP/GUBP.cs
@@ -11,6 +11,10 @@ using System.Linq;
 
 public partial class GUBP : BuildCommand
 {
+	const string StartedTempStorageSuffix = "_Started";
+	const string FailedTempStorageSuffix = "_Failed";
+	const string SucceededTempStorageSuffix = "_Succeeded";
+
 	public string BranchName;
     public int CL = 0;
     public bool bSignBuildProducts = false;
@@ -3025,8 +3029,4 @@ public partial class GUBP : BuildCommand
 			SortedNodes.Add(NodeName, SortedNodes.Count);
 		}
 	}
-
-    string StartedTempStorageSuffix = "_Started";
-    string FailedTempStorageSuffix = "_Failed";
-    string SucceededTempStorageSuffix = "_Succeeded";
 }
