Convert a few GUBP variables into instance variables rather than statics.

diff --git a/Engine/Source/Programs/AutomationTool/GUBP/GUBP.cs b/Engine/Source/Programs/AutomationTool/GUBP/GUBP.cs
--- a/Engine/Source/Programs/AutomationTool/GUBP/GUBP.cs
+++ b/Engine/Source/Programs/AutomationTool/GUBP/GUBP.cs
@@ -20,10 +20,9 @@ public partial class GUBP : BuildCommand
     public bool bSignBuildProducts = false;
     public BranchInfo Branch = null;
     public List<UnrealTargetPlatform> HostPlatforms;
-    public static bool bForceIncrementalCompile = false;
-    static public bool bPreflightBuild = false;
-    public int PreflightShelveCL = 0;
-    static public string PreflightMangleSuffix = "";
+    public bool bForceIncrementalCompile = false;
+    public bool bPreflightBuild = false;
+    public string PreflightMangleSuffix = "";
     public GUBPBranchHacker.BranchOptions BranchOptions = null;	
 
     Dictionary<string, NodeInfo> GUBPNodes;
@@ -1290,6 +1289,8 @@ public partial class GUBP : BuildCommand
     public override void ExecuteBuild()
     {
         Log("************************* GUBP");
+
+		int PreflightShelveCL = 0;
         string PreflightShelveCLString = GetEnvVar("uebp_PreflightShelveCL");
         if ((!String.IsNullOrEmpty(PreflightShelveCLString) && IsBuildMachine) || ParseParam("PreflightTest"))
         {
diff --git a/Engine/Source/Programs/AutomationTool/GUBP/LegacyNodes.cs b/Engine/Source/Programs/AutomationTool/GUBP/LegacyNodes.cs
--- a/Engine/Source/Programs/AutomationTool/GUBP/LegacyNodes.cs
+++ b/Engine/Source/Programs/AutomationTool/GUBP/LegacyNodes.cs
@@ -361,7 +361,7 @@ partial class GUBP
             UE4Build.BuildAgenda Agenda = GetAgenda(bp);
             if (Agenda != null)
             {
-                bool ReallyDeleteBuildProducts = DeleteBuildProducts() && !GUBP.bForceIncrementalCompile;
+                bool ReallyDeleteBuildProducts = DeleteBuildProducts() && !bp.bForceIncrementalCompile;
                 Agenda.DoRetries = false; // these would delete build products
 				bool UseParallelExecutor = bDependentOnCompileTools && (HostPlatform == UnrealTargetPlatform.Win64);
 				UE4Build.Build(Agenda, InDeleteBuildProducts: ReallyDeleteBuildProducts, InUpdateVersionFiles: false, InForceNoXGE: true, InForceUnity: true, InUseParallelExecutor: UseParallelExecutor);
@@ -385,7 +385,7 @@ partial class GUBP
                 }
                 PostBuildProducts(bp);
             }
-            if (Agenda == null || (BuildProducts.Count == 0 && GUBP.bForceIncrementalCompile))
+            if (Agenda == null || (BuildProducts.Count == 0 && bp.bForceIncrementalCompile))
             {
                 SaveRecordOfSuccessAndAddToBuildProducts("Nothing to actually compile");
             }
@@ -449,7 +449,7 @@ partial class GUBP
         public override UE4Build.BuildAgenda GetAgenda(GUBP bp)
         {
             var Agenda = new UE4Build.BuildAgenda();
-			if (HostPlatform == UnrealTargetPlatform.Win64 && !GUBP.bForceIncrementalCompile)
+			if (HostPlatform == UnrealTargetPlatform.Win64 && !bp.bForceIncrementalCompile)
             {
                 Agenda.DotNetProjects.AddRange(
                     new string[] 
@@ -558,7 +558,7 @@ partial class GUBP
         }
         void DeleteStaleDLLs(GUBP bp)
         {
-            if (GUBP.bForceIncrementalCompile)
+            if (bp.bForceIncrementalCompile)
             {
                 return;
             }
@@ -666,7 +666,7 @@ partial class GUBP
 		}
         void DeleteStaleDLLs(GUBP bp)
         {
-            if (GUBP.bForceIncrementalCompile)
+            if (bp.bForceIncrementalCompile)
             {
                 return;
             }
@@ -768,7 +768,7 @@ partial class GUBP
 
 			if (HostPlatform == UnrealTargetPlatform.Win64)
             {
-				if (!GUBP.bForceIncrementalCompile)
+				if (!bp.bForceIncrementalCompile)
 				{
 					Agenda.DotNetProjects.AddRange(
 						new string[] 
@@ -783,7 +783,7 @@ partial class GUBP
 							CombinePaths(@"Engine\Source\Programs\NetworkProfiler\NetworkProfiler.sln"),   
 						}
                     );
-                if (!GUBP.bForceIncrementalCompile)
+                if (!bp.bForceIncrementalCompile)
                 {
 					Agenda.SwarmProject = CombinePaths(@"Engine\Source\Programs\UnrealSwarm\UnrealSwarm.sln");
 				}
@@ -2011,10 +2011,6 @@ partial class GUBP
                 //shared promotable has a shorter name
                 CompleteLabelPrefix = "Promotable";
             }                   
-            if (GUBP.bPreflightBuild)
-            {
-                CompleteLabelPrefix = CompleteLabelPrefix + PreflightMangleSuffix;
-            }
             return CompleteLabelPrefix;
         }
 
@@ -2032,7 +2028,7 @@ partial class GUBP
         {
             BuildProducts = new List<string>();
 
-            if (P4Enabled)
+            if (P4Enabled && !bp.bPreflightBuild)
             {
                 if (AllDependencyBuildProducts.Count == 0)
                 {
@@ -2587,14 +2583,14 @@ partial class GUBP
 		{
 			return base.CISFrequencyQuantumShift(bp) + 3;
 		}
-        public static string GetArchiveDirectory(BranchInfo.BranchUProject InGameProj, UnrealTargetPlatform InHostPlatform, List<UnrealTargetPlatform> InClientTargetPlatforms = null, List<UnrealTargetConfiguration> InClientConfigs = null, List<UnrealTargetPlatform> InServerTargetPlatforms = null, List<UnrealTargetConfiguration> InServerConfigs = null, bool InClientNotGame = false)
+        public static string GetArchiveDirectory(GUBP bp, BranchInfo.BranchUProject InGameProj, UnrealTargetPlatform InHostPlatform, List<UnrealTargetPlatform> InClientTargetPlatforms = null, List<UnrealTargetConfiguration> InClientConfigs = null, List<UnrealTargetPlatform> InServerTargetPlatforms = null, List<UnrealTargetConfiguration> InServerConfigs = null, bool InClientNotGame = false)
         {
             string BaseDir = TempStorage.ResolveSharedBuildDirectory(InGameProj.GameName);
             string NodeName = StaticGetFullName(InGameProj, InHostPlatform, InClientTargetPlatforms, InClientConfigs, InServerTargetPlatforms, InServerConfigs, InClientNotGame);
             string Inner = P4Env.BuildRootEscaped + "-CL-" + P4Env.ChangelistString;
-            if (GUBP.bPreflightBuild)
+            if (bp.bPreflightBuild)
             {
-                Inner = Inner + PreflightMangleSuffix;
+                Inner = Inner + bp.PreflightMangleSuffix;
             }
             string ArchiveDirectory = CombinePaths(BaseDir, NodeName, Inner);
             return ArchiveDirectory;
@@ -2710,7 +2706,7 @@ partial class GUBP
 			string IntermediateArchiveDirectory = FinalArchiveDirectory;
             if (P4Enabled)
             {
-                FinalArchiveDirectory = GetArchiveDirectory(GameProj, HostPlatform, ClientTargetPlatforms, ClientConfigs, ServerTargetPlatforms, ServerConfigs, ClientNotGame);
+                FinalArchiveDirectory = GetArchiveDirectory(bp, GameProj, HostPlatform, ClientTargetPlatforms, ClientConfigs, ServerTargetPlatforms, ServerConfigs, ClientNotGame);
 				IntermediateArchiveDirectory = FinalArchiveDirectory;
 				// Xbox One packaging does not function with remote file systems. Use a temp local directory to package and then move files into final location.
 				if (bXboxOneTarget)
@@ -2843,7 +2839,7 @@ partial class GUBP
                 ProjectArg = " -project=\"" + GameProj.FilePath + "\"";
             }
 
-            string ArchiveDirectory = FormalBuildNode.GetArchiveDirectory(GameProj, HostPlatform, new List<UnrealTargetPlatform>() { ClientTargetPlatform }, InClientConfigs: new List<UnrealTargetConfiguration>() { ClientConfig }, InClientNotGame: GameOrClient == TargetRules.TargetType.Client);
+            string ArchiveDirectory = FormalBuildNode.GetArchiveDirectory(bp, GameProj, HostPlatform, new List<UnrealTargetPlatform>() { ClientTargetPlatform }, InClientConfigs: new List<UnrealTargetConfiguration>() { ClientConfig }, InClientNotGame: GameOrClient == TargetRules.TargetType.Client);
             if (!DirectoryExists_NoExceptions(ArchiveDirectory))
             {
                 throw new AutomationException("Archive directory does not exist {0}, so we can't test the build.", ArchiveDirectory);
diff --git a/Engine/Source/Programs/AutomationTool/Scripts/RocketBuild.Automation.cs b/Engine/Source/Programs/AutomationTool/Scripts/RocketBuild.Automation.cs
--- a/Engine/Source/Programs/AutomationTool/Scripts/RocketBuild.Automation.cs
+++ b/Engine/Source/Programs/AutomationTool/Scripts/RocketBuild.Automation.cs
@@ -131,7 +131,7 @@ namespace Rocket
 
 				// Get the output directory for the build zips
 				string PublishedEngineDir;
-				if(ShouldDoSeriousThingsLikeP4CheckinAndPostToMCP())
+				if(ShouldDoSeriousThingsLikeP4CheckinAndPostToMCP(bp))
 				{
 					PublishedEngineDir = CommandUtils.CombinePaths(CommandUtils.RootSharedTempStorageDirectory(), "Rocket", "Automated", GetBuildLabel(), CommandUtils.GetGenericPlatformName(HostPlatform));
 				}
@@ -155,7 +155,7 @@ namespace Rocket
 				WaitForPromotionNode.AddDependency(PublishRocketSymbolsNode.StaticGetFullName(HostPlatform));
 
 				// Push everything behind the promotion triggers if we're doing things on the build machines
-				if(ShouldDoSeriousThingsLikeP4CheckinAndPostToMCP() || bp.ParseParam("WithRocketPromotable"))
+				if(ShouldDoSeriousThingsLikeP4CheckinAndPostToMCP(bp) || bp.ParseParam("WithRocketPromotable"))
 				{
 					string WaitForTrigger = GUBP.WaitForSharedPromotionUserInput.StaticGetFullName(false);
 
@@ -247,9 +247,9 @@ namespace Rocket
 			return CommandUtils.CombineCommandletParams(CookPlatforms.Distinct().ToArray());
 		}
 
-		public static bool ShouldDoSeriousThingsLikeP4CheckinAndPostToMCP()
+		public static bool ShouldDoSeriousThingsLikeP4CheckinAndPostToMCP(GUBP bp)
 		{
-			return CommandUtils.P4Enabled && CommandUtils.AllowSubmit && !GUBP.bPreflightBuild; // we don't do serious things in a preflight
+			return CommandUtils.P4Enabled && CommandUtils.AllowSubmit && !bp.bPreflightBuild; // we don't do serious things in a preflight
 		}
 
 		public static UnrealTargetPlatform GetSourceHostPlatform(GUBP bp, UnrealTargetPlatform HostPlatform, UnrealTargetPlatform TargetPlatform)
@@ -903,7 +903,7 @@ namespace Rocket
 
 		public override void DoBuild(GUBP bp)
 		{
-			if(RocketBuild.ShouldDoSeriousThingsLikeP4CheckinAndPostToMCP())
+			if(RocketBuild.ShouldDoSeriousThingsLikeP4CheckinAndPostToMCP(bp))
 			{
 				// Make a lookup for all the known debug extensions, and filter all the dependency build products against that
 				HashSet<string> DebugExtensions = new HashSet<string>(Platform.Platforms.Values.SelectMany(x => x.GetDebugFileExtentions()).Distinct().ToArray(), StringComparer.InvariantCultureIgnoreCase);
