Remove some (I think) rarely used GUBP debugging commands in the interest of simplicity (-RelatedToNode, -SetupNode, -SetupRelatedToNode, -OnlyNode).

diff --git a/Engine/Source/Programs/AutomationTool/GUBP/GUBP.cs b/Engine/Source/Programs/AutomationTool/GUBP/GUBP.cs
--- a/Engine/Source/Programs/AutomationTool/GUBP/GUBP.cs
+++ b/Engine/Source/Programs/AutomationTool/GUBP/GUBP.cs
@@ -1468,10 +1468,7 @@ public partial class GUBP : BuildCommand
             ComputeDependentCISFrequencyQuantumShift(NodeToDo.Value, FrequencyOverrides);
         }
 
-		bool bOnlyNode;
-		HashSet<NodeInfo> NodesToDo;
-		Dictionary<string, string> FullNodeDependedOnBy = new Dictionary<string, string>();
-		ParseNodesToDo(WithoutLinux, TimeIndex, CommanderSetup, FullNodeDependedOnBy, out bOnlyNode, out NodesToDo);
+		HashSet<NodeInfo> NodesToDo = ParseNodesToDo(WithoutLinux, TimeIndex, CommanderSetup);
 
 		NodeInfo ExplicitTrigger = null;
 		if (CommanderSetup)
@@ -1509,25 +1506,6 @@ public partial class GUBP : BuildCommand
                 }
             }
         }
-		if (bPreflightBuild)
-		{
-			LogVerbose("Culling triggers and downstream for preflight builds ");
-			HashSet<NodeInfo> NewNodesToDo = new HashSet<NodeInfo>();
-			foreach (NodeInfo NodeToDo in NodesToDo)
-			{
-				string TriggerDot = NodeToDo.ControllingTriggerDotName;
-				if (TriggerDot == "" && !NodeToDo.Node.TriggerNode())
-				{
-					LogVerbose("  Keeping {0}", NodeToDo.Name);
-					NewNodesToDo.Add(NodeToDo);
-				}
-				else
-				{
-					LogVerbose("  Rejecting {0}", NodeToDo.Name);
-				}
-			}
-			NodesToDo = NewNodesToDo;
-		}
 
         Dictionary<string, NodeHistory> GUBPNodesHistory = new Dictionary<string, NodeHistory>();
 
@@ -1558,6 +1536,8 @@ public partial class GUBP : BuildCommand
         string FakeFail = ParseParamValue("FakeFail");
         if(CommanderSetup)
         {
+			Dictionary<string, string> FullNodeDependedOnBy = GetFullNodeDependedOnBy(NodesToDo);
+
 			List<NodeInfo> SeparatePromotables = FindPromotables(GUBPNodes.Values);
 			Dictionary<NodeInfo, List<NodeInfo>> DependentPromotions = FindDependentPromotables(GUBPNodes.Values, SeparatePromotables);
 
@@ -1575,7 +1555,7 @@ public partial class GUBP : BuildCommand
 		}
 		else if(!bListOnly)
 		{
-			ExecuteNodes(OrderedToDo, bOnlyNode, bFake, bFakeEC, bSaveSharedTempStorage, GUBPNodesHistory, CLString, StoreName, FakeFail);
+			ExecuteNodes(OrderedToDo, bFake, bFakeEC, bSaveSharedTempStorage, GUBPNodesHistory, CLString, StoreName, FakeFail);
 		}
         PrintRunTime();
 	}
@@ -1715,72 +1695,41 @@ public partial class GUBP : BuildCommand
 		return FrequencyOverrides;
 	}
 
-	private void ParseNodesToDo(bool WithoutLinux, int TimeIndex, bool CommanderSetup, Dictionary<string, string> FullNodeDependedOnBy, out bool bOnlyNode, out HashSet<NodeInfo> NodesToDo)
+	private HashSet<NodeInfo> ParseNodesToDo(bool WithoutLinux, int TimeIndex, bool CommanderSetup)
 	{
-		bOnlyNode = false;
-		bool bRelatedToNode = false;
-		NodesToDo = new HashSet<NodeInfo>();
+		HashSet<NodeInfo> NodesToDo = new HashSet<NodeInfo>();
 
+		string NodeSpec = ParseParamValue("Node");
+		if (!String.IsNullOrEmpty(NodeSpec))
 		{
-			string NodeSpec = ParseParamValue("Node");
-			if (String.IsNullOrEmpty(NodeSpec))
-			{
-				NodeSpec = ParseParamValue("RelatedToNode");
-				if (!String.IsNullOrEmpty(NodeSpec))
-				{
-					bRelatedToNode = true;
-				}
-			}
-			if (String.IsNullOrEmpty(NodeSpec) && CommanderSetup)
-			{
-				NodeSpec = ParseParamValue("SetupNode");
-				if (String.IsNullOrEmpty(NodeSpec))
-				{
-					NodeSpec = ParseParamValue("SetupRelatedToNode");
-					if (!String.IsNullOrEmpty(NodeSpec))
-					{
-						bRelatedToNode = true;
-					}
-				}
-			}
-			if (String.IsNullOrEmpty(NodeSpec))
+			List<string> Nodes = new List<string>(NodeSpec.Split('+'));
+			foreach (string NodeArg in Nodes)
 			{
-				NodeSpec = ParseParamValue("OnlyNode");
-				if (!String.IsNullOrEmpty(NodeSpec))
+				string NodeName = NodeArg.Trim();
+				bool bFoundAnything = false;
+				if (!String.IsNullOrEmpty(NodeName))
 				{
-					bOnlyNode = true;
-				}
-			}
-			if (!String.IsNullOrEmpty(NodeSpec))
-			{
-				List<string> Nodes = new List<string>(NodeSpec.Split('+'));
-				foreach (string NodeArg in Nodes)
-				{
-					string NodeName = NodeArg.Trim();
-					bool bFoundAnything = false;
-					if (!String.IsNullOrEmpty(NodeName))
+					foreach (KeyValuePair<string, NodeInfo> Node in GUBPNodes)
 					{
-						foreach (KeyValuePair<string, NodeInfo> Node in GUBPNodes)
+						if (Node.Value.Node.GetFullName().Equals(NodeArg, StringComparison.InvariantCultureIgnoreCase) ||
+							Node.Value.Node.AgentSharingGroup.Equals(NodeArg, StringComparison.InvariantCultureIgnoreCase)
+							)
 						{
-							if (Node.Value.Node.GetFullName().Equals(NodeArg, StringComparison.InvariantCultureIgnoreCase) ||
-								Node.Value.Node.AgentSharingGroup.Equals(NodeArg, StringComparison.InvariantCultureIgnoreCase)
-								)
+							if (!NodesToDo.Contains(Node.Value))
 							{
-								if (!NodesToDo.Contains(Node.Value))
-								{
-									NodesToDo.Add(Node.Value);
-								}
-								bFoundAnything = true;
+								NodesToDo.Add(Node.Value);
 							}
+							bFoundAnything = true;
 						}
-						if (!bFoundAnything)
-						{
-							throw new AutomationException("Could not find node named {0}", NodeName);
-						}
+					}
+					if (!bFoundAnything)
+					{
+						throw new AutomationException("Could not find node named {0}", NodeName);
 					}
 				}
 			}
 		}
+
 		string GameSpec = ParseParamValue("Game");
 		if (!String.IsNullOrEmpty(GameSpec))
 		{
@@ -1833,6 +1782,7 @@ public partial class GUBP : BuildCommand
 				}
 			}
 		}
+
 		if (NodesToDo.Count == 0)
 		{
 			LogVerbose("No nodes specified, adding all nodes");
@@ -1858,70 +1808,41 @@ public partial class GUBP : BuildCommand
 		{
 			LogVerbose("  {0}", NodeToDo.Name);
 		}
-		// if we are doing related to, then find things that depend on the selected nodes
-		if (bRelatedToNode)
-		{
-			bool bDoneWithDependencies = false;
-
-			while (!bDoneWithDependencies)
-			{
-				bDoneWithDependencies = true;
-				HashSet<NodeInfo> Fringe = new HashSet<NodeInfo>();
-				foreach (NodeInfo NodeToDo in GUBPNodes.Values)
-				{
-					if (!NodesToDo.Contains(NodeToDo))
-					{
-						foreach (NodeInfo Dep in NodeToDo.Dependencies)
-						{
-							if (NodesToDo.Contains(Dep))
-							{
-								Fringe.Add(NodeToDo);
-								bDoneWithDependencies = false;
-							}
-						}
-					}
-				}
-				NodesToDo.UnionWith(Fringe);
-			}
-		}
 
 		// find things that our nodes depend on
-		if (!bOnlyNode)
+		bool bDoneWithDependencies = false;
+		while (!bDoneWithDependencies)
 		{
-			bool bDoneWithDependencies = false;
-
-			while (!bDoneWithDependencies)
+			bDoneWithDependencies = true;
+			HashSet<NodeInfo> Fringe = new HashSet<NodeInfo>();
+			foreach (NodeInfo NodeToDo in NodesToDo)
 			{
-				bDoneWithDependencies = true;
-				HashSet<NodeInfo> Fringe = new HashSet<NodeInfo>();
-				foreach (NodeInfo NodeToDo in NodesToDo)
+				if (BranchOptions.NodesToRemovePseudoDependencies.Count > 0)
 				{
-					if (BranchOptions.NodesToRemovePseudoDependencies.Count > 0)
+					List<string> ListOfRetainedNodesToRemovePseudoDependencies = BranchOptions.NodesToRemovePseudoDependencies;
+					foreach (string NodeToRemovePseudoDependencies in BranchOptions.NodesToRemovePseudoDependencies)
 					{
-						List<string> ListOfRetainedNodesToRemovePseudoDependencies = BranchOptions.NodesToRemovePseudoDependencies;
-						foreach (string NodeToRemovePseudoDependencies in BranchOptions.NodesToRemovePseudoDependencies)
+						if (NodeToDo.Name == NodeToRemovePseudoDependencies)
 						{
-							if (NodeToDo.Name == NodeToRemovePseudoDependencies)
-							{
-								RemoveAllPseudodependenciesFromNode(NodeToDo.Name);
-								ListOfRetainedNodesToRemovePseudoDependencies.Remove(NodeToDo.Name);
-								break;
-							}
+							RemoveAllPseudodependenciesFromNode(NodeToDo.Name);
+							ListOfRetainedNodesToRemovePseudoDependencies.Remove(NodeToDo.Name);
+							break;
 						}
-						BranchOptions.NodesToRemovePseudoDependencies = ListOfRetainedNodesToRemovePseudoDependencies;
 					}
-					foreach (NodeInfo Dep in NodeToDo.Dependencies)
+					BranchOptions.NodesToRemovePseudoDependencies = ListOfRetainedNodesToRemovePseudoDependencies;
+				}
+				foreach (NodeInfo Dep in NodeToDo.Dependencies)
+				{
+					if (!NodesToDo.Contains(Dep))
 					{
-						if (!NodesToDo.Contains(Dep))
-						{
-							Fringe.Add(Dep);
-							bDoneWithDependencies = false;
-						}
+						Fringe.Add(Dep);
+						bDoneWithDependencies = false;
 					}
 				}
-				NodesToDo.UnionWith(Fringe);
 			}
+			NodesToDo.UnionWith(Fringe);
 		}
+
 		if (TimeIndex != 0)
 		{
 			LogVerbose("Culling based on time index");
@@ -1957,33 +1878,56 @@ public partial class GUBP : BuildCommand
 			}
 			NodesToDo = NewNodesToDo;
 		}
-		//find things that depend on our nodes and setup commander dictionary
-		if (!bOnlyNode)
+		if (bPreflightBuild)
 		{
+			LogVerbose("Culling triggers and downstream for preflight builds ");
+			HashSet<NodeInfo> NewNodesToDo = new HashSet<NodeInfo>();
 			foreach (NodeInfo NodeToDo in NodesToDo)
 			{
-				FullNodeDependedOnBy[NodeToDo.Name] = "";
+				string TriggerDot = NodeToDo.ControllingTriggerDotName;
+				if (TriggerDot == "" && !NodeToDo.Node.TriggerNode())
+				{
+					LogVerbose("  Keeping {0}", NodeToDo.Name);
+					NewNodesToDo.Add(NodeToDo);
+				}
+				else
+				{
+					LogVerbose("  Rejecting {0}", NodeToDo.Name);
+				}
 			}
-			foreach (NodeInfo NodeToDo in NodesToDo)
+			NodesToDo = NewNodesToDo;
+		}
+		return NodesToDo;
+	}
+
+	private static Dictionary<string, string> GetFullNodeDependedOnBy(HashSet<NodeInfo> NodesToDo)
+	{
+		//find things that depend on our nodes and setup commander dictionary
+		Dictionary<string, string> FullNodeDependedOnBy = new Dictionary<string, string>();
+		foreach (NodeInfo NodeToDo in NodesToDo)
+		{
+			FullNodeDependedOnBy[NodeToDo.Name] = "";
+		}
+		foreach (NodeInfo NodeToDo in NodesToDo)
+		{
+			if (!NodeToDo.Node.IsAggregate() && !NodeToDo.Node.IsTest())
 			{
-				if (!NodeToDo.Node.IsAggregate() && !NodeToDo.Node.IsTest())
+				List<NodeInfo> ECDependencies = GetECDependencies(NodeToDo);
+				foreach (NodeInfo Dep in ECDependencies)
 				{
-					List<NodeInfo> ECDependencies = GetECDependencies(NodeToDo);
-					foreach (NodeInfo Dep in ECDependencies)
+					string CurrentValue;
+					if (!FullNodeDependedOnBy.TryGetValue(Dep.Name, out CurrentValue) || CurrentValue.Length == 0)
 					{
-						string CurrentValue;
-						if(!FullNodeDependedOnBy.TryGetValue(Dep.Name, out CurrentValue) || CurrentValue.Length == 0)
-						{
-							FullNodeDependedOnBy[Dep.Name] = NodeToDo.Name;
-						}
-						else
-						{
-							FullNodeDependedOnBy[Dep.Name] += " " + NodeToDo.Name;
-						}
+						FullNodeDependedOnBy[Dep.Name] = NodeToDo.Name;
+					}
+					else
+					{
+						FullNodeDependedOnBy[Dep.Name] += " " + NodeToDo.Name;
 					}
 				}
 			}
 		}
+		return FullNodeDependedOnBy;
 	}
 
 	private void GetFullNodeListAndDirectDependencies(out Dictionary<string, string> FullNodeList, out Dictionary<string, string> FullNodeDirectDependencies)
@@ -2670,7 +2614,7 @@ public partial class GUBP : BuildCommand
 		}
 	}
 
-	void ExecuteNodes(List<NodeInfo> OrdereredToDo, bool bOnlyNode, bool bFake, bool bFakeEC, bool bSaveSharedTempStorage, Dictionary<string, NodeHistory> GUBPNodesHistory, string CLString, string StoreName, string FakeFail)
+	void ExecuteNodes(List<NodeInfo> OrdereredToDo, bool bFake, bool bFakeEC, bool bSaveSharedTempStorage, Dictionary<string, NodeHistory> GUBPNodesHistory, string CLString, string StoreName, string FakeFail)
 	{
         Dictionary<string, NodeInfo> BuildProductToNodeMap = new Dictionary<string, NodeInfo>();
 		foreach (NodeInfo NodeToDo in OrdereredToDo)
@@ -2685,42 +2629,35 @@ public partial class GUBP : BuildCommand
             foreach (NodeInfo Dep in NodeToDo.Dependencies)
             {
                 NodeToDo.Node.AddAllDependent(Dep.Name);
+
                 if (Dep.Node.AllDependencies == null)
                 {
-                    if (!bOnlyNode)
-                    {
-                        throw new AutomationException("Node {0} was not processed yet3?  Processing {1}", Dep, NodeToDo.Name);
-                    }
+					throw new AutomationException("Node {0} was not processed yet?  Processing {1}", Dep, NodeToDo.Name);
                 }
-                else
+
+				foreach (string DepDep in Dep.Node.AllDependencies)
 				{
-					foreach (string DepDep in Dep.Node.AllDependencies)
-					{
-						NodeToDo.Node.AddAllDependent(DepDep);
-					}
+					NodeToDo.Node.AddAllDependent(DepDep);
 				}
 				
                 if (Dep.Node.BuildProducts == null)
                 {
-                    if (!bOnlyNode)
-                    {
-                        throw new AutomationException("Node {0} was not processed yet? Processing {1}", Dep, NodeToDo.Name);
-                    }
+                    throw new AutomationException("Node {0} was not processed yet? Processing {1}", Dep, NodeToDo.Name);
                 }
-                else
+
+				foreach (string Prod in Dep.Node.BuildProducts)
                 {
-                    foreach (string Prod in Dep.Node.BuildProducts)
-                    {
-                        NodeToDo.Node.AddDependentBuildProduct(Prod);
-                    }
-                    if (Dep.Node.AllDependencyBuildProducts == null)
-                    {
-                        throw new AutomationException("Node {0} was not processed yet2?  Processing {1}", Dep.Name, NodeToDo.Name);
-                    }
-                    foreach (string Prod in Dep.Node.AllDependencyBuildProducts)
-                    {
-                        NodeToDo.Node.AddDependentBuildProduct(Prod);
-                    }
+                    NodeToDo.Node.AddDependentBuildProduct(Prod);
+                }
+
+                if (Dep.Node.AllDependencyBuildProducts == null)
+                {
+                    throw new AutomationException("Node {0} was not processed yet2?  Processing {1}", Dep.Name, NodeToDo.Name);
+                }
+
+                foreach (string Prod in Dep.Node.AllDependencyBuildProducts)
+                {
+                    NodeToDo.Node.AddDependentBuildProduct(Prod);
                 }
             }
 
