Change aggregates to be a different class hierarchy to regular build nodes. They don't have many properties in common; they're never actually executed, so they have to be manually filtered out before passing to EC and handled specially when traversing the graph, and they don't suffer the same limitations with being behind a single trigger. Instead of being included in the graph directly, references to them are expanded out to their dependencies when the graph is linked.

diff --git a/Engine/Source/Programs/AutomationTool/GUBP/GUBP.cs b/Engine/Source/Programs/AutomationTool/GUBP/GUBP.cs
--- a/Engine/Source/Programs/AutomationTool/GUBP/GUBP.cs
+++ b/Engine/Source/Programs/AutomationTool/GUBP/GUBP.cs
@@ -8,6 +8,7 @@ using UnrealBuildTool;
 using System.Reflection;
 using System.Xml;
 using System.Linq;
+using System.Diagnostics;
 
 public partial class GUBP : BuildCommand
 {
@@ -25,15 +26,17 @@ public partial class GUBP : BuildCommand
     public string PreflightMangleSuffix = "";
     public GUBPBranchHacker.BranchOptions BranchOptions = null;	
 
-    Dictionary<string, NodeInfo> GUBPNodes;
+	Dictionary<string, NodeInfo> GUBPNodes;
+	Dictionary<string, AggregateInfo> GUBPAggregates;
 
+	[DebuggerDisplay("{Name}")]
 	class NodeInfo
 	{
 		public string Name;
 		public GUBPNode Node;
-		public List<NodeInfo> Dependencies;
-		public List<NodeInfo> PseudoDependencies;
-		public List<NodeInfo> ControllingTriggers;
+		public NodeInfo[] Dependencies;
+		public NodeInfo[] PseudoDependencies;
+		public NodeInfo[] ControllingTriggers;
 		public int FrequencyShift = -1;
 		public bool IsComplete;
 
@@ -49,6 +52,14 @@ public partial class GUBP : BuildCommand
 		}
 	}
 
+	[DebuggerDisplay("{Name}")]
+	class AggregateInfo
+	{
+		public string Name;
+		public AggregateNode Node;
+		public NodeInfo[] Dependencies;
+	}
+
     class NodeHistory
     {
         public int LastSucceeded = 0;
@@ -65,17 +76,28 @@ public partial class GUBP : BuildCommand
     public string AddNode(GUBPNode Node)
     {
         string Name = Node.GetFullName();
-        if (GUBPNodes.ContainsKey(Name))
-        {
+		if (GUBPNodes.ContainsKey(Name) || GUBPAggregates.ContainsKey(Name))
+		{
             throw new AutomationException("Attempt to add a duplicate node {0}", Node.GetFullName());
         }
         GUBPNodes.Add(Name, new NodeInfo{ Name = Name, Node = Node });
         return Name;
     }
 
+	public string AddNode(AggregateNode Node)
+	{
+		string Name = Node.GetFullName();
+		if (GUBPNodes.ContainsKey(Name) || GUBPAggregates.ContainsKey(Name))
+		{
+			throw new AutomationException("Attempt to add a duplicate node {0}", Node.GetFullName());
+		}
+		GUBPAggregates.Add(Name, new AggregateInfo { Name = Name, Node = Node });
+		return Name;
+	}
+
     public bool HasNode(string Node)
     {
-        return GUBPNodes.ContainsKey(Node);
+        return GUBPNodes.ContainsKey(Node) || GUBPAggregates.ContainsKey(Node);
     }
 
     public GUBPNode FindNode(string Node)
@@ -83,6 +105,11 @@ public partial class GUBP : BuildCommand
         return GUBPNodes[Node].Node;
     }
 
+    public AggregateNode FindAggregateNode(string Node)
+    {
+        return GUBPAggregates[Node].Node;
+    }
+
 	public GUBPNode TryFindNode(string Node)
 	{
 		NodeInfo Result;
@@ -178,7 +205,7 @@ public partial class GUBP : BuildCommand
     {
 		if(NodeToDo.ControllingTriggers == null)
 		{
-			NodeToDo.ControllingTriggers = new List<NodeInfo>();
+			NodeToDo.ControllingTriggers = new NodeInfo[0];
 
 			// Find all the dependencies of this node
 			List<NodeInfo> AllDependencies = new List<NodeInfo>();
@@ -196,7 +223,7 @@ public partial class GUBP : BuildCommand
 				{
 					PreviousTrigger = Dependency;
 				}
-				else if(Dependency.ControllingTriggers.Count > 0)
+				else if(Dependency.ControllingTriggers.Length > 0)
 				{
 					PreviousTrigger = Dependency.ControllingTriggers.Last();
 				}
@@ -217,10 +244,12 @@ public partial class GUBP : BuildCommand
 			}
 
 			// Update the list of controlling triggers
-			if(PreviousTriggers.Count == 1)
+			if (PreviousTriggers.Count == 1)
 			{
-				NodeToDo.ControllingTriggers.AddRange(PreviousTriggers[0].ControllingTriggers);
-				NodeToDo.ControllingTriggers.Add(PreviousTriggers[0]);
+				List<NodeInfo> ControllingTriggers = new List<NodeInfo>();
+				ControllingTriggers.AddRange(PreviousTriggers[0].ControllingTriggers);
+				ControllingTriggers.Add(PreviousTriggers[0]);
+				NodeToDo.ControllingTriggers = ControllingTriggers.ToArray();
 			}
 		}
 	}
@@ -482,7 +511,7 @@ public partial class GUBP : BuildCommand
                         string Finished = "";
                         if (UnfinishedTriggers != null)
                         {
-                            if (NodeToDo.ControllingTriggers.Count > 0 && UnfinishedTriggers.Contains(NodeToDo.ControllingTriggers.Last()))
+                            if (NodeToDo.ControllingTriggers.Length > 0 && UnfinishedTriggers.Contains(NodeToDo.ControllingTriggers.Last()))
                             {
                                 Finished = "(not yet triggered)";
                             }
@@ -1447,12 +1476,13 @@ public partial class GUBP : BuildCommand
 		LogVerbose("************************* TimeIndex:           		    {0}", TimeIndex);        
 
 
-        GUBPNodes = new Dictionary<string, NodeInfo>();        
+        GUBPNodes = new Dictionary<string, NodeInfo>();
+		GUBPAggregates = new Dictionary<string, AggregateInfo>();
         Branch = new BranchInfo(HostPlatforms);        
 
 		AddNodesForBranch(TimeIndex, bNoAutomatedTesting);
 
-		LinkGraph(GUBPNodes);
+		LinkGraph(GUBPAggregates, GUBPNodes);
 		FindControllingTriggers(GUBPNodes.Values);
 		FindCompletionState(GUBPNodes.Values, StoreName, LocalOnly);
 
@@ -1538,8 +1568,8 @@ public partial class GUBP : BuildCommand
         {
 			Dictionary<string, string> FullNodeDependedOnBy = GetFullNodeDependedOnBy(NodesToDo);
 
-			List<NodeInfo> SeparatePromotables = FindPromotables(GUBPNodes.Values);
-			Dictionary<NodeInfo, List<NodeInfo>> DependentPromotions = FindDependentPromotables(GUBPNodes.Values, SeparatePromotables);
+			List<AggregateInfo> SeparatePromotables = FindPromotables(GUBPAggregates.Values);
+			Dictionary<NodeInfo, List<AggregateInfo>> DependentPromotions = FindDependentPromotables(GUBPNodes.Values, SeparatePromotables);
 
 			Dictionary<string, string> FullNodeList;
 			Dictionary<string, string> FullNodeDirectDependencies;
@@ -1560,10 +1590,10 @@ public partial class GUBP : BuildCommand
         PrintRunTime();
 	}
 
-	private static List<NodeInfo> FindPromotables(IEnumerable<NodeInfo> NodesToDo)
+	private static List<AggregateInfo> FindPromotables(IEnumerable<AggregateInfo> NodesToDo)
 	{
-		List<NodeInfo> SeparatePromotables = new List<NodeInfo>();
-		foreach (NodeInfo NodeToDo in NodesToDo)
+		List<AggregateInfo> SeparatePromotables = new List<AggregateInfo>();
+		foreach (AggregateInfo NodeToDo in NodesToDo)
 		{
 			if (NodeToDo.Node.IsSeparatePromotable())
 			{
@@ -1573,74 +1603,112 @@ public partial class GUBP : BuildCommand
 		return SeparatePromotables;
 	}
 
-	private static Dictionary<NodeInfo, List<NodeInfo>> FindDependentPromotables(IEnumerable<NodeInfo> NodesToDo, IEnumerable<NodeInfo> SeparatePromotions)
+	private static Dictionary<NodeInfo, List<AggregateInfo>> FindDependentPromotables(IEnumerable<NodeInfo> NodesToDo, IEnumerable<AggregateInfo> SeparatePromotions)
 	{
-		Dictionary<NodeInfo, List<NodeInfo>> DependentPromotions = NodesToDo.ToDictionary(x => x, x => new List<NodeInfo>());
-		foreach (NodeInfo NodeToDo in SeparatePromotions)
+		Dictionary<NodeInfo, List<AggregateInfo>> DependentPromotions = NodesToDo.ToDictionary(x => x, x => new List<AggregateInfo>());
+		foreach (AggregateInfo SeparatePromotion in SeparatePromotions)
 		{
-			List<NodeInfo> Dependencies = GetECDependencies(NodeToDo);
-			foreach (NodeInfo Dep in Dependencies)
+			NodeInfo[] Dependencies = SeparatePromotion.Dependencies.SelectMany(x => GetDependencies(x)).Distinct().ToArray();
+			foreach (NodeInfo Dependency in Dependencies)
 			{
-				if (!Dep.Node.IsPromotableAggregate())
-				{
-					DependentPromotions[Dep].Add(NodeToDo);
-				}
+				DependentPromotions[Dependency].Add(SeparatePromotion);
 			}
 		}
 		return DependentPromotions;
 	}
 
-	private static void LinkGraph(Dictionary<string, NodeInfo> GUBPNodes)
+	private static void LinkGraph(Dictionary<string, AggregateInfo> AggregateNodes, Dictionary<string, NodeInfo> BuildNodes)
 	{
 		int NumErrors = 0;
-		foreach (NodeInfo NodeToDo in GUBPNodes.Values)
+		foreach (AggregateInfo AggregateNode in AggregateNodes.Values)
 		{
-			// Find all the dependencies
-			NodeToDo.Dependencies = new List<NodeInfo>();
-			foreach (string DependencyName in NodeToDo.Node.FullNamesOfDependencies)
+			LinkAggregate(AggregateNode, AggregateNodes, BuildNodes, ref NumErrors);
+		}
+		foreach (NodeInfo BuildNode in BuildNodes.Values)
+		{
+			LinkNode(BuildNode, AggregateNodes, BuildNodes, ref NumErrors);
+		}
+		if(NumErrors > 0)
+		{
+			throw new AutomationException("Failed to link graph ({0} errors).", NumErrors);
+		}
+	}
+
+	private static void LinkAggregate(AggregateInfo Aggregate, Dictionary<string, AggregateInfo> AggregateNodes, Dictionary<string, NodeInfo> BuildNodes, ref int NumErrors)
+	{
+		if (Aggregate.Dependencies == null)
+		{
+			Aggregate.Dependencies = new NodeInfo[0];
+
+			HashSet<NodeInfo> Dependencies = new HashSet<NodeInfo>();
+			foreach (string DependencyName in Aggregate.Node.Dependencies)
 			{
-				NodeInfo Dependency;
-				if (!GUBPNodes.TryGetValue(DependencyName, out Dependency))
-				{
-					CommandUtils.LogError("Node {0} is not in the graph. It is a dependency of {1}.", DependencyName, NodeToDo.Name);
-					NumErrors++;
-				}
-				else if(Dependency == NodeToDo)
+				AggregateInfo AggregateDependency;
+				if(AggregateNodes.TryGetValue(DependencyName, out AggregateDependency))
 				{
-					CommandUtils.LogError("Node {0} has a self arc.", NodeToDo.Name);
-					NumErrors++;
+					LinkAggregate(AggregateDependency, AggregateNodes, BuildNodes, ref NumErrors);
+					Dependencies.UnionWith(AggregateDependency.Dependencies);
+					continue;
 				}
-				else
+
+				NodeInfo Dependency;
+				if(BuildNodes.TryGetValue(DependencyName, out Dependency))
 				{
-					NodeToDo.Dependencies.Add(Dependency);
+					Dependencies.Add(Dependency);
+					continue;
 				}
+
+				CommandUtils.LogError("Node {0} is not in the graph. It is a dependency of {1}.", DependencyName, Aggregate.Name);
+				NumErrors++;
 			}
+			Aggregate.Dependencies = Dependencies.ToArray();
+		}
+	}
 
-			// Find all the pseudo-dependencies
-			NodeToDo.PseudoDependencies = new List<NodeInfo>();
-			foreach (string PseudoDependencyName in NodeToDo.Node.FullNamesOfPseudosependencies)
+	private static void LinkNode(NodeInfo BuildNode, Dictionary<string, AggregateInfo> GUBPAggregates, Dictionary<string, NodeInfo> GUBPNodes, ref int NumErrors)
+	{
+		// Find all the dependencies
+		HashSet<NodeInfo> Dependencies = new HashSet<NodeInfo>();
+		foreach (string DependencyName in BuildNode.Node.FullNamesOfDependencies)
+		{
+			if (!FindDependencies(DependencyName, GUBPAggregates, GUBPNodes, Dependencies))
 			{
-				NodeInfo PseudoDependency;
-				if (!GUBPNodes.TryGetValue(PseudoDependencyName, out PseudoDependency))
-				{
-					CommandUtils.LogError("Node {0} is not in the graph. It is a pseudodependency of {1}.", PseudoDependencyName, NodeToDo.Name);
-					NumErrors++;
-				}
-				else if(PseudoDependency == NodeToDo)
-				{
-					CommandUtils.LogError("Node {0} has a self pseudoarc.", NodeToDo.Name);
-					NumErrors++;
-				}
-				else
-				{
-					NodeToDo.PseudoDependencies.Add(PseudoDependency);
-				}
+				CommandUtils.LogError("Node {0} is not in the graph. It is a dependency of {1}.", DependencyName, BuildNode.Name);
+				NumErrors++;
 			}
 		}
-		if (NumErrors > 0)
+		BuildNode.Dependencies = Dependencies.ToArray();
+
+		// Find all the pseudo-dependencies
+		HashSet<NodeInfo> PseudoDependencies = new HashSet<NodeInfo>();
+		foreach (string PseudoDependencyName in BuildNode.Node.FullNamesOfPseudosependencies)
 		{
-			throw new AutomationException("Failed to link graph ({0} errors).", NumErrors);
+			if (!FindDependencies(PseudoDependencyName, GUBPAggregates, GUBPNodes, PseudoDependencies))
+			{
+				CommandUtils.LogError("Node {0} is not in the graph. It is a pseudodependency of {1}.", PseudoDependencyName, BuildNode.Name);
+				NumErrors++;
+			}
+		}
+		BuildNode.PseudoDependencies = PseudoDependencies.ToArray();
+	}
+
+	private static bool FindDependencies(string Name, Dictionary<string, AggregateInfo> NameToAggregate, Dictionary<string, NodeInfo> NameToNode, HashSet<NodeInfo> Dependencies)
+	{
+		AggregateInfo AggregateDependency;
+		if (NameToAggregate.TryGetValue(Name, out AggregateDependency))
+		{
+			Dependencies.UnionWith(AggregateDependency.Dependencies);
+			return true;
 		}
+
+		NodeInfo NodeDependency;
+		if (NameToNode.TryGetValue(Name, out NodeDependency))
+		{
+			Dependencies.Add(NodeDependency);
+			return true;
+		}
+
+		return false;
 	}
 
 	private static Dictionary<NodeInfo, int> ApplyFrequencyBarriers(Dictionary<string, NodeInfo> GUBPNodes, Dictionary<string, sbyte> FrequencyBarriers)
@@ -1722,6 +1790,20 @@ public partial class GUBP : BuildCommand
 							bFoundAnything = true;
 						}
 					}
+					foreach (KeyValuePair<string, AggregateInfo> AggregateNode in GUBPAggregates)
+					{
+						if(AggregateNode.Key.Equals(NodeArg, StringComparison.InvariantCultureIgnoreCase))
+						{
+							foreach(NodeInfo Node in AggregateNode.Value.Dependencies)
+							{
+								if(!NodesToDo.Contains(Node))
+								{
+									NodesToDo.Add(Node);
+								}
+								bFoundAnything = true;
+							}
+						}
+					}
 					if (!bFoundAnything)
 					{
 						throw new AutomationException("Could not find node named {0}", NodeName);
@@ -1745,7 +1827,10 @@ public partial class GUBP : BuildCommand
 						{
 							if (GameProj.Options(UnrealTargetPlatform.Win64).bIsPromotable)
 							{
-								NodesToDo.Add(GUBPNodes[GameAggregatePromotableNode.StaticGetFullName(GameProj)]);
+								foreach (NodeInfo Dependency in GUBPAggregates[GameAggregatePromotableNode.StaticGetFullName(GameProj)].Dependencies)
+								{
+									NodesToDo.Add(Dependency);
+								}
 							}
 							foreach (KeyValuePair<string, NodeInfo> Node in GUBPNodes)
 							{
@@ -2140,7 +2225,7 @@ public partial class GUBP : BuildCommand
 		}
 	}
 
-	private void DoCommanderSetup(int TimeIndex, bool bSkipTriggers, bool bFake, bool bFakeEC, string CLString, NodeInfo ExplicitTrigger, Dictionary<string, string> FullNodeList, Dictionary<string, string> FullNodeDirectDependencies, Dictionary<string, string> FullNodeDependedOnBy, Dictionary<NodeInfo, List<NodeInfo>> DependentPromotions, List<NodeInfo> SeparatePromotables, Dictionary<string, int> FullNodeListSortKey, Dictionary<string, NodeHistory> GUBPNodesHistory, List<NodeInfo> OrdereredToDo, List<NodeInfo> UnfinishedTriggers, string FakeFail)
+	private void DoCommanderSetup(int TimeIndex, bool bSkipTriggers, bool bFake, bool bFakeEC, string CLString, NodeInfo ExplicitTrigger, Dictionary<string, string> FullNodeList, Dictionary<string, string> FullNodeDirectDependencies, Dictionary<string, string> FullNodeDependedOnBy, Dictionary<NodeInfo, List<AggregateInfo>> DependentPromotions, List<AggregateInfo> SeparatePromotables, Dictionary<string, int> FullNodeListSortKey, Dictionary<string, NodeHistory> GUBPNodesHistory, List<NodeInfo> OrdereredToDo, List<NodeInfo> UnfinishedTriggers, string FakeFail)
 	{
 
 		if (OrdereredToDo.Count == 0)
@@ -2165,11 +2250,11 @@ public partial class GUBP : BuildCommand
 		{
 			ECProps.Add(string.Format("DependedOnBy/{0}={1}", NodePair.Key, NodePair.Value));
 		}
-		foreach (KeyValuePair<NodeInfo, List<NodeInfo>> NodePair in DependentPromotions)
+		foreach (KeyValuePair<NodeInfo, List<AggregateInfo>> NodePair in DependentPromotions)
 		{
 			ECProps.Add(string.Format("DependentPromotions/{0}={1}", NodePair.Key.Name, String.Join(" ", NodePair.Value.Select(x => x.Name))));
 		}
-		foreach (NodeInfo Node in SeparatePromotables)
+		foreach (AggregateInfo Node in SeparatePromotables)
 		{
 			ECProps.Add(string.Format("PossiblePromotables/{0}={1}", Node.Name, ""));
 		}
@@ -2189,7 +2274,7 @@ public partial class GUBP : BuildCommand
 			// remove nodes that have unfinished triggers
 			foreach (NodeInfo NodeToDo in OrdereredToDo)
 			{
-				NodeInfo ControllingTrigger = (NodeToDo.ControllingTriggers.Count > 0)? NodeToDo.ControllingTriggers.Last() : null;
+				NodeInfo ControllingTrigger = (NodeToDo.ControllingTriggers.Length > 0)? NodeToDo.ControllingTriggers.Last() : null;
 				bool bNoUnfinishedTriggers = !UnfinishedTriggers.Contains(ControllingTrigger);
 
 				if (bNoUnfinishedTriggers)
diff --git a/Engine/Source/Programs/AutomationTool/GUBP/LegacyBranchSetup.cs b/Engine/Source/Programs/AutomationTool/GUBP/LegacyBranchSetup.cs
--- a/Engine/Source/Programs/AutomationTool/GUBP/LegacyBranchSetup.cs
+++ b/Engine/Source/Programs/AutomationTool/GUBP/LegacyBranchSetup.cs
@@ -506,7 +506,7 @@ partial class GUBP
                     }
                     if (EditorTestNodes.Count > 0)
                     {
-                        AddNode(new GameAggregateNode(this, HostPlatform, Branch.BaseEngineProject, "AllEditorTests", EditorTestNodes, 0.0f));
+                        AddNode(new GameAggregateNode(this, HostPlatform, Branch.BaseEngineProject, "AllEditorTests", EditorTestNodes));
                     }
                 }
 
@@ -547,7 +547,7 @@ partial class GUBP
                                 {
                                     GamePlatforms.Add(Plat);
                                 }
-                                if (!GUBPNodes.ContainsKey(GamePlatformMonolithicsNode.StaticGetFullName(HostPlatform, Branch.BaseEngineProject, Plat)))
+                                if (!HasNode(GamePlatformMonolithicsNode.StaticGetFullName(HostPlatform, Branch.BaseEngineProject, Plat)))
                                 {
 									if(GamePlatformMonolithicsNode.HasPrecompiledTargets(Branch.BaseEngineProject, HostPlatform, Plat))
 									{
@@ -557,7 +557,7 @@ partial class GUBP
                                 }
 								if (Plat == UnrealTargetPlatform.Win32 && Target.Rules.GUBP_BuildWindowsXPMonolithics() && Kind == TargetRules.TargetType.Game)
 								{
-									if (!GUBPNodes.ContainsKey(GamePlatformMonolithicsNode.StaticGetFullName(HostPlatform, Branch.BaseEngineProject, Plat, true)))
+									if (!HasNode(GamePlatformMonolithicsNode.StaticGetFullName(HostPlatform, Branch.BaseEngineProject, Plat, true)))
 									{
 										AddNode(new GamePlatformMonolithicsNode(this, HostPlatform, ActivePlatforms, Branch.BaseEngineProject, Plat, true));
 									}
@@ -595,11 +595,11 @@ partial class GUBP
                             if (ActivePlatforms.Contains(Plat))
                             {
                                 string CookedPlatform = Platform.Platforms[Plat].GetCookPlatform(Kind == TargetRules.TargetType.Server, Kind == TargetRules.TargetType.Client, "");								
-                                if (!GUBPNodes.ContainsKey(CookNode.StaticGetFullName(HostPlatform, Branch.BaseEngineProject, CookedPlatform)))
+                                if (!HasNode(CookNode.StaticGetFullName(HostPlatform, Branch.BaseEngineProject, CookedPlatform)))
                                 {
                                     GameCookNodes.Add(AddNode(new CookNode(this, HostPlatform, Branch.BaseEngineProject, Plat, CookedPlatform)));
                                 }
-                                if (!GUBPNodes.ContainsKey(GamePlatformCookedAndCompiledNode.StaticGetFullName(HostPlatform, Branch.BaseEngineProject, Plat)))
+                                if (!HasNode(GamePlatformCookedAndCompiledNode.StaticGetFullName(HostPlatform, Branch.BaseEngineProject, Plat)))
                                 {
                                     AddNode(new GamePlatformCookedAndCompiledNode(this, HostPlatform, Branch.BaseEngineProject, Plat, false));
                                 }
@@ -616,7 +616,7 @@ partial class GUBP
                                     }
                                     if (ThisMonoGameTestNodes.Count > 0)
                                     {										
-                                        GameTestNodes.Add(AddNode(new GameAggregateNode(this, HostPlatform, Branch.BaseEngineProject, "CookedTests_" + Plat.ToString() + "_" + Kind.ToString() + HostPlatformNode.StaticGetHostPlatformSuffix(HostPlatform), ThisMonoGameTestNodes, 0.0f)));
+                                        GameTestNodes.Add(AddNode(new GameAggregateNode(this, HostPlatform, Branch.BaseEngineProject, "CookedTests_" + Plat.ToString() + "_" + Kind.ToString() + HostPlatformNode.StaticGetHostPlatformSuffix(HostPlatform), ThisMonoGameTestNodes)));
                                     }
                                 }
 
@@ -628,11 +628,11 @@ partial class GUBP
                                     }
                                     if (BranchOptions.ProjectsToCook.Contains(NonCodeProject.GameName) || BranchOptions.ProjectsToCook.Count == 0)
                                     {
-                                    if (!GUBPNodes.ContainsKey(CookNode.StaticGetFullName(HostPlatform, NonCodeProject, CookedPlatform)))
+                                    if (!HasNode(CookNode.StaticGetFullName(HostPlatform, NonCodeProject, CookedPlatform)))
                                     {
                                         GameCookNodes.Add(AddNode(new CookNode(this, HostPlatform, NonCodeProject, Plat, CookedPlatform)));
                                     }
-                                    if (!GUBPNodes.ContainsKey(GamePlatformCookedAndCompiledNode.StaticGetFullName(HostPlatform, NonCodeProject, Plat)))
+                                    if (!HasNode(GamePlatformCookedAndCompiledNode.StaticGetFullName(HostPlatform, NonCodeProject, Plat)))
                                     {
                                         AddNode(new GamePlatformCookedAndCompiledNode(this, HostPlatform, NonCodeProject, Plat, false));
 
@@ -657,16 +657,15 @@ partial class GUBP
 														{
 															BuildAgentSharingGroup = "";
 														}
-														GUBPNodes[CookNode.StaticGetFullName(HostPlatform, NonCodeProject, CookedPlatform)].Node.AgentSharingGroup = BuildAgentSharingGroup;
-														GUBPNodes[GamePlatformCookedAndCompiledNode.StaticGetFullName(HostPlatform, NonCodeProject, Plat)].Node.AgentSharingGroup = BuildAgentSharingGroup;
-														GUBPNodes[NodeName].Node.AgentSharingGroup = BuildAgentSharingGroup;
+														FindNode(CookNode.StaticGetFullName(HostPlatform, NonCodeProject, CookedPlatform)).AgentSharingGroup = BuildAgentSharingGroup;
+														FindNode(NodeName).AgentSharingGroup = BuildAgentSharingGroup;
 													}
 													else
 													{
 														//GUBPNodes[NodeName].AgentSharingGroup = FormalAgentSharingGroup;
 															if (Plat == UnrealTargetPlatform.XboxOne)
 														{
-															GUBPNodes[NodeName].Node.AgentSharingGroup = "";
+															FindNode(NodeName).AgentSharingGroup = "";
 														}
 													}
 													if (PlatPair.bTest)
@@ -689,7 +688,7 @@ partial class GUBP
 										}
 										if (ThisMonoGameTestNodes.Count > 0)
                                             {
-											GameTestNodes.Add(AddNode(new GameAggregateNode(this, HostPlatform, NonCodeProject, "CookedTests_" + Plat.ToString() + "_" + Kind.ToString() + HostPlatformNode.StaticGetHostPlatformSuffix(HostPlatform), ThisMonoGameTestNodes, 0.0f)));
+											GameTestNodes.Add(AddNode(new GameAggregateNode(this, HostPlatform, NonCodeProject, "CookedTests_" + Plat.ToString() + "_" + Kind.ToString() + HostPlatformNode.StaticGetHostPlatformSuffix(HostPlatform), ThisMonoGameTestNodes)));
 										}
 									}
 								}
@@ -788,7 +787,7 @@ partial class GUBP
 						}
 						if (EditorTestNodes.Count > 0)
 						{
-							AddNode(new GameAggregateNode(this, HostPlatform, CodeProj, "AllEditorTests", EditorTestNodes, 0.0f));
+							AddNode(new GameAggregateNode(this, HostPlatform, CodeProj, "AllEditorTests", EditorTestNodes));
 						}
 					}
                 }
@@ -818,7 +817,7 @@ partial class GUBP
                             }
 							if(Plat == UnrealTargetPlatform.Win32 && Target.Rules.GUBP_BuildWindowsXPMonolithics())
 							{
-								if(!GUBPNodes.ContainsKey(GamePlatformMonolithicsNode.StaticGetFullName(HostPlatform, CodeProj, Plat, true)))
+								if(!HasNode(GamePlatformMonolithicsNode.StaticGetFullName(HostPlatform, CodeProj, Plat, true)))
 								{
 									AddNode(new GamePlatformMonolithicsNode(this, HostPlatform, ActivePlatforms, CodeProj, Plat, true));
 								}
@@ -833,7 +832,7 @@ partial class GUBP
                                 {
                                     GamePlatforms.Add(Plat);
                                 }
-                                if (!GUBPNodes.ContainsKey(GamePlatformMonolithicsNode.StaticGetFullName(HostPlatform, CodeProj, Plat)))
+                                if (!HasNode(GamePlatformMonolithicsNode.StaticGetFullName(HostPlatform, CodeProj, Plat)))
                                 {
 									if(GamePlatformMonolithicsNode.HasPrecompiledTargets(CodeProj, HostPlatform, Plat))
 									{
@@ -849,11 +848,11 @@ partial class GUBP
 									{
 										CookedPlatform = Target.Rules.GUBP_AlternateCookPlatform(HostPlatform, CookedPlatform);
 									}
-									if (!GUBPNodes.ContainsKey(CookNode.StaticGetFullName(HostPlatform, CodeProj, CookedPlatform)))
+									if (!HasNode(CookNode.StaticGetFullName(HostPlatform, CodeProj, CookedPlatform)))
 									{
 										AddNode(new CookNode(this, HostPlatform, CodeProj, Plat, CookedPlatform));
 									}
-									if (!GUBPNodes.ContainsKey(GamePlatformCookedAndCompiledNode.StaticGetFullName(HostPlatform, CodeProj, Plat)))
+									if (!HasNode(GamePlatformCookedAndCompiledNode.StaticGetFullName(HostPlatform, CodeProj, Plat)))
 									{
 										AddNode(new GamePlatformCookedAndCompiledNode(this, HostPlatform, CodeProj, Plat, true));
 									}
@@ -899,16 +898,15 @@ partial class GUBP
 												{
 													BuildAgentSharingGroup = "";
 												}
-												GUBPNodes[CookNode.StaticGetFullName(HostPlatform, CodeProj, CookedPlatform)].Node.AgentSharingGroup = BuildAgentSharingGroup;
-												GUBPNodes[GamePlatformCookedAndCompiledNode.StaticGetFullName(HostPlatform, CodeProj, Plat)].Node.AgentSharingGroup = BuildAgentSharingGroup;
-												GUBPNodes[FormalNodeName].Node.AgentSharingGroup = BuildAgentSharingGroup;
+												FindNode(CookNode.StaticGetFullName(HostPlatform, CodeProj, CookedPlatform)).AgentSharingGroup = BuildAgentSharingGroup;
+												FindNode(FormalNodeName).AgentSharingGroup = BuildAgentSharingGroup;
 											}
 											else
 											{
 												//GUBPNodes[FormalNodeName].AgentSharingGroup = FormalAgentSharingGroup;
 												if (Plat == UnrealTargetPlatform.XboxOne)
 												{
-													GUBPNodes[FormalNodeName].Node.AgentSharingGroup = "";
+													FindNode(FormalNodeName).AgentSharingGroup = "";
 												}
 											}
 											if (Config.bTest)
@@ -935,7 +933,7 @@ partial class GUBP
 										}
 										if (ThisMonoGameTestNodes.Count > 0)
 										{
-											GameTestNodes.Add(AddNode(new GameAggregateNode(this, HostPlatform, CodeProj, "CookedTests_" + Plat.ToString() + "_" + Kind.ToString() + HostPlatformNode.StaticGetHostPlatformSuffix(HostPlatform), ThisMonoGameTestNodes, 0.0f)));
+											GameTestNodes.Add(AddNode(new GameAggregateNode(this, HostPlatform, CodeProj, "CookedTests_" + Plat.ToString() + "_" + Kind.ToString() + HostPlatformNode.StaticGetHostPlatformSuffix(HostPlatform), ThisMonoGameTestNodes)));
 										}
 									}
 								}
diff --git a/Engine/Source/Programs/AutomationTool/GUBP/LegacyNodes.cs b/Engine/Source/Programs/AutomationTool/GUBP/LegacyNodes.cs
--- a/Engine/Source/Programs/AutomationTool/GUBP/LegacyNodes.cs
+++ b/Engine/Source/Programs/AutomationTool/GUBP/LegacyNodes.cs
@@ -70,14 +70,6 @@ partial class GUBP
         {
             return false;
         }
-		public virtual bool IsPromotableAggregate()
-		{
-			return false;
-		}
-		public virtual bool IsSeparatePromotable()
-		{
-			return false;
-		}
 		public virtual string NodeHostPlatform()
 		{
 			return "";
@@ -1575,30 +1567,41 @@ partial class GUBP
         }
     }
 
-    public class AggregateNode : GUBPNode
+    public abstract class AggregateNode
     {
+		public List<string> Dependencies = new List<string>();
+
         public AggregateNode()
         {
         }
-        public override bool RunInEC()
-        {
-            return false;
-        }		
-        public override bool IsAggregate()
-        {
-            return true;
-        }
-        public override void DoBuild(GUBP bp)
-        {
-            BuildProducts = new List<string>();
-        }
-        public override void DoFakeBuild(GUBP bp) // this is used to more rapidly test a build system, it does nothing but save a record of success as a build product
-        {
-            BuildProducts = new List<string>();
-        }
+
+		public void AddDependency(string Name)
+		{
+			if(!Dependencies.Contains(Name))
+			{
+				Dependencies.Add(Name);
+			}
+		}
+
+		public abstract string GetFullName();
+
+		public virtual string GameNameIfAnyForTempStorage()
+		{
+			return "";
+		}
+
+		public virtual bool IsPromotableAggregate()
+		{
+			return false;
+		}
+
+		public virtual bool IsSeparatePromotable()
+		{
+			return false;
+		}
     }
 
-    public class HostPlatformAggregateNode : AggregateNode
+    public abstract class HostPlatformAggregateNode : AggregateNode
     {
         protected UnrealTargetPlatform HostPlatform;
         public HostPlatformAggregateNode(UnrealTargetPlatform InHostPlatform)
@@ -2166,10 +2169,6 @@ partial class GUBP
         {
             return StaticGetFullName();
         }
-        public override int CISFrequencyQuantumShift(GUBP bp)
-        {
-            return base.CISFrequencyQuantumShift(bp) + 5;
-        }
     }
 
     public class CookNode : HostPlatformNode
@@ -2351,21 +2350,6 @@ partial class GUBP
                     }
                 }
             }
-
-            // put these in the right agent group, even though they aren't exposed to EC to sort right.
-            if (InGameProj.GameName != bp.Branch.BaseEngineProject.GameName && GameProj.Properties.Targets.ContainsKey(TargetRules.TargetType.Editor))
-            {
-                var Options = InGameProj.Options(HostPlatform);
-                if ((Options.bIsPromotable || Options.bTestWithShared) && !Options.bSeparateGamePromotion)
-                {
-                    AgentSharingGroup = "SharedCooks" + StaticGetHostPlatformSuffix(HostPlatform);
-                }
-            }
-            else
-            {
-                AgentSharingGroup = "SharedCooks" + StaticGetHostPlatformSuffix(HostPlatform);
-            }
-
         }
 
         public static string StaticGetFullName(UnrealTargetPlatform InHostPlatform, BranchInfo.BranchUProject InGameProj, UnrealTargetPlatform InTargetPlatform)
@@ -3209,24 +3193,18 @@ partial class GUBP
     {
         BranchInfo.BranchUProject GameProj;
         string AggregateName;
-        float ECPriority;
 
-        public GameAggregateNode(GUBP bp, UnrealTargetPlatform InHostPlatform, BranchInfo.BranchUProject InGameProj, string InAggregateName, List<string> Dependencies, float InECPriority = 100.0f)
+        public GameAggregateNode(GUBP bp, UnrealTargetPlatform InHostPlatform, BranchInfo.BranchUProject InGameProj, string InAggregateName, List<string> Dependencies)
             : base(InHostPlatform)
         {
             GameProj = InGameProj;
             AggregateName = InAggregateName;
-            ECPriority = InECPriority;
 
             foreach (var Dep in Dependencies)
             {
                 AddDependency(Dep);
             }
         }
-        public override float Priority()
-        {
-            return ECPriority;
-        }
 
         public static string StaticGetFullName(UnrealTargetPlatform InHostPlatform, BranchInfo.BranchUProject InGameProj, string InAggregateName)
         {
diff --git a/Engine/Source/Programs/AutomationTool/Scripts/RocketBuild.Automation.cs b/Engine/Source/Programs/AutomationTool/Scripts/RocketBuild.Automation.cs
--- a/Engine/Source/Programs/AutomationTool/Scripts/RocketBuild.Automation.cs
+++ b/Engine/Source/Programs/AutomationTool/Scripts/RocketBuild.Automation.cs
@@ -114,7 +114,7 @@ namespace Rocket
 				bp.AddNode(new GatherRocketNode(HostPlatform, TargetPlatforms, LocalOutputDir));
 
 				// Add the aggregate node for the entire install
-				GUBP.GUBPNode PromotableNode = bp.FindNode(GUBP.SharedAggregatePromotableNode.StaticGetFullName());
+				GUBP.SharedAggregatePromotableNode PromotableNode = (GUBP.SharedAggregatePromotableNode)bp.FindAggregateNode(GUBP.SharedAggregatePromotableNode.StaticGetFullName());
 				PromotableNode.AddDependency(FilterRocketNode.StaticGetFullName(HostPlatform));
 				PromotableNode.AddDependency(BuildDerivedDataCacheNode.StaticGetFullName(HostPlatform));
 
@@ -145,7 +145,7 @@ namespace Rocket
 				bp.AddNode(new PublishRocketSymbolsNode(bp, HostPlatform, TargetPlatforms, PublishedEngineDir + "Symbols"));
 
 				// Add a dependency on this being published as part of the shared promotable being labeled
-				GUBP.SharedLabelPromotableSuccessNode LabelPromotableNode = (GUBP.SharedLabelPromotableSuccessNode)bp.FindNode(GUBP.SharedLabelPromotableSuccessNode.StaticGetFullName());
+				GUBP.SharedLabelPromotableSuccessNode LabelPromotableNode = (GUBP.SharedLabelPromotableSuccessNode)bp.FindAggregateNode(GUBP.SharedLabelPromotableSuccessNode.StaticGetFullName());
 				LabelPromotableNode.AddDependency(PublishRocketNode.StaticGetFullName(HostPlatform));
 				LabelPromotableNode.AddDependency(PublishRocketSymbolsNode.StaticGetFullName(HostPlatform));
 
