Use aggregates to specify nodes to build for each game, rather than searching the graph manually.

diff --git a/Engine/Source/Programs/AutomationTool/GUBP/GUBP.cs b/Engine/Source/Programs/AutomationTool/GUBP/GUBP.cs
--- a/Engine/Source/Programs/AutomationTool/GUBP/GUBP.cs
+++ b/Engine/Source/Programs/AutomationTool/GUBP/GUBP.cs
@@ -475,7 +475,8 @@ public partial class GUBP : BuildCommand
         Log("Took {0}s to get P4 history", BuildDuration / 1000);
 
     }
-    void PrintNodes(GUBP bp, List<NodeInfo> Nodes, Dictionary<string, NodeHistory> GUBPNodesHistory, List<NodeInfo> UnfinishedTriggers = null)
+
+    void PrintNodes(GUBP bp, List<NodeInfo> Nodes, IEnumerable<AggregateInfo> Aggregates, Dictionary<string, NodeHistory> GUBPNodesHistory, List<NodeInfo> UnfinishedTriggers = null)
     {
         bool bShowAllChanges = bp.ParseParam("AllChanges") && GUBPNodesHistory != null;
         bool bShowChanges = (bp.ParseParam("Changes") && GUBPNodesHistory != null) || bShowAllChanges;
@@ -489,6 +490,8 @@ public partial class GUBP : BuildCommand
         bool bShowTriggers = true;
         string LastControllingTrigger = "";
         string LastAgentGroup = "";
+
+        Log("*********** Desired And Dependent Nodes, in order.");
         foreach (NodeInfo NodeToDo in Nodes)
         {
             if (ECOnly && !NodeToDo.Node.RunInEC())
@@ -593,7 +596,23 @@ public partial class GUBP : BuildCommand
                 }
             }
         }
+
+		AggregateInfo[] MatchingAggregates = Aggregates.Where(x => x.Dependencies.All(y => Nodes.Contains(y))).ToArray();
+		if(MatchingAggregates.Length > 0)
+		{
+			Log("*********** Aggregates");
+			foreach(AggregateInfo Aggregate in MatchingAggregates.OrderBy(x => x.Name))
+			{
+				StringBuilder Note = new StringBuilder("    " + Aggregate.Name);
+				if (Aggregate.Node.IsPromotableAggregate())
+				{
+					Note.Append(" (promotable)");
+				}
+				Log(Note.ToString());
+			}
+		}
     }
+
     void SaveGraphVisualization(List<NodeInfo> Nodes)
     {
         List<GraphNode> GraphNodes = new List<GraphNode>();
@@ -1558,8 +1577,8 @@ public partial class GUBP : BuildCommand
 
 		List<NodeInfo> UnfinishedTriggers = FindUnfinishedTriggers(bSkipTriggers, ExplicitTrigger, OrderedToDo);
 
-        LogVerbose("*********** Desired And Dependent Nodes, in order.");
-        PrintNodes(this, OrderedToDo, GUBPNodesHistory, UnfinishedTriggers);		
+		PrintNodes(this, OrderedToDo, GUBPAggregates.Values, GUBPNodesHistory, UnfinishedTriggers);
+
         //check sorting
 		CheckSortOrder(OrderedToDo);
 
@@ -1765,107 +1784,44 @@ public partial class GUBP : BuildCommand
 
 	private HashSet<NodeInfo> ParseNodesToDo(bool WithoutLinux, int TimeIndex, bool CommanderSetup)
 	{
-		HashSet<NodeInfo> NodesToDo = new HashSet<NodeInfo>();
+		List<string> NamesToDo = new List<string>();
+
+		string NodeParam = ParseParamValue("Node", null);
+		if (NodeParam != null)
+		{
+			NamesToDo.AddRange(NodeParam.Split('+'));
+		}
 
-		string NodeSpec = ParseParamValue("Node");
-		if (!String.IsNullOrEmpty(NodeSpec))
+		string GameParam = ParseParamValue("Game", null);
+		if (GameParam != null)
 		{
-			List<string> Nodes = new List<string>(NodeSpec.Split('+'));
-			foreach (string NodeArg in Nodes)
+			NamesToDo.AddRange(GameParam.Split('+').Select(x => FullGameAggregateNode.StaticGetFullName(x)));
+		}
+
+		HashSet<NodeInfo> NodesToDo = new HashSet<NodeInfo>();
+		foreach (string NameToDo in NamesToDo)
+		{
+			int FoundNames = 0;
+			foreach (NodeInfo PotentialNode in GUBPNodes.Values)
 			{
-				string NodeName = NodeArg.Trim();
-				bool bFoundAnything = false;
-				if (!String.IsNullOrEmpty(NodeName))
+				if (String.Compare(PotentialNode.Name, NameToDo, StringComparison.InvariantCultureIgnoreCase) == 0 || String.Compare(PotentialNode.Node.AgentSharingGroup, NameToDo, StringComparison.InvariantCultureIgnoreCase) == 0)
 				{
-					foreach (KeyValuePair<string, NodeInfo> Node in GUBPNodes)
-					{
-						if (Node.Value.Node.GetFullName().Equals(NodeArg, StringComparison.InvariantCultureIgnoreCase) ||
-							Node.Value.Node.AgentSharingGroup.Equals(NodeArg, StringComparison.InvariantCultureIgnoreCase)
-							)
-						{
-							if (!NodesToDo.Contains(Node.Value))
-							{
-								NodesToDo.Add(Node.Value);
-							}
-							bFoundAnything = true;
-						}
-					}
-					foreach (KeyValuePair<string, AggregateInfo> AggregateNode in GUBPAggregates)
-					{
-						if(AggregateNode.Key.Equals(NodeArg, StringComparison.InvariantCultureIgnoreCase))
-						{
-							foreach(NodeInfo Node in AggregateNode.Value.Dependencies)
-							{
-								if(!NodesToDo.Contains(Node))
-								{
-									NodesToDo.Add(Node);
-								}
-								bFoundAnything = true;
-							}
-						}
-					}
-					if (!bFoundAnything)
-					{
-						throw new AutomationException("Could not find node named {0}", NodeName);
-					}
+					NodesToDo.Add(PotentialNode);
+					FoundNames++;
 				}
 			}
-		}
-
-		string GameSpec = ParseParamValue("Game");
-		if (!String.IsNullOrEmpty(GameSpec))
-		{
-			List<string> Games = new List<string>(GameSpec.Split('+'));
-			foreach (string GameArg in Games)
+			foreach (AggregateInfo PotentialAggregate in GUBPAggregates.Values)
 			{
-				string GameName = GameArg.Trim();
-				if (!String.IsNullOrEmpty(GameName))
+				if (String.Compare(PotentialAggregate.Name, NameToDo, StringComparison.InvariantCultureIgnoreCase) == 0)
 				{
-					foreach (BranchInfo.BranchUProject GameProj in Branch.CodeProjects)
-					{
-						if (GameProj.GameName.Equals(GameName, StringComparison.InvariantCultureIgnoreCase))
-						{
-							if (GameProj.Options(UnrealTargetPlatform.Win64).bIsPromotable)
-							{
-								foreach (NodeInfo Dependency in GUBPAggregates[GameAggregatePromotableNode.StaticGetFullName(GameProj)].Dependencies)
-								{
-									NodesToDo.Add(Dependency);
-								}
-							}
-							foreach (KeyValuePair<string, NodeInfo> Node in GUBPNodes)
-							{
-								if (Node.Value.Node.GameNameIfAnyForTempStorage() == GameProj.GameName)
-								{
-									NodesToDo.Add(Node.Value);
-								}
-							}
-
-							GameName = null;
-						}
-					}
-					if (GameName != null)
-					{
-						foreach (BranchInfo.BranchUProject GameProj in Branch.NonCodeProjects)
-						{
-							if (GameProj.GameName.Equals(GameName, StringComparison.InvariantCultureIgnoreCase))
-							{
-								foreach (KeyValuePair<string, NodeInfo> Node in GUBPNodes)
-								{
-									if (Node.Value.Node.GameNameIfAnyForTempStorage() == GameProj.GameName)
-									{
-										NodesToDo.Add(Node.Value);
-									}
-								}
-								GameName = null;
-							}
-						}
-					}
-					if (GameName != null)
-					{
-						throw new AutomationException("Could not find game named {0}", GameName);
-					}
+					NodesToDo.UnionWith(PotentialAggregate.Dependencies);
+					FoundNames++;
 				}
 			}
+			if (FoundNames == 0)
+			{
+				throw new AutomationException("Could not find node named {0}", FoundNames);
+			}
 		}
 
 		if (NodesToDo.Count == 0)
@@ -2300,7 +2256,7 @@ public partial class GUBP : BuildCommand
 		using(TelemetryStopwatch PrintNodesTimer = new TelemetryStopwatch("SetupCommanderPrint"))
 		{
 			Log("*********** EC Nodes, in order.");
-			PrintNodes(this, OrdereredToDo, GUBPNodesHistory, UnfinishedTriggers);
+			PrintNodes(this, OrdereredToDo, GUBPAggregates.Values, GUBPNodesHistory, UnfinishedTriggers);
 		}
 		// here we are just making sure everything before the explicit trigger is completed.
 		if (ExplicitTrigger != null)
diff --git a/Engine/Source/Programs/AutomationTool/GUBP/LegacyBranchSetup.cs b/Engine/Source/Programs/AutomationTool/GUBP/LegacyBranchSetup.cs
--- a/Engine/Source/Programs/AutomationTool/GUBP/LegacyBranchSetup.cs
+++ b/Engine/Source/Programs/AutomationTool/GUBP/LegacyBranchSetup.cs
@@ -1055,5 +1055,26 @@ partial class GUBP
             AddNode(PreflightSuccessNode);
         }
 #endif
+
+		// Add aggregate nodes for every project in the branch
+		foreach (BranchInfo.BranchUProject GameProj in Branch.AllProjects)
+		{
+			List<string> NodeNames = new List<string>();
+			if (GameProj.Options(UnrealTargetPlatform.Win64).bIsPromotable)
+			{
+				NodeNames.Add(GameAggregatePromotableNode.StaticGetFullName(GameProj));
+			}
+			foreach (NodeInfo Node in GUBPNodes.Values)
+			{
+				if (Node.Node.GameNameIfAnyForTempStorage() == GameProj.GameName)
+				{
+					NodeNames.Add(Node.Name);
+				}
+			}
+			if (NodeNames.Count > 0)
+			{
+				AddNode(new FullGameAggregateNode(GameProj.GameName, NodeNames));
+			}
+		}
 	}
 }
diff --git a/Engine/Source/Programs/AutomationTool/GUBP/LegacyNodes.cs b/Engine/Source/Programs/AutomationTool/GUBP/LegacyNodes.cs
--- a/Engine/Source/Programs/AutomationTool/GUBP/LegacyNodes.cs
+++ b/Engine/Source/Programs/AutomationTool/GUBP/LegacyNodes.cs
@@ -1786,6 +1786,30 @@ partial class GUBP
         }
     }
 
+	public class FullGameAggregateNode : AggregateNode
+	{
+		string GameName;
+
+		public FullGameAggregateNode(string InGameName, List<string> NodeNames)
+		{
+			GameName = InGameName;
+
+			foreach (string NodeName in NodeNames)
+			{
+				AddDependency(NodeName);
+			}
+		}
+
+		public static string StaticGetFullName(string GameName)
+		{
+			return "FullGameAggregate_" + GameName;
+		}
+
+		public override string GetFullName()
+		{
+			return StaticGetFullName(GameName);
+		}
+	}
 
     public class WaitForUserInput : GUBPNode
     {
