UEPLAT-826 :  use  HTML5 SDK from third party directories.

   - Always pickup HTML5 sdk from third party directory.
   - Use the new html5 template by default.

diff --git a/Engine/Build/HTML5/Readme.txt b/Engine/Build/HTML5/Readme.txt
--- a/Engine/Build/HTML5/Readme.txt
+++ b/Engine/Build/HTML5/Readme.txt
@@ -14,14 +14,18 @@ Files Required for final deployment.
 Local Testing
 ------------- 
 
-run HTML5Launcher.exe (via mono on Mac) to start a web sever which is conigured to serve compressed files on localhost. This is not a production quality server.
+run HTML5Launcher.exe (via mono on Mac) to start a web sever which is configured to serve compressed files on localhost. This is not a production quality server.
 add -ServerPort=XXXX to the command line if necessary to change the serving port. Default port is 8000. 
 
 
 How to deploy on Apache 
 ------------------------ 
 
-Todo. 
+Add the following in .conf or .htaccess file. ( edit the directory name accordingly )
+
+<Directory "${SRVROOT}/htdocs">
+ AddEncoding gzip gz
+</Directory>
 
 
 
diff --git a/Engine/Source/Developer/HTML5/HTML5PlatformEditor/Classes/HTML5TargetSettings.h b/Engine/Source/Developer/HTML5/HTML5PlatformEditor/Classes/HTML5TargetSettings.h
--- a/Engine/Source/Developer/HTML5/HTML5PlatformEditor/Classes/HTML5TargetSettings.h
+++ b/Engine/Source/Developer/HTML5/HTML5PlatformEditor/Classes/HTML5TargetSettings.h
@@ -51,13 +51,6 @@ public:
  	UPROPERTY(GlobalConfig, EditAnywhere, Category=Memory, Meta = (DisplayName = "Port to use when deploying game from the editor", ClampMin="49152", ClampMax="65535"))
  	int32 DeployServerPort;
 
-
-	/**
-	* Setting to control which template to use
-	*/
-	UPROPERTY(GlobalConfig, EditAnywhere, Category = Packaging, Meta = (DisplayName = "Use new template [experimental]"))
-	bool UseExperimentalTemplate;
-
 	/**
 	 * Use a loading level and download maps during transitions.                                                                     
 	 */
diff --git a/Engine/Source/Developer/HTML5/HTML5PlatformEditor/Private/HTML5PlatformEditorModule.cpp b/Engine/Source/Developer/HTML5/HTML5PlatformEditor/Private/HTML5PlatformEditorModule.cpp
--- a/Engine/Source/Developer/HTML5/HTML5PlatformEditor/Private/HTML5PlatformEditorModule.cpp
+++ b/Engine/Source/Developer/HTML5/HTML5PlatformEditor/Private/HTML5PlatformEditorModule.cpp
@@ -20,11 +20,24 @@ class FHTML5PlatformEditorModule
 
 	virtual void StartupModule() override
 	{
-		// register settings
+#if PLATFORM_WINDOWS
+		FString SDKPath = FPaths::EngineDir() / TEXT("Source") / TEXT("ThirdParty") / TEXT("HTML5") / TEXT("emsdk") / TEXT("Win64");
+#elif PLATFORM_MAC
+		FString SDKPath = FPaths::EngineDir() / TEXT("Source") / TEXT("ThirdParty") / TEXT("HTML5") / TEXT("emsdk") / TEXT("Mac");
+#elif PLATFORM_LINUX
+		FString SDKPath = FPaths::EngineDir() / TEXT("Source") / TEXT("ThirdParty") / TEXT("HTML5") / TEXT("emsdk") / TEXT("Linux");
+#else 
+		return; 
+#endif 
+		// we don't have the SDK, don't bother setting this up. 
+		if (!IFileManager::Get().DirectoryExists(*SDKPath))
+		{
+			return; 
+		}
+
+		// register settings6
 		static FName PropertyEditor("PropertyEditor");
 		FPropertyEditorModule& PropertyModule = FModuleManager::GetModuleChecked<FPropertyEditorModule>(PropertyEditor);
-		PropertyModule.RegisterCustomPropertyTypeLayout("HTML5SDKPath", FOnGetPropertyTypeCustomizationInstance::CreateStatic(&FHTML5SDKPathCustomization::MakeInstance));
-
 		ISettingsModule* SettingsModule = FModuleManager::GetModulePtr<ISettingsModule>("Settings");
 
 		if (SettingsModule != nullptr)
@@ -41,12 +54,6 @@ class FHTML5PlatformEditorModule
  				GetMutableDefault<UHTML5SDKSettings>()
 			);
 		}
-
-		auto* Settings = GetMutableDefault<UHTML5SDKSettings>();
-		if (Settings->DeviceMap.Num() == 0)
-		{
-			Settings->QueryKnownBrowserLocations();
-		}
 	}
 
 	virtual void ShutdownModule() override
diff --git a/Engine/Source/Developer/HTML5/HTML5PlatformEditor/Private/HTML5SDKSettings.cpp b/Engine/Source/Developer/HTML5/HTML5PlatformEditor/Private/HTML5SDKSettings.cpp
--- a/Engine/Source/Developer/HTML5/HTML5PlatformEditor/Private/HTML5SDKSettings.cpp
+++ b/Engine/Source/Developer/HTML5/HTML5PlatformEditor/Private/HTML5SDKSettings.cpp
@@ -29,262 +29,12 @@ void UHTML5SDKSettings::PostEditChangeProperty(struct FPropertyChangedEvent& Pro
 	{
 		TargetPlatformModule = FModuleManager::LoadModulePtr<IHTML5TargetPlatformModule>("HTML5TargetPlatform");
 	}
-	
-	if (TargetPlatformModule) 
+	if (TargetPlatformModule)
 	{
-		UpdateGlobalUserConfigFile();//SaveConfig();
+		UpdateGlobalUserConfigFile();
 		TargetPlatformModule->RefreshAvailableDevices();
 	}
 }
 
-void UHTML5SDKSettings::QueryKnownBrowserLocations()
-{
-	if (DeviceMap.Num() != 0) 
-	{
-		return;
-	}
-
-	struct FBrowserLocation 
-	{
-		FString Name;
-		FString Path;
-	} PossibleLocations[] = 
-	{
-#if PLATFORM_WINDOWS
-		{TEXT("Nightly(64bit)"), TEXT("C:/Program Files/Nightly/firefox.exe")},
-		{TEXT("Nightly"), TEXT("C:/Program Files (x86)/Nightly/firefox.exe")},
-		{TEXT("Firefox"), TEXT("C:/Program Files (x86)/Mozilla Firefox/firefox.exe")},
-		{TEXT("Chrome"), TEXT("C:/Program Files (x86)/Google/Chrome/Application/chrome.exe")},
-#elif PLATFORM_MAC
-		{TEXT("Safari"), TEXT("/Applications/Safari.app")},
-		{TEXT("Firefox"),TEXT("/Applications/Firefox.app")},
-		{TEXT("Chrome"),TEXT("/Applications/Google Chrome.app")},
-#elif PLATFORM_LINUX
-		{TEXT("Firefox"), TEXT("/usr/bin/firefox")},
-#else
-		{TEXT("Firefox"), TEXT("")},
-#endif
-	};
-
-	for (const auto& Loc : PossibleLocations)
-	{
-		if (FPlatformFileManager::Get().GetPlatformFile().FileExists(*Loc.Path) || FPlatformFileManager::Get().GetPlatformFile().DirectoryExists(*Loc.Path))
-		{
-			FHTML5DeviceMapping NewDevice;
-			NewDevice.DeviceName = Loc.Name;
-			NewDevice.DevicePath.FilePath = Loc.Path;
-			DeviceMap.Add(NewDevice);
-		}
-	}
-}
-
-TSharedRef<IPropertyTypeCustomization> FHTML5SDKPathCustomization::MakeInstance()
-{
-	return MakeShareable(new FHTML5SDKPathCustomization());
-}
-
-void FHTML5SDKPathCustomization::CustomizeHeader(TSharedRef<IPropertyHandle> StructPropertyHandle, class FDetailWidgetRow& HeaderRow, IPropertyTypeCustomizationUtils& StructCustomizationUtils)
-{
-	PathProperty = StructPropertyHandle->GetChildHandle("SDKPath");
-	VersionProperty = StructPropertyHandle->GetChildHandle("EmscriptenVersion");
-
-	InvalidVersion = MakeShareable(new FHTML5Version());
-	InvalidVersion->Text = LOCTEXT("IncorrectSDKDirectory", "Incorrect SDK Directory");
-
-	UpdateAvailableVersions();
-
-	if (PathProperty.IsValid())
-	{
-		HeaderRow.ValueContent()
-			.MinDesiredWidth(125.0f)
-			.MaxDesiredWidth(600.0f)
-			[
-				SNew(SHorizontalBox)
-				+ SHorizontalBox::Slot()
-				.FillWidth(1.0f)
-				.VAlign(VAlign_Center)
-				[
-					SNew(SBox)
-					.MinDesiredWidth(200)
-					[
-						PathProperty->CreatePropertyValueWidget()
-					]
-				]
-				+ SHorizontalBox::Slot()
-					.AutoWidth()
-					.Padding(FMargin(4.0f, 0.0f, 0.0f, 0.0f))
-					.VAlign(VAlign_Center)
-					[
-						SAssignNew(BrowseButton, SButton)
-						.ButtonStyle(FEditorStyle::Get(), "HoverHintOnly")
-						.ToolTipText(LOCTEXT("FolderButtonToolTipText", "Choose a directory from this computer"))
-						.OnClicked(FOnClicked::CreateSP(this, &FHTML5SDKPathCustomization::OnPickDirectory, PathProperty.ToSharedRef()))
-						.ContentPadding(2.0f)
-						.ForegroundColor(FSlateColor::UseForeground())
-						.IsFocusable(false)
-						[
-							SNew(SImage)
-							.Image(FEditorStyle::GetBrush("PropertyWindow.Button_Ellipsis"))
-							.ColorAndOpacity(FSlateColor::UseForeground())
-						]
-					]
-				+ SHorizontalBox::Slot()
-					.AutoWidth()
-					.Padding(FMargin(4.0f, 0.0f, 0.0f, 0.0f))
-					.VAlign(VAlign_Center)
-					[
-						SAssignNew(SDKSelector, SComboBox< TSharedPtr<FHTML5Version> >)
-						.OptionsSource(&Options)
-						.InitiallySelectedItem(Options[0])
-						.OnSelectionChanged(this, &FHTML5SDKPathCustomization::OnSelectionChanged)
-						.ContentPadding(2)
-						.OnGenerateWidget(this, &FHTML5SDKPathCustomization::OnGenerateWidget)
-						[
-							SNew(STextBlock)
-							.Text(this, &FHTML5SDKPathCustomization::GetSelectedText)
-						]
-					]
-			]
-		.NameContent()
-			[
-				StructPropertyHandle->CreatePropertyNameWidget()
-			];
-	}
-}
-
-void FHTML5SDKPathCustomization::CustomizeChildren(TSharedRef<IPropertyHandle> StructPropertyHandle, class IDetailChildrenBuilder& StructBuilder, IPropertyTypeCustomizationUtils& StructCustomizationUtils)
-{
-}
-
-FReply FHTML5SDKPathCustomization::OnPickDirectory(TSharedRef<IPropertyHandle> PropertyHandle)// const
-{
-	FString Directory;
-	IDesktopPlatform* DesktopPlatform = FDesktopPlatformModule::Get();
-	auto& PlatformFile = FPlatformFileManager::Get().GetPlatformFile();
-	auto* Settings = GetMutableDefault<UHTML5SDKSettings>();
-	if (!DesktopPlatform || !Settings)
-	{
-		return FReply::Handled();
-	}
-
-	TSharedPtr<SWindow> ParentWindow = FSlateApplication::Get().FindWidgetWindow(BrowseButton.ToSharedRef());
-	void* ParentWindowHandle = (ParentWindow.IsValid() && ParentWindow->GetNativeWindow().IsValid()) ? ParentWindow->GetNativeWindow()->GetOSWindowHandle() : nullptr;
-
-	FString StartDirectory = FEditorDirectories::Get().GetLastDirectory(ELastDirectory::GENERIC_IMPORT);
-	for (;;)
-	{
-		if (DesktopPlatform->OpenDirectoryDialog(ParentWindowHandle, LOCTEXT("FolderDialogTitle", "Choose a directory").ToString(), StartDirectory, Directory))
-		{
-			FText FailureReason;
- 			if (PlatformFile.DirectoryExists(*Directory))
- 			{
-				// TODO: Update config value too
-				Settings->EmscriptenRoot.SDKPath = Directory;
-				PathProperty->SetValue(Directory);
-				UpdateAvailableVersions();
-				SDKSelector->SetSelectedItem(SelectedVersionItem);
-				SDKSelector->RefreshOptions();
- 			}
- 			else
-			{
-				StartDirectory = Directory;
-				FMessageDialog::Open(EAppMsgType::Ok, FailureReason);
-				continue;
-			}
-		}
-		break;
-	}
-
-	return FReply::Handled();
-}
-
-TSharedRef<SWidget> FHTML5SDKPathCustomization::OnGenerateWidget(TSharedPtr<FHTML5Version> InItem)
-{
-	return SNew(STextBlock)
-		.Text(InItem->Text);
-}
-
-FText FHTML5SDKPathCustomization::GetSelectedText() const
-{
-	return SelectedVersionItem->Text;
-}
-
-void FHTML5SDKPathCustomization::OnSelectionChanged(TSharedPtr<FHTML5Version> InItem, ESelectInfo::Type InSeletionInfo)
-{
-	SelectedVersionItem = InItem;
-	auto* Settings = GetMutableDefault<UHTML5SDKSettings>();
-	VersionProperty->SetValue(InItem->Version.VersionNumberAsString());
-}
-
-void FHTML5SDKPathCustomization::UpdateAvailableVersions()
-{
-	Options.Empty();
-	Options.Add(InvalidVersion);
-
-	//Take the cheap option here. SDK is now uninstalled. Don't query UBT!
-	// If we clear this function we're good to install again
-	PlatformInfo::UpdatePlatformSDKStatus(TEXT("HTML5"), PlatformInfo::EPlatformSDKStatus::NotInstalled);
-
-	SelectedVersionItem = Options[0];
-
-	auto* Settings = GetMutableDefault<UHTML5SDKSettings>();
-	auto& PlatformFile = FPlatformFileManager::Get().GetPlatformFile();
-	TargetPlatformModule = FModuleManager::LoadModulePtr<IHTML5TargetPlatformModule>("HTML5TargetPlatform");
-	if (!Settings || !TargetPlatformModule)
-	{
-		return;
-	}
-
-	TSharedPtr<FHTML5Version> PrevSelectedVersion(new FHTML5Version);
-	PrevSelectedVersion->Version.VersionNumberFromString(*Settings->EmscriptenRoot.EmscriptenVersion);
-
-	TArray<FHTML5SDKVersionNumber> SDKVersions;
-	TargetPlatformModule->GetInstalledSDKVersions(*(Settings->EmscriptenRoot.SDKPath), SDKVersions);
-
-	if (SDKVersions.Num() == 0) 
-	{
-		//Try the environment variable EMSCRIPTEN - This assumes latest version.
-		TCHAR SDKDirectory[32768] = { 0 };
-		FPlatformMisc::GetEnvironmentVariable(TEXT("EMSCRIPTEN"), SDKDirectory, ARRAY_COUNT(SDKDirectory));
-		TargetPlatformModule->GetInstalledSDKVersions(SDKDirectory, SDKVersions);
-		PrevSelectedVersion->Version.MakeLatestVersionNumber();
-	}
-
-	if (SDKVersions.Num() == 0)
-	{
-		return;
-	}
-	Options.Empty();
-	for (const auto& i : SDKVersions)
-	{
-		TSharedPtr<FHTML5Version> Ver(new FHTML5Version);
-		Ver->Text = FText::FromString(i.VersionNumberAsString());
-		Ver->Version = i;
-		Options.Add(Ver);
-	}
-
-	TSharedPtr<FHTML5Version> Latest(new FHTML5Version);
-	*Latest = *Options.Last();
-	Latest->Text = LOCTEXT("LastestVesion", "Latest");
-	Latest->Version.MakeLatestVersionNumber();
-	Options.Add(Latest);
-
-	int32 FoundIndex = Options.IndexOfByPredicate([&](const TSharedPtr<FHTML5Version>& RHS)
-	{
-		return RHS->Version == PrevSelectedVersion->Version;
-	});
-	if (FoundIndex != INDEX_NONE)
-	{
-		SelectedVersionItem = Options[FoundIndex];
-	}
-	else
-	{
-		SelectedVersionItem = Options.Last();
-	}
-
-	//Take the cheap option here. If we got here, there is a valid install selected. 
-	// Don't query UBT, it'll repeat the same test via a separate process! :(
-	PlatformInfo::UpdatePlatformSDKStatus(TEXT("HTML5"), PlatformInfo::EPlatformSDKStatus::Installed);
-}
 #endif //WITH_EDITOR
 #undef LOCTEXT_NAMESPACE
\ No newline at end of file
diff --git a/Engine/Source/Developer/HTML5/HTML5PlatformEditor/Private/HTML5SDKSettings.h b/Engine/Source/Developer/HTML5/HTML5PlatformEditor/Private/HTML5SDKSettings.h
--- a/Engine/Source/Developer/HTML5/HTML5PlatformEditor/Private/HTML5SDKSettings.h
+++ b/Engine/Source/Developer/HTML5/HTML5PlatformEditor/Private/HTML5SDKSettings.h
@@ -23,19 +23,6 @@ struct FHTML5DeviceMapping
 	FFilePath DevicePath;
 };
 
-USTRUCT()
-struct FHTML5SDKPath
-{
-	GENERATED_USTRUCT_BODY()
-
-	UPROPERTY(EditAnywhere, Category = SDK_Path)
-	FString SDKPath;
-
-	// Note that 0.0.0 is invalid and -1.-1.-1 is latest
-	UPROPERTY(EditAnywhere, Category = SDK_Path)
-	FString EmscriptenVersion;
-};
-
 /**
  * Implements the settings for the HTML5 SDK setup.
  */
@@ -45,80 +32,17 @@ class HTML5PLATFORMEDITOR_API UHTML5SDKSettings : public UObject
 public:
 	GENERATED_UCLASS_BODY()
 
-	// Path to the python executable. Can be blank is python is on the PATH
-	UPROPERTY(GlobalConfig, EditAnywhere, Category = HTML5_SDK_Paths, Meta = (DisplayName = "Location of Python exe"))
-	FFilePath Python;
-
 	// Available browsers that can be used when launching HTML5 builds.
 	UPROPERTY(GlobalConfig, EditAnywhere, Category = HTML5_Devices, Meta = (DisplayName = "Available browsers"))
 	TArray<FHTML5DeviceMapping> DeviceMap;
 
-	// Path to Emscripten SDK install directory. This is the directory which contains emsdk.bat and emsdk.sh
-	UPROPERTY(GlobalConfig, EditAnywhere, Category = HTML5_SDK_Paths, Meta = (DisplayName = "Location of Emscripten SDK"))
-	FHTML5SDKPath EmscriptenRoot;
-
 #if WITH_EDITOR
+
 	// UObject interface
 	virtual void PostEditChangeProperty(struct FPropertyChangedEvent& PropertyChangedEvent) override;
 	// End of UObject interface
-	void QueryKnownBrowserLocations();
-
 	ITargetPlatformManagerModule * TargetManagerModule;
 #endif
 private:
 	IHTML5TargetPlatformModule* TargetPlatformModule;
 };
-
-#if WITH_EDITOR
-struct FHTML5Version
-{
-	// Version text to display
-	FText Text;
-
-	FHTML5SDKVersionNumber Version;
-
-	bool operator == (const FHTML5Version& RHS) const
-	{
-		return Version == RHS.Version;
-	}
-};
-
-class FHTML5SDKPathCustomization : public IPropertyTypeCustomization
-{
-public:
-	static TSharedRef<IPropertyTypeCustomization> MakeInstance();
-
-	// IPropertyTypeCustomization interface 
-	virtual void CustomizeHeader(TSharedRef<class IPropertyHandle> StructPropertyHandle, class FDetailWidgetRow& HeaderRow, IPropertyTypeCustomizationUtils& StructCustomizationUtils) override;
-
-	virtual void CustomizeChildren(TSharedRef<class IPropertyHandle> StructPropertyHandle, class IDetailChildrenBuilder& StructBuilder, IPropertyTypeCustomizationUtils& StructCustomizationUtils) override;
-
-private:
-
-	FReply OnPickDirectory(TSharedRef<IPropertyHandle> PropertyHandle);// const;
-
-	TSharedRef<SWidget> OnGenerateWidget(TSharedPtr<FHTML5Version> InItem);
-
-	void OnSelectionChanged(TSharedPtr<FHTML5Version> InItem, ESelectInfo::Type InSeletionInfo);
-
-	FText GetSelectedText() const;
-
-	void UpdateAvailableVersions();
-
-	/** The browse button widget */
-	TSharedPtr<SButton> BrowseButton;
-
-	TSharedPtr< SComboBox< TSharedPtr<FHTML5Version> > > SDKSelector;
-
-	TArray< TSharedPtr<FHTML5Version> > Options;
-
-	TSharedPtr< FHTML5Version > SelectedVersionItem;
-
-	TSharedPtr< FHTML5Version > InvalidVersion;
-
-	IHTML5TargetPlatformModule* TargetPlatformModule;
-
-	TSharedPtr<IPropertyHandle> PathProperty;
-	TSharedPtr<IPropertyHandle> VersionProperty;
-};
-#endif // WITH_EDITOR
diff --git a/Engine/Source/Developer/HTML5/HTML5TargetPlatform/Private/HTML5TargetDevice.cpp b/Engine/Source/Developer/HTML5/HTML5TargetPlatform/Private/HTML5TargetDevice.cpp
--- a/Engine/Source/Developer/HTML5/HTML5TargetPlatform/Private/HTML5TargetDevice.cpp
+++ b/Engine/Source/Developer/HTML5/HTML5TargetPlatform/Private/HTML5TargetDevice.cpp
@@ -21,10 +21,9 @@ void FHTML5TargetDevice::Disconnect( )
 {
 }
 
-
 FTargetDeviceId FHTML5TargetDevice::GetId( ) const
 {
-	return FTargetDeviceId(TargetPlatform.PlatformName(), GetName());
+	return FTargetDeviceId(TargetPlatform.PlatformName(), Path);
 }
 
 
diff --git a/Engine/Source/Developer/HTML5/HTML5TargetPlatform/Private/HTML5TargetDevice.h b/Engine/Source/Developer/HTML5/HTML5TargetPlatform/Private/HTML5TargetDevice.h
--- a/Engine/Source/Developer/HTML5/HTML5TargetPlatform/Private/HTML5TargetDevice.h
+++ b/Engine/Source/Developer/HTML5/HTML5TargetPlatform/Private/HTML5TargetDevice.h
@@ -34,9 +34,10 @@ public:
 	 * @param InTargetPlatform - The target platform.
 	 * @param InName - The device name.
 	 */
-	FHTML5TargetDevice( const ITargetPlatform& InTargetPlatform, const FString& InName )
+	FHTML5TargetDevice( const ITargetPlatform& InTargetPlatform, const FString& InName, const FString& InPath )
 		: TargetPlatform(InTargetPlatform), 
-		  Name(InName)
+		  Name(InName),
+		  Path(InPath)
 	{ }
 
 	/**
@@ -114,4 +115,6 @@ private:
 	const ITargetPlatform& TargetPlatform;
 
 	FString Name; 
+
+	FString Path;
 };
diff --git a/Engine/Source/Developer/HTML5/HTML5TargetPlatform/Private/HTML5TargetPlatform.cpp b/Engine/Source/Developer/HTML5/HTML5TargetPlatform/Private/HTML5TargetPlatform.cpp
--- a/Engine/Source/Developer/HTML5/HTML5TargetPlatform/Private/HTML5TargetPlatform.cpp
+++ b/Engine/Source/Developer/HTML5/HTML5TargetPlatform/Private/HTML5TargetPlatform.cpp
@@ -15,8 +15,7 @@
 
 FHTML5TargetPlatform::FHTML5TargetPlatform( )
 {
-	RefreshAvailableDevices();
-
+	RefreshHTML5Setup();
 #if WITH_ENGINE
 	// load up texture settings from the config file
 	HTML5LODSettings = nullptr;
@@ -62,113 +61,22 @@ ITargetDevicePtr FHTML5TargetPlatform::GetDevice( const FTargetDeviceId& DeviceI
 
 bool FHTML5TargetPlatform::IsSdkInstalled(bool bProjectHasCode, FString& OutDocumentationPath) const
 {
-	FConfigSection SDKPaths;
-	//New style detection of devices
-	if (HTML5EngineSettings.Find("/Script/HTML5PlatformEditor.HTML5SDKSettings"))
-	{
-
-		SDKPaths = HTML5EngineSettings["/Script/HTML5PlatformEditor.HTML5SDKSettings"];
-		for (auto It : SDKPaths)
-		{
-			const FString& Platform = It.Key.ToString();
-			FString Path = It.Value;
-			if (Platform == "EmscriptenRoot")
-			{
-				FString SDKPathString;
-				FString VersionString;
-				auto Index = Path.Find(TEXT("SDKPath="));
-				if (Index != INDEX_NONE)
-				{
-					Index+=9;
-					auto Index2 = Path.Find(TEXT("\""), ESearchCase::IgnoreCase, ESearchDir::FromStart, Index);
-					if (Index2 != INDEX_NONE)
-					{
-						SDKPathString = Path.Mid(Index, Index2-Index);
-					}
-				}
-				Index = Path.Find(TEXT("EmscriptenVersion="));
-				if (Index != INDEX_NONE)
-				{
-					Index += 19;
-					auto Index2 = Path.Find(TEXT("\""), ESearchCase::IgnoreCase, ESearchDir::FromStart, Index);
-					if (Index2 != INDEX_NONE)
-					{
-						VersionString = Path.Mid(Index, Index2 - Index);
-					}
-				}
-				if (!VersionString.IsEmpty() && !SDKPathString.IsEmpty())
-				{
-					bool bFoundSDK = false;
-					FHTML5SDKVersionNumber RequiredVersion;
-					RequiredVersion.VersionNumberFromString(*VersionString);
-					TArray<FHTML5SDKVersionNumber> SDKVersions;
-					GetInstalledSDKVersions(*SDKPathString, SDKVersions);
-
-					if (SDKVersions.Num() > 0 &&
-						RequiredVersion.VersionNumber.Major == -1 &&
-						RequiredVersion.VersionNumber.Minor == -1 &&
-						RequiredVersion.VersionNumber.Revision == -1)
-					{
-						bFoundSDK = true;
-						SDKPathString = SDKVersions.Last().VersionPath;
-					}
-					else
-					{
-						for (const auto& i : SDKVersions)
-						{
-							if (i == RequiredVersion)
-							{
-								SDKPathString = i.VersionPath;
-								bFoundSDK = true;
-								break;
-							}
-						}
-					}
-
-					
-					return bFoundSDK && IFileManager::Get().DirectoryExists(*SDKPathString);
-				}
-			}
-		}
-	}
-
-	// Old fallbacks
-	if (HTML5EngineSettings.Find("HTML5SDKPaths"))
-	{
-		FString SectionName = "HTML5SDKPaths";
-		SDKPaths = HTML5EngineSettings[SectionName];
-		for (auto It : SDKPaths)
-		{
-			const FString& Platform = It.Key.ToString();
-			const FString& Path = It.Value;
-			{
-				if (Platform == "Emscripten" && IFileManager::Get().DirectoryExists(*Path))
-				{
-					return true;
-				}
 #if PLATFORM_WINDOWS
-				if (Platform == "Windows" && IFileManager::Get().DirectoryExists(*Path))
-				{
-					return true;
-				}
-#endif 
-#if PLATFORM_MAC
-				if (Platform == "Mac" && IFileManager::Get().DirectoryExists(*Path))
-				{
-					return true;
-				}
+		FString SDKPath = FPaths::EngineDir() / TEXT("Source") / TEXT("ThirdParty") / TEXT("HTML5") / TEXT("emsdk") / TEXT("Win64");
+#elif PLATFORM_MAC
+		FString SDKPath = FPaths::EngineDir() / TEXT("Source") / TEXT("ThirdParty") / TEXT("HTML5") / TEXT("emsdk") / TEXT("Mac");
+#elif PLATFORM_LINUX
+		FString SDKPath = FPaths::EngineDir() / TEXT("Source") / TEXT("ThirdParty") / TEXT("HTML5") / TEXT("emsdk") / TEXT("Linux");
+#else 
+		return; 
 #endif 
-			}
-		}
-	}
 
-	TCHAR BaseSDKPath[512];
-	FPlatformMisc::GetEnvironmentVariable(TEXT("EMSCRIPTEN"), BaseSDKPath, ARRAY_COUNT(BaseSDKPath));
-	if (FString(BaseSDKPath).Len() > 0 && IFileManager::Get().DirectoryExists(BaseSDKPath) )
+	FString SDKDirectory = FPaths::ConvertRelativePathToFull(SDKPath);
+
+	if (IFileManager::Get().DirectoryExists(*SDKDirectory))
 	{
 		return true; 	
 	}
-
 	return false; 
 }
 
@@ -319,10 +227,14 @@ FName FHTML5TargetPlatform::GetWaveFormat( const USoundWave* Wave ) const
 
 #endif // WITH_ENGINE
 
-void FHTML5TargetPlatform::RefreshAvailableDevices()
+void FHTML5TargetPlatform::RefreshHTML5Setup()
 {
-	// load the final HTML5 engine settings for this game
-	FConfigCacheIni::LoadLocalIniFile(HTML5EngineSettings, TEXT("Engine"), true, *PlatformName());
+	FString Temp; 
+	if (!FHTML5TargetPlatform::IsSdkInstalled(true, Temp))
+	{
+		// nothing to do. 
+		return;
+	}
 
 	//New style detection of devices
 	for (const auto& Device : LocalDevice)
@@ -337,9 +249,10 @@ void FHTML5TargetPlatform::RefreshAvailableDevices()
 		Config = &HTML5EngineSettings;
 	}
 	TArray<FString> ValueArray;
-	if (Config->Find("/Script/HTML5PlatformEditor.HTML5SDKSettings")) 
+
+	if (Config->Find("/Script/HTML5PlatformEditor.HTML5SDKSettings"))
 	{
-		FConfigSection AvaliableDevicesNewSection =  (*Config)["/Script/HTML5PlatformEditor.HTML5SDKSettings"];
+		FConfigSection AvaliableDevicesNewSection = (*Config)["/Script/HTML5PlatformEditor.HTML5SDKSettings"];
 		for (auto It : AvaliableDevicesNewSection)
 		{
 			ValueArray.Reset();
@@ -368,7 +281,7 @@ void FHTML5TargetPlatform::RefreshAvailableDevices()
 					(FPlatformFileManager::Get().GetPlatformFile().FileExists(*DevicePath) ||
                      FPlatformFileManager::Get().GetPlatformFile().DirectoryExists(*DevicePath)))
 				{
-					ITargetDevicePtr Device = MakeShareable(new FHTML5TargetDevice(*this, FString::Printf(TEXT("%s on %s"), *DeviceName, FPlatformProcess::ComputerName())));
+					ITargetDevicePtr Device = MakeShareable(new FHTML5TargetDevice(*this, DeviceName, DevicePath));
 					LocalDevice.Add(Device);
 					DeviceDiscoveredEvent.Broadcast(Device.ToSharedRef());
 				}
@@ -376,78 +289,36 @@ void FHTML5TargetPlatform::RefreshAvailableDevices()
 		}
 	}
 
-	// Old Fallback
-	if (LocalDevice.Num() == 0)
+	struct FBrowserLocation
+	{
+		FString Name;
+		FString Path;
+	} PossibleLocations[] =
 	{
-		FString DeviceSectionName;
-
 #if PLATFORM_WINDOWS
-		DeviceSectionName = "HTML5DevicesWindows";
+		{ TEXT("Nightly(64bit)"), TEXT("C:/Program Files/Nightly/firefox.exe") },
+		{ TEXT("Nightly"), TEXT("C:/Program Files (x86)/Nightly/firefox.exe") },
+		{ TEXT("Firefox"), TEXT("C:/Program Files (x86)/Mozilla Firefox/firefox.exe") },
+	//	{ TEXT("Chrome"), TEXT("C:/Program Files (x86)/Google/Chrome/Application/chrome.exe") },
 #elif PLATFORM_MAC
-		DeviceSectionName = "HTML5DevicesMac";
-#else 
-		DeviceSectionName = "HTML5DevicesLinux";
-#endif 
-		if (HTML5EngineSettings.Find(DeviceSectionName))
-		{
-			FConfigSection AvaliableDevicesSection = HTML5EngineSettings[DeviceSectionName];
-			for (auto It : AvaliableDevicesSection)
-			{
-				const FString& BrowserName = It.Key.ToString();
-				const FString& BrowserPath = It.Value;
-				if (FPlatformFileManager::Get().GetPlatformFile().FileExists(*It.Value))
-				{
-					ITargetDevicePtr Device = MakeShareable(new FHTML5TargetDevice(*this, FString::Printf(TEXT("%s on %s"), *It.Key.ToString(), FPlatformProcess::ComputerName())));
-					LocalDevice.Add(Device);
-					DeviceDiscoveredEvent.Broadcast(Device.ToSharedRef());
-				}
-			}
-		}
-	}
-}
-
-void FHTML5TargetPlatform::GetInstalledSDKVersions(const TCHAR* SDKDirectory, TArray<FHTML5SDKVersionNumber>& OutSDKs)
-{
-	if (!SDKDirectory)
-	{
-		return;
-	}
+		{ TEXT("Safari"), TEXT("/Applications/Safari.app") },
+		{ TEXT("Firefox"), TEXT("/Applications/Firefox.app") },
+		{ TEXT("Chrome"), TEXT("/Applications/Google Chrome.app") },
+#elif PLATFORM_LINUX
+		{ TEXT("Firefox"), TEXT("/usr/bin/firefox") },
+#else
+		{ TEXT("Firefox"), TEXT("") },
+#endif
+	};
 
-	struct FVersionSearch : public IPlatformFile::FDirectoryVisitor
+	// Add default devices..
+	for (const auto& Loc : PossibleLocations)
 	{
-		IPlatformFile* PlatformFile;
-		TArray<FHTML5SDKVersionNumber>& SDKs;
-		FString VersionFilePath;
-
-		FVersionSearch(IPlatformFile* InPlatformFile, TArray<FHTML5SDKVersionNumber>& InSDKs)
-			: PlatformFile(InPlatformFile)
-			, SDKs(InSDKs)
-		{}
-		virtual bool Visit(const TCHAR* FilenameOrDirectory, bool bIsDirectory)
+		if (FPlatformFileManager::Get().GetPlatformFile().FileExists(*Loc.Path) || FPlatformFileManager::Get().GetPlatformFile().DirectoryExists(*Loc.Path))
 		{
-			if (bIsDirectory == true)
-			{
-				VersionFilePath = FPaths::Combine(FilenameOrDirectory, TEXT("emscripten-version.txt"));
-				if (PlatformFile->FileExists(*VersionFilePath))
-				{
-					FString VersionText;
-					FFileHelper::LoadFileToString(VersionText, *VersionFilePath);
-					VersionText.Trim();
-					VersionText.TrimTrailing();
-					FHTML5SDKVersionNumber Ver;
-					Ver.VersionPath = FilenameOrDirectory;
-					Ver.VersionNumberFromString(*VersionText);
-					SDKs.Add(Ver);
-				}
-			}
-			return true;
+			ITargetDevicePtr Device = MakeShareable(new FHTML5TargetDevice(*this, *Loc.Name, *Loc.Path));
+			LocalDevice.Add(Device);
+			DeviceDiscoveredEvent.Broadcast(Device.ToSharedRef());
 		}
-	};
-
-	auto& PlatformFile = FPlatformFileManager::Get().GetPlatformFile();
-
-	FVersionSearch VersionSearch(&PlatformFile, OutSDKs);
-	FString EmscriptenSDkDir = FPaths::Combine(SDKDirectory, TEXT("emscripten"));
-	PlatformFile.IterateDirectory(*EmscriptenSDkDir, VersionSearch);
-	OutSDKs.Sort();
+	}
 }
diff --git a/Engine/Source/Developer/HTML5/HTML5TargetPlatform/Private/HTML5TargetPlatform.h b/Engine/Source/Developer/HTML5/HTML5TargetPlatform/Private/HTML5TargetPlatform.h
--- a/Engine/Source/Developer/HTML5/HTML5TargetPlatform/Private/HTML5TargetPlatform.h
+++ b/Engine/Source/Developer/HTML5/HTML5TargetPlatform/Private/HTML5TargetPlatform.h
@@ -23,7 +23,7 @@ public:
 	 */
 	FHTML5TargetPlatform( );
 
-	void RefreshAvailableDevices();
+	void RefreshHTML5Setup();
 
 	// Begin ITargetPlatform interface
 
@@ -94,8 +94,6 @@ public:
 
 	// End ITargetPlatform interface
 
-	static void GetInstalledSDKVersions(const TCHAR* SDKDirectory, TArray<FHTML5SDKVersionNumber>& OutSDKs);
-
 private:
 
 	// Holds the HTML5 engine settings.
@@ -119,4 +117,5 @@ private:
 
 	// Holds an event delegate that is executed when a target device has been lost, i.e. disconnected or timed out.
 	FOnTargetDeviceLost DeviceLostEvent;
+
 };
diff --git a/Engine/Source/Developer/HTML5/HTML5TargetPlatform/Private/HTML5TargetPlatformModule.cpp b/Engine/Source/Developer/HTML5/HTML5TargetPlatform/Private/HTML5TargetPlatformModule.cpp
--- a/Engine/Source/Developer/HTML5/HTML5TargetPlatform/Private/HTML5TargetPlatformModule.cpp
+++ b/Engine/Source/Developer/HTML5/HTML5TargetPlatform/Private/HTML5TargetPlatformModule.cpp
@@ -55,15 +55,10 @@ public:
 	{
 		if (Singleton)
 		{
-			((FHTML5TargetPlatform*)Singleton)->RefreshAvailableDevices();
+			((FHTML5TargetPlatform*)Singleton)->RefreshHTML5Setup();
 		}
 	}
 
-	virtual void GetInstalledSDKVersions(const TCHAR* SDKDirectory, TArray<FHTML5SDKVersionNumber>& OutSDKs) override
-	{
-		FHTML5TargetPlatform::GetInstalledSDKVersions(SDKDirectory, OutSDKs);
-	}
-
 public:
 
 	// Begin IModuleInterface interface
diff --git a/Engine/Source/Developer/HTML5/HTML5TargetPlatform/Public/IHTML5TargetPlatformModule.h b/Engine/Source/Developer/HTML5/HTML5TargetPlatform/Public/IHTML5TargetPlatformModule.h
--- a/Engine/Source/Developer/HTML5/HTML5TargetPlatform/Public/IHTML5TargetPlatformModule.h
+++ b/Engine/Source/Developer/HTML5/HTML5TargetPlatform/Public/IHTML5TargetPlatformModule.h
@@ -3,77 +3,6 @@
 #pragma once
 
 #include "ITargetPlatformModule.h"
-
-class HTML5TARGETPLATFORM_API FHTML5SDKVersionNumber
-{
-public:
-	FHTML5SDKVersionNumber()
-	{
-		// Zero is invalid, -1 is latest.
-		VersionNumber.Major = 0;
-		VersionNumber.Minor = 0;
-		VersionNumber.Revision = 0;
-	}
-
-	// Directory of the SDK 
-	FString VersionPath;
-
-	//Version Number
-	struct
-	{
-		int32 Major;
-		int32 Minor;
-		int32 Revision;
-	} VersionNumber;
-
-	void VersionNumberFromString(const TCHAR* VersionString)
-	{
-		// We expect the format "1.23.45"
-		const TCHAR* VStr = VersionString;
-		TCHAR* VEnd = nullptr;
-		int32 Nums[3] = { 0 };
-		VersionNumber.Major = 0;
-		VersionNumber.Minor = 0;
-		VersionNumber.Revision = 0;
-		for (auto& N : Nums)
-		{
-			N = FCString::Strtoi(VStr, &VEnd, 10);
-			if (!VEnd || VStr == VEnd || (*VEnd != '.' && *VEnd != 0))
-			{
-				return;
-			}
-			VStr = VEnd + 1;
-		}
-		VersionNumber.Major = Nums[0];
-		VersionNumber.Minor = Nums[1];
-		VersionNumber.Revision = Nums[2];
-	}
-
-	FString VersionNumberAsString() const
-	{
-		return FString::Printf(TEXT("%d.%d.%d"), VersionNumber.Major, VersionNumber.Minor, VersionNumber.Revision);
-	}
-
-	void MakeLatestVersionNumber()
-	{
-		VersionNumber.Major = -1;
-		VersionNumber.Minor = -1;
-		VersionNumber.Revision = -1;
-	}
-
-	bool operator < (const FHTML5SDKVersionNumber& RHS) const
-	{
-		return (VersionNumber.Major < RHS.VersionNumber.Major) ||
-			(VersionNumber.Major == RHS.VersionNumber.Major && VersionNumber.Minor < RHS.VersionNumber.Minor) ||
-			((VersionNumber.Major == RHS.VersionNumber.Major) && (VersionNumber.Minor == RHS.VersionNumber.Minor) && (VersionNumber.Revision < RHS.VersionNumber.Revision));
-	}
-
-	bool operator == (const FHTML5SDKVersionNumber& RHS) const
-	{
-		return VersionNumber.Major == RHS.VersionNumber.Major && VersionNumber.Minor == RHS.VersionNumber.Minor && VersionNumber.Revision == RHS.VersionNumber.Revision;
-	}
-};
-
 /**
  * Interface for HTML5TargetPlatformModule module.
  */
@@ -86,11 +15,6 @@ public:
 	 */
 	virtual void RefreshAvailableDevices() = 0;
 
-	/**
-	 * Fills a list of installed SDK version numbers
-	 */
-	virtual void GetInstalledSDKVersions(const TCHAR* SDKDirectory, TArray<FHTML5SDKVersionNumber>& OutSDKs) = 0;
-
 protected:
 
 	/**
diff --git a/Engine/Source/Programs/AutomationTool/HTML5/HTML5Platform.Automation.cs b/Engine/Source/Programs/AutomationTool/HTML5/HTML5Platform.Automation.cs
--- a/Engine/Source/Programs/AutomationTool/HTML5/HTML5Platform.Automation.cs
+++ b/Engine/Source/Programs/AutomationTool/HTML5/HTML5Platform.Automation.cs
@@ -9,78 +9,80 @@ using AutomationTool;
 using UnrealBuildTool;
 using System.Diagnostics;
 using System.Threading;
+using System.Threading.Tasks;
 using Ionic.Zip;
 
+
 public class HTML5Platform : Platform
 {
-    public HTML5Platform()
+	public HTML5Platform()
 		: base(UnrealTargetPlatform.HTML5)
 	{
 	}
 
 	public override void Package(ProjectParams Params, DeploymentContext SC, int WorkingCL)
 	{
-        Log("Package {0}", Params.RawProjectPath);
-        
-        string PackagePath = Path.Combine(Path.GetDirectoryName(Params.RawProjectPath), "Binaries", "HTML5");
-        if (!Directory.Exists(PackagePath))
-        {
-            Directory.CreateDirectory(PackagePath);
-        }
-        string FinalDataLocation = Path.Combine(PackagePath, Params.ShortProjectName) + ".data";
+		Log("Package {0}", Params.RawProjectPath);
+		
+		string PackagePath = Path.Combine(Path.GetDirectoryName(Params.RawProjectPath), "Binaries", "HTML5");
+		if (!Directory.Exists(PackagePath))
+		{
+			Directory.CreateDirectory(PackagePath);
+		}
+		string FinalDataLocation = Path.Combine(PackagePath, Params.ShortProjectName) + ".data";
 
-        var ConfigCache = new UnrealBuildTool.ConfigCacheIni(UnrealTargetPlatform.HTML5, "Engine", Path.GetDirectoryName(Params.RawProjectPath), CommandUtils.CombinePaths(CommandUtils.CmdEnv.LocalRoot, "Engine"));
+		var ConfigCache = new UnrealBuildTool.ConfigCacheIni(UnrealTargetPlatform.HTML5, "Engine", Path.GetDirectoryName(Params.RawProjectPath), CommandUtils.CombinePaths(CommandUtils.CmdEnv.LocalRoot, "Engine"));
 
-        if (HTMLPakAutomation.CanCreateMapPaks(Params))
-        {
-            HTMLPakAutomation PakAutomation = new HTMLPakAutomation(Params, SC);
-
-            // Create Necessary Paks. 
-            PakAutomation.CreateEnginePak();
-            PakAutomation.CreateGamePak();
-            PakAutomation.CreateContentDirectoryPak();
-
-            // Create Emscripten Package from Necessary Paks. - This will be the VFS. 
-            PakAutomation.CreateEmscriptenDataPackage(PackagePath, FinalDataLocation);
-
-            // Create All Map Paks which  will be downloaded on the fly. 
-            PakAutomation.CreateMapPak();
-
-            // Create Delta Paks if setup.
-            List<string> Paks = new List<string>();
-            ConfigCache.GetArray("/Script/HTML5PlatformEditor.HTML5TargetSettings", "LevelTransitions", out Paks);
-
-            if (Paks != null)
-            {
-                foreach (var Pak in Paks)
-                {
-                    var Matched = Regex.Matches(Pak, "\"[^\"]+\"", RegexOptions.IgnoreCase);
-                    string MapFrom = Path.GetFileNameWithoutExtension(Matched[0].ToString().Replace("\"", ""));
-                    string MapTo = Path.GetFileNameWithoutExtension(Matched[1].ToString().Replace("\"", ""));
-                    PakAutomation.CreateDeltaMapPaks(MapFrom, MapTo);
-                }
-            }
-        }
-        else
-        {
-            // we need to operate in the root
-            using (new PushedDirectory(Path.Combine(Params.BaseStageDirectory, "HTML5")))
-            {
-                string PythonPath = HTML5SDKInfo.PythonPath();
-                string PackagerPath = HTML5SDKInfo.EmscriptenPackager();
-
-                string CmdLine = string.Format("\"{0}\" \"{1}\" --preload . --js-output=\"{1}.js\"", PackagerPath, FinalDataLocation);
-                RunAndLog(CmdEnv, PythonPath, CmdLine);
-            }
-        }
+		if (HTMLPakAutomation.CanCreateMapPaks(Params))
+		{
+			HTMLPakAutomation PakAutomation = new HTMLPakAutomation(Params, SC);
 
-        // copy the "Executable" to the package directory
-        string GameExe = Path.GetFileNameWithoutExtension(Params.ProjectGameExeFilename);
-        if (Params.ClientConfigsToBuild[0].ToString() != "Development")
-        {
-            GameExe += "-HTML5-" + Params.ClientConfigsToBuild[0].ToString();
-        }
-        GameExe += ".js";
+			// Create Necessary Paks. 
+			PakAutomation.CreateEnginePak();
+			PakAutomation.CreateGamePak();
+			PakAutomation.CreateContentDirectoryPak();
+
+			// Create Emscripten Package from Necessary Paks. - This will be the VFS. 
+			PakAutomation.CreateEmscriptenDataPackage(PackagePath, FinalDataLocation);
+
+			// Create All Map Paks which  will be downloaded on the fly. 
+			PakAutomation.CreateMapPak();
+
+			// Create Delta Paks if setup.
+			List<string> Paks = new List<string>();
+			ConfigCache.GetArray("/Script/HTML5PlatformEditor.HTML5TargetSettings", "LevelTransitions", out Paks);
+
+			if (Paks != null)
+			{
+				foreach (var Pak in Paks)
+				{
+					var Matched = Regex.Matches(Pak, "\"[^\"]+\"", RegexOptions.IgnoreCase);
+					string MapFrom = Path.GetFileNameWithoutExtension(Matched[0].ToString().Replace("\"", ""));
+					string MapTo = Path.GetFileNameWithoutExtension(Matched[1].ToString().Replace("\"", ""));
+					PakAutomation.CreateDeltaMapPaks(MapFrom, MapTo);
+				}
+			}
+		}
+		else
+		{
+			// we need to operate in the root
+			string PythonPath = HTML5SDKInfo.Python();
+			string PackagerPath = HTML5SDKInfo.EmscriptenPackager();
+
+			using (new PushedDirectory(Path.Combine(Params.BaseStageDirectory, "HTML5")))
+			{
+				string CmdLine = string.Format("\"{0}\" \"{1}\" --preload . --js-output=\"{1}.js\"", PackagerPath, FinalDataLocation);
+				RunAndLog(CmdEnv, PythonPath, CmdLine);
+			}
+		}
+
+		// copy the "Executable" to the package directory
+		string GameExe = Path.GetFileNameWithoutExtension(Params.ProjectGameExeFilename);
+		if (Params.ClientConfigsToBuild[0].ToString() != "Development")
+		{
+			GameExe += "-HTML5-" + Params.ClientConfigsToBuild[0].ToString();
+		}
+		GameExe += ".js";
 			
         // ensure the ue4game binary exists, if applicable
         string FullGameExePath = Path.Combine(Path.GetDirectoryName(Params.ProjectGameExeFilename), GameExe);
@@ -90,53 +92,49 @@ public class HTML5Platform : Platform
 	        throw new AutomationException(ErrorCodes.Error_MissingExecutable, "Stage Failed. Could not find application {0}. You may need to build the UE4 project with your target configuration and platform.", FullGameExePath);
         }
 
-        if (Path.Combine(Path.GetDirectoryName(Params.ProjectGameExeFilename), GameExe) != Path.Combine(PackagePath, GameExe))
-        {
-            File.Copy(Path.Combine(Path.GetDirectoryName(Params.ProjectGameExeFilename), GameExe), Path.Combine(PackagePath, GameExe), true);
-            File.Copy(Path.Combine(Path.GetDirectoryName(Params.ProjectGameExeFilename), GameExe) + ".mem", Path.Combine(PackagePath, GameExe) + ".mem", true);
+		if (Path.Combine(Path.GetDirectoryName(Params.ProjectGameExeFilename), GameExe) != Path.Combine(PackagePath, GameExe))
+		{
+			File.Copy(Path.Combine(Path.GetDirectoryName(Params.ProjectGameExeFilename), GameExe), Path.Combine(PackagePath, GameExe), true);
+			File.Copy(Path.Combine(Path.GetDirectoryName(Params.ProjectGameExeFilename), GameExe) + ".mem", Path.Combine(PackagePath, GameExe) + ".mem", true);
 			File.Copy(Path.Combine(Path.GetDirectoryName(Params.ProjectGameExeFilename), GameExe) + ".symbols", Path.Combine(PackagePath, GameExe) + ".symbols", true);
-        }
+		}
 
-        File.SetAttributes(Path.Combine(PackagePath, GameExe), FileAttributes.Normal);
-        File.SetAttributes(Path.Combine(PackagePath, GameExe) + ".mem", FileAttributes.Normal);
+		File.SetAttributes(Path.Combine(PackagePath, GameExe), FileAttributes.Normal);
+		File.SetAttributes(Path.Combine(PackagePath, GameExe) + ".mem", FileAttributes.Normal);
 		File.SetAttributes(Path.Combine(PackagePath, GameExe) + ".symbols", FileAttributes.Normal);
 
 
-        // put the HTML file to the package directory
-
-        bool UseExperimentalTemplate = false;
-        ConfigCache.GetBool("/Script/HTML5PlatformEditor.HTML5TargetSettings", "UseExperimentalTemplate", out UseExperimentalTemplate);
-        string TemplateFileName = UseExperimentalTemplate ? "GameX.html.template" : "Game.html.template";
-
-        string TemplateFile = Path.Combine(CombinePaths(CmdEnv.LocalRoot, "Engine"), "Build", "HTML5", TemplateFileName);
-        string OutputFile = Path.Combine(PackagePath, (Params.ClientConfigsToBuild[0].ToString() != "Development" ? (Params.ShortProjectName + "-HTML5-" + Params.ClientConfigsToBuild[0].ToString()) : Params.ShortProjectName)) + ".html";
+		// put the HTML file to the package directory
+		string TemplateFileName = "GameX.html.template"; 
+		string TemplateFile = Path.Combine(CombinePaths(CmdEnv.LocalRoot, "Engine"), "Build", "HTML5", TemplateFileName);
+		string OutputFile = Path.Combine(PackagePath, (Params.ClientConfigsToBuild[0].ToString() != "Development" ? (Params.ShortProjectName + "-HTML5-" + Params.ClientConfigsToBuild[0].ToString()) : Params.ShortProjectName)) + ".html";
 
-        // find Heap Size.
-        ulong HeapSize;
+		// find Heap Size.
+		ulong HeapSize;
 
-        int ConfigHeapSize = 0;
-        // Valuer set by Editor UI
-        var bGotHeapSize = ConfigCache.GetInt32("/Script/HTML5PlatformEditor.HTML5TargetSettings", "HeapSize" + Params.ClientConfigsToBuild[0].ToString(), out ConfigHeapSize);
+		int ConfigHeapSize = 0;
+		// Valuer set by Editor UI
+		var bGotHeapSize = ConfigCache.GetInt32("/Script/HTML5PlatformEditor.HTML5TargetSettings", "HeapSize" + Params.ClientConfigsToBuild[0].ToString(), out ConfigHeapSize);
 
-        // Fallback if the previous method failed
-        if (!bGotHeapSize && !ConfigCache.GetInt32("BuildSettings", "HeapSize" + Params.ClientConfigsToBuild[0].ToString(), out ConfigHeapSize)) // in Megs.
-        {
-            // we couldn't find a per config heap size, look for a common one.
-            if (!ConfigCache.GetInt32("BuildSettings", "HeapSize", out ConfigHeapSize))
-            {
-                ConfigHeapSize = Params.IsCodeBasedProject ? 1024 : 512;
-                Log("Could not find Heap Size setting in .ini for Client config {0}", Params.ClientConfigsToBuild[0].ToString());
-            }
-        }
+		// Fallback if the previous method failed
+		if (!bGotHeapSize && !ConfigCache.GetInt32("BuildSettings", "HeapSize" + Params.ClientConfigsToBuild[0].ToString(), out ConfigHeapSize)) // in Megs.
+		{
+			// we couldn't find a per config heap size, look for a common one.
+			if (!ConfigCache.GetInt32("BuildSettings", "HeapSize", out ConfigHeapSize))
+			{
+				ConfigHeapSize = Params.IsCodeBasedProject ? 1024 : 512;
+				Log("Could not find Heap Size setting in .ini for Client config {0}", Params.ClientConfigsToBuild[0].ToString());
+			}
+		}
 
-        HeapSize = (ulong)ConfigHeapSize * 1024L * 1024L; // convert to bytes.
-        Log("Setting Heap size to {0} Mb ", ConfigHeapSize);
+		HeapSize = (ulong)ConfigHeapSize * 1024L * 1024L; // convert to bytes.
+		Log("Setting Heap size to {0} Mb ", ConfigHeapSize);
 
 
-        GenerateFileFromTemplate(TemplateFile, OutputFile, Params.ShortProjectName, Params.ClientConfigsToBuild[0].ToString(), Params.StageCommandline, !Params.IsCodeBasedProject, HeapSize);
+		GenerateFileFromTemplate(TemplateFile, OutputFile, Params.ShortProjectName, Params.ClientConfigsToBuild[0].ToString(), Params.StageCommandline, !Params.IsCodeBasedProject, HeapSize);
 
-        string JSDir = Path.Combine(CombinePaths(CmdEnv.LocalRoot, "Engine"), "Build", "HTML5");
-        string OutDir = PackagePath;
+		string JSDir = Path.Combine(CombinePaths(CmdEnv.LocalRoot, "Engine"), "Build", "HTML5");
+		string OutDir = PackagePath;
 
 		// Gather utlity .js files and combine into one file
 		string[] UtilityJavaScriptFiles = Directory.GetFiles(JSDir, "*.js");
@@ -148,22 +146,23 @@ public class HTML5Platform : Platform
 			File.AppendAllText(DestinationFile, Data);
 		}
 
-		// Compress all files. 
-
+		// Compress all files. These are independent tasks which can be threaded. 
+		Task[] CompressionTasks = new Task[6];
 		//data file.
-		CompressFile(FinalDataLocation,FinalDataLocation + ".gz");
+		CompressionTasks[0] = Task.Factory.StartNew( () => CompressFile(FinalDataLocation, FinalDataLocation + ".gz"));		
 		// data file .js driver.
-		CompressFile(FinalDataLocation + ".js" , FinalDataLocation + ".js.gz");
+		CompressionTasks[1] = Task.Factory.StartNew( () => CompressFile(FinalDataLocation + ".js" , FinalDataLocation + ".js.gz"));
 		// main js.
-		CompressFile(Path.Combine(PackagePath, GameExe), Path.Combine(PackagePath, GameExe) + ".gz");
+		CompressionTasks[2] = Task.Factory.StartNew(() => CompressFile(Path.Combine(PackagePath, GameExe), Path.Combine(PackagePath, GameExe) + ".gz"));
 		// mem init file.
-		CompressFile(Path.Combine(PackagePath, GameExe) + ".mem", Path.Combine(PackagePath, GameExe) + ".mem.gz");
+		CompressionTasks[3] = Task.Factory.StartNew(() => CompressFile(Path.Combine(PackagePath, GameExe) + ".mem", Path.Combine(PackagePath, GameExe) + ".mem.gz"));
 		// symbols file.
-		CompressFile(Path.Combine(PackagePath, GameExe) + ".symbols", Path.Combine(PackagePath, GameExe) + ".symbols.gz");
+		CompressionTasks[4] = Task.Factory.StartNew(() => CompressFile(Path.Combine(PackagePath, GameExe) + ".symbols", Path.Combine(PackagePath, GameExe) + ".symbols.gz"));
 		// Utility 
-		CompressFile(OutDir + "/Utility.js", OutDir + "/Utility.js.gz");
+		CompressionTasks[5] = Task.Factory.StartNew(() => CompressFile(OutDir + "/Utility.js", OutDir + "/Utility.js.gz"));
 
-        PrintRunTime();
+		Task.WaitAll(CompressionTasks);
+		PrintRunTime();
 	}
 
 	void CompressFile(string Source, string Destination)
@@ -201,10 +200,10 @@ public class HTML5Platform : Platform
 		}
 	}
 
-    public override bool RequiresPackageToDeploy
-    {
-        get { return true; }
-    }
+	public override bool RequiresPackageToDeploy
+	{
+		get { return true; }
+	}
 
 
 	protected void GenerateFileFromTemplate(string InTemplateFile, string InOutputFile, string InGameName, string InGameConfiguration, string InArguments, bool IsContentOnly, ulong HeapSize)
@@ -242,7 +241,7 @@ public class HTML5Platform : Platform
 					for (int i = 0; i < Arguments.Length - 1; ++i)
 					{
 						ArgumentString += "'";
-			   		    ArgumentString += Arguments[i];
+						ArgumentString += Arguments[i];
 						ArgumentString += "'";
 						ArgumentString += ",' ',";
 					}
@@ -330,56 +329,17 @@ public class HTML5Platform : Platform
 				SC.ArchivedFiles.Add(PakFile, DestPak);
 			}	
 		}
-
 	}
 
 	public override ProcessResult RunClient(ERunOptions ClientRunFlags, string ClientApp, string ClientCmdLine, ProjectParams Params)
 	{
 		// look for browser
-		var ConfigCache = new UnrealBuildTool.ConfigCacheIni(UnrealTargetPlatform.HTML5, "Engine", Path.GetDirectoryName(Params.RawProjectPath), CombinePaths(CmdEnv.LocalRoot, "Engine"));
-		bool ok = false;
-		List<string> Devices;
-		string browserPath = "";
-		string DeviceName = Params.Device.Split('@')[1];
-		DeviceName = DeviceName.Substring(0, DeviceName.LastIndexOf(" on "));
-
-		if (ConfigCache.GetArray("/Script/HTML5PlatformEditor.HTML5SDKSettings", "DeviceMap", out Devices))
-		{
-			foreach (var Dev in Devices)
-			{
-				var Matched = Regex.Match(Dev, "\\(DeviceName=\"(.*)\",DevicePath=\\(FilePath=\"(.*)\"\\)\\)", RegexOptions.IgnoreCase);
-				if (Matched.Success && Matched.Groups[1].ToString() == DeviceName)
-				{
-					browserPath = Matched.Groups[2].ToString();
-					ok = true;
-					break;
-				}
-			}
-		}
-
-		if (!ok && HTML5SDKInfo.bAllowFallbackSDKSettings)
-		{
-			string DeviceSection;
-
-			if (Utils.IsRunningOnMono)
-			{
-				DeviceSection = "HTML5DevicesMac";
-			}
-			else
-			{
-				DeviceSection = "HTML5DevicesWindows";
-			}
-
-			ok = ConfigCache.GetString(DeviceSection, DeviceName, out browserPath);
-		}
-
-		if (!ok)
-		{
-			throw new System.Exception("Incorrect browser configuration in HTML5Engine.ini ");
-		}
+		string BrowserPath = Params.Device.Replace("HTML5@", "");
 
 		// open the webpage
 		Int32 ServerPort = 8000;
+
+		var ConfigCache = new UnrealBuildTool.ConfigCacheIni(UnrealTargetPlatform.HTML5, "Engine", Path.GetDirectoryName(Params.RawProjectPath), CombinePaths(CmdEnv.LocalRoot, "Engine"));
 		ConfigCache.GetInt32("/Script/HTML5PlatformEditor.HTML5TargetSettings", "DeployServerPort", out ServerPort);
 		string WorkingDirectory = Path.GetDirectoryName(ClientApp);
 		string url = Path.GetFileName(ClientApp) +".html";
@@ -408,7 +368,7 @@ public class HTML5Platform : Platform
 
 		// Check HTML5LaunchHelper source for command line args
 
-		var LowerBrowserPath = browserPath.ToLower();
+		var LowerBrowserPath = BrowserPath.ToLower();
 		var ProfileDirectory = Path.Combine(Utils.GetUserSettingDirectory(), "UE4_HTML5", "user");
 
 		string BrowserCommandline = url; 
@@ -422,13 +382,12 @@ public class HTML5Platform : Platform
 			BrowserCommandline += "  " +  String.Format("-no-remote -profile \"{0}\"", Path.Combine(ProfileDirectory, "firefox"));
 		}
 
-		string LauncherArguments = string.Format( " -Browser=\"{0}\" + -BrowserCommandLine=\"{1}\" -ServerPort=\"{2}\" -ServerRoot=\"{3}\" ", new object[] { browserPath,BrowserCommandline,ServerPort,WorkingDirectory });
+		string LauncherArguments = string.Format(" -Browser=\"{0}\" + -BrowserCommandLine=\"{1}\" -ServerPort=\"{2}\" -ServerRoot=\"{3}\" ", new object[] { BrowserPath, BrowserCommandline, ServerPort, WorkingDirectory });
 
 		var LaunchHelperPath = CombinePaths(CmdEnv.LocalRoot, "Engine/Binaries/DotNET/HTML5LaunchHelper.exe");
 		ProcessResult BrowserProcess = Run(LaunchHelperPath, LauncherArguments, null, ClientRunFlags | ERunOptions.NoWaitForExit);
 	
 		return BrowserProcess;
-
 	}
 
 	public override string GetCookPlatform(bool bDedicatedServer, bool bIsClientOnly, string CookFlavor)
@@ -436,27 +395,27 @@ public class HTML5Platform : Platform
 		return "HTML5";
 	}
 
-    public override string GetCookExtraCommandLine(ProjectParams Params)
-    {
-        return HTMLPakAutomation.CanCreateMapPaks(Params) ? " -GenerateDependenciesForMaps " : ""; 
-    }
+	public override string GetCookExtraCommandLine(ProjectParams Params)
+	{
+		return HTMLPakAutomation.CanCreateMapPaks(Params) ? " -GenerateDependenciesForMaps " : ""; 
+	}
 
-    public override List<string> GetCookExtraMaps()
-    {
-        var Maps = new List<string>();
-        Maps.Add("/Engine/Maps/Entry");
-        return Maps; 
-    }
+	public override List<string> GetCookExtraMaps()
+	{
+		var Maps = new List<string>();
+		Maps.Add("/Engine/Maps/Entry");
+		return Maps; 
+	}
 
 	public override bool DeployPakInternalLowerCaseFilenames()
 	{
 		return false;
 	}
 
-    public override PakType RequiresPak(ProjectParams Params)
-    {
-        return HTMLPakAutomation.CanCreateMapPaks(Params) ? PakType.Never : PakType.Always;
-    }
+	public override PakType RequiresPak(ProjectParams Params)
+	{
+		return HTMLPakAutomation.CanCreateMapPaks(Params) ? PakType.Never : PakType.Always;
+	}
 
 	public override string GetPlatformPakCommandLine()
 	{
diff --git a/Engine/Source/Programs/AutomationTool/HTML5/HTML5Platform.PakFiles.Automation.cs b/Engine/Source/Programs/AutomationTool/HTML5/HTML5Platform.PakFiles.Automation.cs
--- a/Engine/Source/Programs/AutomationTool/HTML5/HTML5Platform.PakFiles.Automation.cs
+++ b/Engine/Source/Programs/AutomationTool/HTML5/HTML5Platform.PakFiles.Automation.cs
@@ -12,338 +12,338 @@ using System.Diagnostics;
 
 public class HTMLPakAutomation
 {
-    ProjectParams Params;
-
-    DeploymentContext SC;
-
-    string PakOrderFileLocation;
-
-    Dictionary<string, object> DependencyJson;
-    
-    public static bool CanCreateMapPaks(ProjectParams Param)
-    {
-        bool UseAsyncLevelLoading = false;
-        var ConfigCache = new UnrealBuildTool.ConfigCacheIni(UnrealTargetPlatform.HTML5, "Engine", Path.GetDirectoryName(Param.RawProjectPath), CommandUtils.CombinePaths(CommandUtils.CmdEnv.LocalRoot, "Engine"));
-        ConfigCache.GetBool("/Script/HTML5PlatformEditor.HTML5TargetSettings", "UseAsyncLevelLoading", out UseAsyncLevelLoading);
-
-        if (Param.Run)
-            return false; 
-
-        return UseAsyncLevelLoading;
-    }
-
-    public HTMLPakAutomation(ProjectParams InParam, DeploymentContext InSC)
-    {
-        Params = InParam;
-        SC = InSC;
-
-        var PakOrderFileLocationBase = CommandUtils.CombinePaths(SC.ProjectRoot, "Build", SC.FinalCookPlatform, "FileOpenOrder");
-        PakOrderFileLocation = CommandUtils.CombinePaths(PakOrderFileLocationBase, "GameOpenOrder.log");
-        if (!CommandUtils.FileExists_NoExceptions(PakOrderFileLocation))
-        {
-            // Use a fall back, it doesn't matter if this file exists or not. GameOpenOrder.log is preferred however
-            PakOrderFileLocation = CommandUtils.CombinePaths(PakOrderFileLocationBase, "EditorOpenOrder.log");
-        }
-
-        string PakPath = Path.Combine(new string[] { Path.GetDirectoryName(Params.RawProjectPath), "Binaries", "HTML5", SC.ShortProjectName});
-        if (Directory.Exists(PakPath))
-        {
-            Directory.Delete(PakPath,true);
-        }
-
-        // read in the json file. 
-        string JsonFile = CommandUtils.CombinePaths(new string[] { SC.ProjectRoot, "Saved", "Cooked", "HTML5", Params.ShortProjectName, "MapDependencyGraph.json" });
-        string text = File.ReadAllText(JsonFile);
-
-        DependencyJson = fastJSON.JSON.Instance.ToObject<Dictionary<string, object>>(text);
-
-    }
-
-    /// <summary>
-    /// Create Emscripten Data package with Necessary Packages. 
-    /// 
-    /// </summary>
-    public void CreateEmscriptenDataPackage(string PackagePath, string FinalDataLocation)
-    {
-        // we need to operate in the root
-        using (new PushedDirectory(Path.Combine(PackagePath)))
-        {
-            string PythonPath = HTML5SDKInfo.PythonPath();
-            string PackagerPath = HTML5SDKInfo.EmscriptenPackager();
-
-            string CmdLine = string.Format("\"{0}\" \"{1}\" --preload {2}\\ --js-output=\"{1}.js\"", PackagerPath, FinalDataLocation, SC.ShortProjectName);
-            CommandUtils.RunAndLog(CommandUtils.CmdEnv, PythonPath, CmdLine);
-        }
-    }
-
-    /// <summary>
-    /// Create a Pak from Engine Contents. 
-    /// </summary>
-    public void CreateEnginePak()
-    {
-        string StagedDir = Path.Combine(Params.BaseStageDirectory, "HTML5");// CommandUtils.CombinePaths(SC.ProjectRoot, "Saved", "Cooked", "HTML5");
-    
-         string PakPath = Path.Combine( new string [] { Path.GetDirectoryName(Params.RawProjectPath), "Binaries", "HTML5", SC.ShortProjectName, "Content","Paks"});
-         if (!Directory.Exists(PakPath))
-         {
-             Directory.CreateDirectory(PakPath);
-         }
-
-         // find list of files in the Engine Dir. 
-         string[] files = Directory.GetFiles( StagedDir + "\\Engine\\", "*", SearchOption.AllDirectories);
-
-         Dictionary<string,string> UnrealPakResponseFile = new Dictionary<string,string>();
-
-         // create response file structure. 
-         foreach( string file in files )  
-         {
-             string MountPoint = file.Remove(0, StagedDir.Length + 1);
-             UnrealPakResponseFile[file] = "../../../" +  MountPoint; 
-         }
-
-        RunUnrealPak(UnrealPakResponseFile, PakPath + "\\EnginePak.pak", null, "", true);
-    }
-
-
-    /// <summary>
-    /// Create a pak from a minimum set of game files
-    /// </summary>
-    public void CreateGamePak()
-    {
-        // create a pak of game config and first level files. 
-
-        string StagedDir = Path.Combine(Params.BaseStageDirectory, "HTML5");// CommandUtils.CombinePaths(SC.ProjectRoot, "Saved", "Cooked", "HTML5");
-
-        string PakPath = Path.Combine(new string[] { Path.GetDirectoryName(Params.RawProjectPath), "Binaries", "HTML5", Params.ShortProjectName, "Content", "Paks" });
-        if (!Directory.Exists(PakPath))
-        {
-            Directory.CreateDirectory(PakPath);
-        }
-
-        List<string> AllFiles = new List<string>();
-        try
-        {
-            string[] SlateFiles = Directory.GetFiles(StagedDir + "\\" + Params.ShortProjectName + "\\Content\\Slate", "*", SearchOption.AllDirectories);
-            AllFiles.AddRange(SlateFiles);
-        }
-        catch ( Exception /*Ex*/)
-        {
-            // not found. 
-        }
-
-        try
-        {
-            string[] Resources = Directory.GetFiles(StagedDir + "\\" + Params.ShortProjectName + "\\Content\\Resources", "*", SearchOption.AllDirectories);
-            AllFiles.AddRange(Resources);
-        }
-        catch (Exception /*Ex*/)
-        {
-            // not found. 
-        }
-
-        string[] ConfigFiles = Directory.GetFiles(StagedDir + "\\" + Params.ShortProjectName + "\\Config", "*", SearchOption.AllDirectories);
-
-        AllFiles.AddRange(ConfigFiles);
-
-        string[] TopLevelFiles = Directory.GetFiles(StagedDir + "\\" + Params.ShortProjectName + "\\", "*", SearchOption.TopDirectoryOnly);
-
-        AllFiles.AddRange(TopLevelFiles);
-
-        Dictionary<string, string> UnrealPakResponseFile = new Dictionary<string, string>();
-
-        // create response file structure. 
-        foreach (string file in AllFiles)
-        {
-            string MountPoint = file.Remove(0, StagedDir.Length + 1);
-            UnrealPakResponseFile[file] = "../../../" + MountPoint;
-        }
-
-        RunUnrealPak(UnrealPakResponseFile, PakPath + "\\" + Params.ShortProjectName + ".pak", null, "", true);
-        
-    }
-
-    /// <summary>
-    /// Create ContentDirectory Pak.
-    /// </summary>
-    public void CreateContentDirectoryPak()
-    {
-        object[] Value = (object[])DependencyJson["ContentDirectoryAssets"];
-        List<string> ContentDirectory = Array.ConvertAll<object, string>(Value, Convert.ToString).ToList();
-        CreatePak("ContentDirectoryAssets", ContentDirectory);
-    }
-
-
-    /// <summary>
-    /// Create Maps Paks based on JSON file created during Cooking. (-GenerateDependenciesForMaps)
-    /// </summary>
-    public void CreateMapPak()
-    {
-        // read in the json file. 
-
-        object[] Value = (object[])DependencyJson["ContentDirectoryAssets"];
-        List<string> ContentDirectory = Array.ConvertAll<object, string>(Value, Convert.ToString).ToList();
-
-        foreach (string Pak in DependencyJson.Keys)
-        {
-            object[] ObjectList = (object[])DependencyJson[Pak];
-            List<string> Assets = Array.ConvertAll<object, string>(ObjectList, Convert.ToString).ToList();
-
-            // add itself.
-            Assets.Add(Pak);
-
-            if (Pak != "ContentDirectoryAssets")
-            {
-                // remove content directory assets. 
-                Assets = Assets.Except(ContentDirectory).ToList();
-                string[] Tokens = Pak.Split('/');
-                CreatePak(Tokens[Tokens.Length - 1], Assets);
-            }
-        }
-
-    }
-
-    /// <summary>
-    /// Create Delta Paks based on Level Transitions. 
-    /// </summary>
-    /// <param name="MapFrom"></param>
-    /// <param name="MapTo"></param>
-    public void CreateDeltaMapPaks(string MapFrom, string MapTo)
-    {
-        foreach (string NameKey in DependencyJson.Keys)
-        {
-            if (NameKey.Contains(MapFrom))
-            { 
-                MapFrom = NameKey;
-            }
-            if (NameKey.Contains(MapTo))
-            { 
-                MapTo = NameKey;
-            }
-        }
-
-        List<string> MapFromAssets = Array.ConvertAll<object, string>(DependencyJson[MapFrom] as object[], Convert.ToString).ToList();
-        MapFromAssets.Add(MapFrom);
-        List<string> MapToAssets = Array.ConvertAll<object, string>(DependencyJson[MapTo] as object[], Convert.ToString).ToList();
-        MapToAssets.Add(MapTo);
-
-        object[] Value = (object[])DependencyJson["ContentDirectoryAssets"];
-        List<string> ContentDirectory = Array.ConvertAll<object, string>(Value, Convert.ToString).ToList();
-
-        // Delta. 
-        List<string> DeltaAssets = MapToAssets.Except(MapFromAssets).Except(ContentDirectory).ToList();
-
-        // create a pak file for this delta. 
-        string[] Tokens_MapFrom = MapFrom.Split('/');
-        string[] Tokens_MapTo = MapTo.Split('/');
-
-        string DeltaPakName = Tokens_MapFrom[Tokens_MapFrom.Length - 1] + "_" + Tokens_MapTo[Tokens_MapTo.Length - 1];
-        CreatePak(DeltaPakName, DeltaAssets);
-    }
-
-    /// <summary>
-    /// Writes a pak response file to disk
-    /// </summary>
-    /// <param name="Filename"></param>
-    /// <param name="ResponseFile"></param>
-    private static void WritePakResponseFile(string Filename, Dictionary<string, string> ResponseFile, bool Compressed)
-    {
-        using (var Writer = new StreamWriter(Filename, false, new System.Text.UTF8Encoding(true)))
-        {
-            foreach (var Entry in ResponseFile)
-            {
-                string Line = String.Format("\"{0}\" \"{1}\"", Entry.Key, Entry.Value);
-                if (Compressed)
-                {
-                    Line += " -compress";
-                }
-                Writer.WriteLine(Line);
-            }
-        }
-    }
-
-    /// <summary>
-    /// Create a Pak file
-    /// </summary>
-    private void RunUnrealPak(Dictionary<string, string> UnrealPakResponseFile, string OutputLocation, string EncryptionKeys, string PlatformOptions, bool Compressed)
-    {
-        if (UnrealPakResponseFile.Count < 1)
-        {
-            return;
-        }
-        string PakName = Path.GetFileNameWithoutExtension(OutputLocation);
-        string UnrealPakResponseFileName = CommandUtils.CombinePaths(CommandUtils.CmdEnv.LogFolder, "PakList_" + PakName + ".txt");
-        WritePakResponseFile(UnrealPakResponseFileName, UnrealPakResponseFile, Compressed);
-
-        var UnrealPakExe = CommandUtils.CombinePaths(CommandUtils.CmdEnv.LocalRoot, "Engine/Binaries/Win64/UnrealPak.exe");
-
-        string CmdLine = CommandUtils.MakePathSafeToUseWithCommandLine(OutputLocation) + " -create=" + CommandUtils.MakePathSafeToUseWithCommandLine(UnrealPakResponseFileName);
-        if (!String.IsNullOrEmpty(EncryptionKeys))
-        {
-            CmdLine += " -sign=" + CommandUtils.MakePathSafeToUseWithCommandLine(EncryptionKeys);
-        }
-        if (GlobalCommandLine.Installed)
-        {
-            CmdLine += " -installed";
-        }
-        CmdLine += " -order=" + CommandUtils.MakePathSafeToUseWithCommandLine(PakOrderFileLocation);
-        if (GlobalCommandLine.UTF8Output)
-        {
-            CmdLine += " -UTF8Output";
-        }
-        CmdLine += PlatformOptions;
-        CommandUtils.RunAndLog(CommandUtils.CmdEnv, UnrealPakExe, CmdLine, Options: CommandUtils.ERunOptions.Default | CommandUtils.ERunOptions.UTF8Output);
-    }
-
-    /// <summary>
-    /// Create a Pak.
-    /// </summary>
-    /// <param name="MapFrom"></param>
-    /// <param name="MapTo"></param>
-    private void CreatePak(string Name, List<string> Assets)
-    {
-        string CookedDir = Path.Combine(Params.BaseStageDirectory, "HTML5");
-
-        string PakPath = Path.Combine(new string[] { Path.GetDirectoryName(Params.RawProjectPath), "Binaries", "HTML5", Params.ShortProjectName, "Content", "Paks" });
-        if (!Directory.Exists(PakPath))
-        {
-            Directory.CreateDirectory(PakPath);
-        }
-
-        Assets.Add(Name); 
-        List<string> Files = new List<string>();
-        foreach( string Asset in Assets)
-        {
-            string GameDir = "/Game/";
-            if (Asset.StartsWith(GameDir))
-            {
-                string PathOnDiskasset  = CommandUtils.CombinePaths(CookedDir, Params.ShortProjectName, "Content", Asset.Remove(0, GameDir.Length) + ".uasset");
-                string PathOnDiskmap    = CommandUtils.CombinePaths(CookedDir, Params.ShortProjectName, "Content", Asset.Remove(0, GameDir.Length)+ ".umap");
-
-                if (File.Exists(PathOnDiskmap))
-                {
-                    Files.Add(PathOnDiskmap);
-                }
-                else if (File.Exists(PathOnDiskasset))
-                {
-                    Files.Add(PathOnDiskasset);
-                }
-                else
-                {
-                    System.Console.WriteLine(Asset + " not found, skipping !!"); 
-                }
-            }
-        }
-
-        // we have all the files. 
-        Dictionary<string, string> UnrealPakResponseFile = new Dictionary<string, string>();
-
-        // create response file structure. 
-        foreach (string file in Files)
-        {
-            string MountPoint = file.Remove(0, CookedDir.Length + 1);
-            UnrealPakResponseFile[file] ="../../../" +  MountPoint;
-        }
-        RunUnrealPak(UnrealPakResponseFile, PakPath + "\\"+ Name + ".pak", null, "", true);
-    }
+	ProjectParams Params;
+
+	DeploymentContext SC;
+
+	string PakOrderFileLocation;
+
+	Dictionary<string, object> DependencyJson;
+	
+	public static bool CanCreateMapPaks(ProjectParams Param)
+	{
+		bool UseAsyncLevelLoading = false;
+		var ConfigCache = new UnrealBuildTool.ConfigCacheIni(UnrealTargetPlatform.HTML5, "Engine", Path.GetDirectoryName(Param.RawProjectPath), CommandUtils.CombinePaths(CommandUtils.CmdEnv.LocalRoot, "Engine"));
+		ConfigCache.GetBool("/Script/HTML5PlatformEditor.HTML5TargetSettings", "UseAsyncLevelLoading", out UseAsyncLevelLoading);
+
+		if (Param.Run)
+			return false; 
+
+		return UseAsyncLevelLoading;
+	}
+
+	public HTMLPakAutomation(ProjectParams InParam, DeploymentContext InSC)
+	{
+		Params = InParam;
+		SC = InSC;
+
+		var PakOrderFileLocationBase = CommandUtils.CombinePaths(SC.ProjectRoot, "Build", SC.FinalCookPlatform, "FileOpenOrder");
+		PakOrderFileLocation = CommandUtils.CombinePaths(PakOrderFileLocationBase, "GameOpenOrder.log");
+		if (!CommandUtils.FileExists_NoExceptions(PakOrderFileLocation))
+		{
+			// Use a fall back, it doesn't matter if this file exists or not. GameOpenOrder.log is preferred however
+			PakOrderFileLocation = CommandUtils.CombinePaths(PakOrderFileLocationBase, "EditorOpenOrder.log");
+		}
+
+		string PakPath = Path.Combine(new string[] { Path.GetDirectoryName(Params.RawProjectPath), "Binaries", "HTML5", SC.ShortProjectName});
+		if (Directory.Exists(PakPath))
+		{
+			Directory.Delete(PakPath,true);
+		}
+
+		// read in the json file. 
+		string JsonFile = CommandUtils.CombinePaths(new string[] { SC.ProjectRoot, "Saved", "Cooked", "HTML5", Params.ShortProjectName, "MapDependencyGraph.json" });
+		string text = File.ReadAllText(JsonFile);
+
+		DependencyJson = fastJSON.JSON.Instance.ToObject<Dictionary<string, object>>(text);
+
+	}
+
+	/// <summary>
+	/// Create Emscripten Data package with Necessary Packages. 
+	/// 
+	/// </summary>
+	public void CreateEmscriptenDataPackage(string PackagePath, string FinalDataLocation)
+	{
+		string PythonPath = HTML5SDKInfo.Python();
+		string PackagerPath = HTML5SDKInfo.EmscriptenPackager();
+
+		// we need to operate in the root
+		using (new PushedDirectory(Path.Combine(PackagePath)))
+		{
+			string CmdLine = string.Format("\"{0}\" \"{1}\" --preload {2}\\ --js-output=\"{1}.js\"", PackagerPath, FinalDataLocation, SC.ShortProjectName);
+			CommandUtils.RunAndLog(CommandUtils.CmdEnv, PythonPath, CmdLine);
+		}
+	}
+
+	/// <summary>
+	/// Create a Pak from Engine Contents. 
+	/// </summary>
+	public void CreateEnginePak()
+	{
+		string StagedDir = Path.Combine(Params.BaseStageDirectory, "HTML5");// CommandUtils.CombinePaths(SC.ProjectRoot, "Saved", "Cooked", "HTML5");
+	
+		 string PakPath = Path.Combine( new string [] { Path.GetDirectoryName(Params.RawProjectPath), "Binaries", "HTML5", SC.ShortProjectName, "Content","Paks"});
+		 if (!Directory.Exists(PakPath))
+		 {
+			 Directory.CreateDirectory(PakPath);
+		 }
+
+		 // find list of files in the Engine Dir. 
+		 string[] files = Directory.GetFiles( StagedDir + "\\Engine\\", "*", SearchOption.AllDirectories);
+
+		 Dictionary<string,string> UnrealPakResponseFile = new Dictionary<string,string>();
+
+		 // create response file structure. 
+		 foreach( string file in files )  
+		 {
+			 string MountPoint = file.Remove(0, StagedDir.Length + 1);
+			 UnrealPakResponseFile[file] = "../../../" +  MountPoint; 
+		 }
+
+		RunUnrealPak(UnrealPakResponseFile, PakPath + "\\EnginePak.pak", null, "", true);
+	}
+
+
+	/// <summary>
+	/// Create a pak from a minimum set of game files
+	/// </summary>
+	public void CreateGamePak()
+	{
+		// create a pak of game config and first level files. 
+
+		string StagedDir = Path.Combine(Params.BaseStageDirectory, "HTML5");// CommandUtils.CombinePaths(SC.ProjectRoot, "Saved", "Cooked", "HTML5");
+
+		string PakPath = Path.Combine(new string[] { Path.GetDirectoryName(Params.RawProjectPath), "Binaries", "HTML5", Params.ShortProjectName, "Content", "Paks" });
+		if (!Directory.Exists(PakPath))
+		{
+			Directory.CreateDirectory(PakPath);
+		}
+
+		List<string> AllFiles = new List<string>();
+		try
+		{
+			string[] SlateFiles = Directory.GetFiles(StagedDir + "\\" + Params.ShortProjectName + "\\Content\\Slate", "*", SearchOption.AllDirectories);
+			AllFiles.AddRange(SlateFiles);
+		}
+		catch ( Exception /*Ex*/)
+		{
+			// not found. 
+		}
+
+		try
+		{
+			string[] Resources = Directory.GetFiles(StagedDir + "\\" + Params.ShortProjectName + "\\Content\\Resources", "*", SearchOption.AllDirectories);
+			AllFiles.AddRange(Resources);
+		}
+		catch (Exception /*Ex*/)
+		{
+			// not found. 
+		}
+
+		string[] ConfigFiles = Directory.GetFiles(StagedDir + "\\" + Params.ShortProjectName + "\\Config", "*", SearchOption.AllDirectories);
+
+		AllFiles.AddRange(ConfigFiles);
+
+		string[] TopLevelFiles = Directory.GetFiles(StagedDir + "\\" + Params.ShortProjectName + "\\", "*", SearchOption.TopDirectoryOnly);
+
+		AllFiles.AddRange(TopLevelFiles);
+
+		Dictionary<string, string> UnrealPakResponseFile = new Dictionary<string, string>();
+
+		// create response file structure. 
+		foreach (string file in AllFiles)
+		{
+			string MountPoint = file.Remove(0, StagedDir.Length + 1);
+			UnrealPakResponseFile[file] = "../../../" + MountPoint;
+		}
+
+		RunUnrealPak(UnrealPakResponseFile, PakPath + "\\" + Params.ShortProjectName + ".pak", null, "", true);
+		
+	}
+
+	/// <summary>
+	/// Create ContentDirectory Pak.
+	/// </summary>
+	public void CreateContentDirectoryPak()
+	{
+		object[] Value = (object[])DependencyJson["ContentDirectoryAssets"];
+		List<string> ContentDirectory = Array.ConvertAll<object, string>(Value, Convert.ToString).ToList();
+		CreatePak("ContentDirectoryAssets", ContentDirectory);
+	}
+
+
+	/// <summary>
+	/// Create Maps Paks based on JSON file created during Cooking. (-GenerateDependenciesForMaps)
+	/// </summary>
+	public void CreateMapPak()
+	{
+		// read in the json file. 
+
+		object[] Value = (object[])DependencyJson["ContentDirectoryAssets"];
+		List<string> ContentDirectory = Array.ConvertAll<object, string>(Value, Convert.ToString).ToList();
+
+		foreach (string Pak in DependencyJson.Keys)
+		{
+			object[] ObjectList = (object[])DependencyJson[Pak];
+			List<string> Assets = Array.ConvertAll<object, string>(ObjectList, Convert.ToString).ToList();
+
+			// add itself.
+			Assets.Add(Pak);
+
+			if (Pak != "ContentDirectoryAssets")
+			{
+				// remove content directory assets. 
+				Assets = Assets.Except(ContentDirectory).ToList();
+				string[] Tokens = Pak.Split('/');
+				CreatePak(Tokens[Tokens.Length - 1], Assets);
+			}
+		}
+
+	}
+
+	/// <summary>
+	/// Create Delta Paks based on Level Transitions. 
+	/// </summary>
+	/// <param name="MapFrom"></param>
+	/// <param name="MapTo"></param>
+	public void CreateDeltaMapPaks(string MapFrom, string MapTo)
+	{
+		foreach (string NameKey in DependencyJson.Keys)
+		{
+			if (NameKey.Contains(MapFrom))
+			{ 
+				MapFrom = NameKey;
+			}
+			if (NameKey.Contains(MapTo))
+			{ 
+				MapTo = NameKey;
+			}
+		}
+
+		List<string> MapFromAssets = Array.ConvertAll<object, string>(DependencyJson[MapFrom] as object[], Convert.ToString).ToList();
+		MapFromAssets.Add(MapFrom);
+		List<string> MapToAssets = Array.ConvertAll<object, string>(DependencyJson[MapTo] as object[], Convert.ToString).ToList();
+		MapToAssets.Add(MapTo);
+
+		object[] Value = (object[])DependencyJson["ContentDirectoryAssets"];
+		List<string> ContentDirectory = Array.ConvertAll<object, string>(Value, Convert.ToString).ToList();
+
+		// Delta. 
+		List<string> DeltaAssets = MapToAssets.Except(MapFromAssets).Except(ContentDirectory).ToList();
+
+		// create a pak file for this delta. 
+		string[] Tokens_MapFrom = MapFrom.Split('/');
+		string[] Tokens_MapTo = MapTo.Split('/');
+
+		string DeltaPakName = Tokens_MapFrom[Tokens_MapFrom.Length - 1] + "_" + Tokens_MapTo[Tokens_MapTo.Length - 1];
+		CreatePak(DeltaPakName, DeltaAssets);
+	}
+
+	/// <summary>
+	/// Writes a pak response file to disk
+	/// </summary>
+	/// <param name="Filename"></param>
+	/// <param name="ResponseFile"></param>
+	private static void WritePakResponseFile(string Filename, Dictionary<string, string> ResponseFile, bool Compressed)
+	{
+		using (var Writer = new StreamWriter(Filename, false, new System.Text.UTF8Encoding(true)))
+		{
+			foreach (var Entry in ResponseFile)
+			{
+				string Line = String.Format("\"{0}\" \"{1}\"", Entry.Key, Entry.Value);
+				if (Compressed)
+				{
+					Line += " -compress";
+				}
+				Writer.WriteLine(Line);
+			}
+		}
+	}
+
+	/// <summary>
+	/// Create a Pak file
+	/// </summary>
+	private void RunUnrealPak(Dictionary<string, string> UnrealPakResponseFile, string OutputLocation, string EncryptionKeys, string PlatformOptions, bool Compressed)
+	{
+		if (UnrealPakResponseFile.Count < 1)
+		{
+			return;
+		}
+		string PakName = Path.GetFileNameWithoutExtension(OutputLocation);
+		string UnrealPakResponseFileName = CommandUtils.CombinePaths(CommandUtils.CmdEnv.LogFolder, "PakList_" + PakName + ".txt");
+		WritePakResponseFile(UnrealPakResponseFileName, UnrealPakResponseFile, Compressed);
+
+		var UnrealPakExe = CommandUtils.CombinePaths(CommandUtils.CmdEnv.LocalRoot, "Engine/Binaries/Win64/UnrealPak.exe");
+
+		string CmdLine = CommandUtils.MakePathSafeToUseWithCommandLine(OutputLocation) + " -create=" + CommandUtils.MakePathSafeToUseWithCommandLine(UnrealPakResponseFileName);
+		if (!String.IsNullOrEmpty(EncryptionKeys))
+		{
+			CmdLine += " -sign=" + CommandUtils.MakePathSafeToUseWithCommandLine(EncryptionKeys);
+		}
+		if (GlobalCommandLine.Installed)
+		{
+			CmdLine += " -installed";
+		}
+		CmdLine += " -order=" + CommandUtils.MakePathSafeToUseWithCommandLine(PakOrderFileLocation);
+		if (GlobalCommandLine.UTF8Output)
+		{
+			CmdLine += " -UTF8Output";
+		}
+		CmdLine += PlatformOptions;
+		CommandUtils.RunAndLog(CommandUtils.CmdEnv, UnrealPakExe, CmdLine, Options: CommandUtils.ERunOptions.Default | CommandUtils.ERunOptions.UTF8Output);
+	}
+
+	/// <summary>
+	/// Create a Pak.
+	/// </summary>
+	/// <param name="MapFrom"></param>
+	/// <param name="MapTo"></param>
+	private void CreatePak(string Name, List<string> Assets)
+	{
+		string CookedDir = Path.Combine(Params.BaseStageDirectory, "HTML5");
+
+		string PakPath = Path.Combine(new string[] { Path.GetDirectoryName(Params.RawProjectPath), "Binaries", "HTML5", Params.ShortProjectName, "Content", "Paks" });
+		if (!Directory.Exists(PakPath))
+		{
+			Directory.CreateDirectory(PakPath);
+		}
+
+		Assets.Add(Name); 
+		List<string> Files = new List<string>();
+		foreach( string Asset in Assets)
+		{
+			string GameDir = "/Game/";
+			if (Asset.StartsWith(GameDir))
+			{
+				string PathOnDiskasset  = CommandUtils.CombinePaths(CookedDir, Params.ShortProjectName, "Content", Asset.Remove(0, GameDir.Length) + ".uasset");
+				string PathOnDiskmap    = CommandUtils.CombinePaths(CookedDir, Params.ShortProjectName, "Content", Asset.Remove(0, GameDir.Length)+ ".umap");
+
+				if (File.Exists(PathOnDiskmap))
+				{
+					Files.Add(PathOnDiskmap);
+				}
+				else if (File.Exists(PathOnDiskasset))
+				{
+					Files.Add(PathOnDiskasset);
+				}
+				else
+				{
+					System.Console.WriteLine(Asset + " not found, skipping !!"); 
+				}
+			}
+		}
+
+		// we have all the files. 
+		Dictionary<string, string> UnrealPakResponseFile = new Dictionary<string, string>();
+
+		// create response file structure. 
+		foreach (string file in Files)
+		{
+			string MountPoint = file.Remove(0, CookedDir.Length + 1);
+			UnrealPakResponseFile[file] ="../../../" +  MountPoint;
+		}
+		RunUnrealPak(UnrealPakResponseFile, PakPath + "\\"+ Name + ".pak", null, "", true);
+	}
 
 }
\ No newline at end of file
diff --git a/Engine/Source/Programs/UnrealBuildTool/HTML5/HTML5SDKInfo.cs b/Engine/Source/Programs/UnrealBuildTool/HTML5/HTML5SDKInfo.cs
--- a/Engine/Source/Programs/UnrealBuildTool/HTML5/HTML5SDKInfo.cs
+++ b/Engine/Source/Programs/UnrealBuildTool/HTML5/HTML5SDKInfo.cs
@@ -9,417 +9,152 @@ using System.Text.RegularExpressions;
 
 namespace UnrealBuildTool
 {
-    public class HTML5SDKInfo
-    {
-        static string EmscriptenSettingsPath = Path.Combine(Environment.GetFolderPath(Environment.SpecialFolder.UserProfile), ".emscripten");
-		static ConfigCacheIni ConfigCache = null;
-		public const bool bAllowFallbackSDKSettings = true;
+	public class HTML5SDKInfo
+	{
+		static string SDKBase { get { return Path.GetFullPath(Path.Combine(new string[] { BuildConfiguration.RelativeEnginePath,  "Source" , "ThirdParty", "HTML5", "emsdk" })); } }
+		static string SDKVersion = "1.30.0";
+		static string EMSCRIPTEN_ROOT { get { return Path.Combine(SDKBase, "emscripten", SDKVersion); } }
 
-        static Dictionary<string, string> ReadEmscriptenSettings()
-        {
-            // Check HTML5ToolChain.cs for duplicate
-            if (!System.IO.File.Exists(EmscriptenSettingsPath))
-            {
-                return new Dictionary<string, string>();
-            }
+		static string PYTHON_WIN { get { return Path.Combine(SDKBase, "Win64","python" , "2.7.5.3_64bit", "python.exe"); } }
+		// python is default installed on mac. 
+		static string PYTHON_MAC { get { return Path.Combine( "/usr/bin/python"); } }
 
-            Dictionary<string, string> Settings = new Dictionary<string, string>();
-            System.IO.StreamReader SettingFile = new System.IO.StreamReader(EmscriptenSettingsPath);
-            string EMLine = null;
-            while ((EMLine = SettingFile.ReadLine()) != null)
-            {
-                EMLine = EMLine.Split('#')[0];
-                string Pattern1 = @"(\w*)\s*=\s*['\[]?([^'\]\r\n]*)['\]]?";
-                Regex Rgx = new Regex(Pattern1, RegexOptions.IgnoreCase);
-                MatchCollection Matches = Rgx.Matches(EMLine);
-                foreach (Match Matched in Matches)
-                {
-                    if (Matched.Groups.Count == 3 && Matched.Groups[2].ToString() != "")
-                    {
-                        Settings[Matched.Groups[1].ToString()] = Matched.Groups[2].ToString();
-                    }
-                }
-            }
-
-            return Settings;
-        }
-
-		public static bool ParseProjectPath(string[] Arguments, out string ProjectPath)
+		static string PYTHON
 		{
-			string TargetName = null;
-			foreach (var Argument in Arguments)
+			get
 			{
-				// Treat any platform names as such
-				UnrealTargetPlatform ParsedPlatform = UEBuildPlatform.ConvertStringToPlatform(Argument);
-				if (ParsedPlatform == UnrealTargetPlatform.Unknown)
+				switch (BuildHostPlatform.Current.Platform)
 				{
-					var ArgUpper = Argument.ToUpperInvariant();
-					switch (ArgUpper)
-					{
-						// Skip any build configurations 
-						case "DEBUG":
-						case "DEBUGGAME":
-						case "DEVELOPMENT":
-						case "SHIPPING":
-						case "TEST":
-							break;
-						default:
-							// Ignore any args except -project
-							if (!ArgUpper.StartsWith("-") || ArgUpper.StartsWith("-PROJECT="))
-							{
-								if (ArgUpper.StartsWith("-PROJECT="))
-								{
-									ArgUpper = ArgUpper.Remove(0, 9).Trim();
-								}
-								// If running Rocket, the PossibleTargetName could contain a path
-								TargetName = ArgUpper;
-
-								// If a project file was not specified see if we can find one
-								string CheckProjectFile = UProjectInfo.GetProjectForTarget(TargetName);
-								if (string.IsNullOrEmpty(CheckProjectFile))
-								{
-									CheckProjectFile = UProjectInfo.GetProjectFilePath(TargetName);
-								}
-								if (string.IsNullOrEmpty(CheckProjectFile) == false)
-								{
-									Log.TraceVerbose("Found project file for {0} - {1}", TargetName, CheckProjectFile);
-									string NewProjectFilename = CheckProjectFile;
-									if (Path.IsPathRooted(NewProjectFilename) == false)
-									{
-										NewProjectFilename = Path.GetFullPath(NewProjectFilename);
-									}
-
-									NewProjectFilename = NewProjectFilename.Replace("\\", "/");
-									ProjectPath = Path.GetDirectoryName(NewProjectFilename);
-									return true;
-								} 
-
-								if ((TargetName.Contains("/") || TargetName.Contains("\\")) && TargetName.Contains(".UPROJECT"))
-								{
-									// Parse off the path
-									ProjectPath = Path.GetDirectoryName(TargetName);
-									return true;
-								}
-							}
-						break;
-					}
+					case UnrealTargetPlatform.Win64:
+						return PYTHON_WIN;
+					case UnrealTargetPlatform.Mac:
+						return PYTHON_MAC;
+					default:
+						return "";
 				}
 			}
-
-			ProjectPath = "";
-			return false;
 		}
 
-		private class SDKVersion : IComparable<SDKVersion>
-		{
-			public System.Version Version;
-
-			public string Directory;
-
-			public static bool operator ==(SDKVersion LHS, SDKVersion RHS)
-			{
-				return LHS.Version == RHS.Version;
-			}
-
-			public static bool operator !=(SDKVersion LHS, SDKVersion RHS)
-			{
-				return LHS.Version != RHS.Version;
-			}
+		static string LLVM_ROOT_WIN { get { return Path.Combine(SDKBase, "Win64", "clang", "e" + SDKVersion + "_64bit"); } }
+		static string LLVM_ROOT_MAC { get { return Path.Combine(SDKBase, "Mac", "clang", "e" + SDKVersion + "_64bit"); } }
 
-			// Override the Object.Equals(object o) method:
-			public override bool Equals(object o)
+		static string LLVM_ROOT
+		{
+			get
 			{
-				try
+				switch (BuildHostPlatform.Current.Platform)
 				{
-					return (bool)(this == (SDKVersion)o);
+					case UnrealTargetPlatform.Win64:
+						return LLVM_ROOT_WIN;
+					case UnrealTargetPlatform.Mac:
+						return LLVM_ROOT_MAC;
+					default:
+						return "";
 				}
-				catch
-				{
-					return false;
-				}
-			}
-
-			// Override the Object.GetHashCode() method:
-			public override int GetHashCode()
-			{
-				return Version.GetHashCode();
-			}
-
-			public int CompareTo(SDKVersion other)
-			{
-				return this.Version.CompareTo(other.Version);
 			}
 		}
 
-		static List<SDKVersion> GetInstalledVersions(string RootPath)
-		{
-			List<SDKVersion> Versions = new List<SDKVersion>();
-			SDKVersion Ver = null;
-			if (!System.IO.Directory.Exists(Path.Combine(RootPath, "emscripten")))
-			{
-				return Versions;
-			}
+		static string NODE_JS_WIN { get { return Path.Combine(SDKBase, "Win64", "node", "0.12.2_64bit", "node.exe"); } }
+		static string NODE_JS_MAC { get { return Path.Combine(SDKBase, "Mac", "node", "0.12.2_64bit", "bin", "node"); } }
 
-			string[] Directories = Directory.GetDirectories(Path.Combine(RootPath, "emscripten"), "*", SearchOption.TopDirectoryOnly);
-			foreach (var Dir in Directories)
+		static string NODE_JS
+		{
+			get
 			{
-				var VersionFilePath = Path.Combine(Dir, "emscripten-version.txt");
-				if (System.IO.File.Exists(VersionFilePath))
+				switch (BuildHostPlatform.Current.Platform)
 				{
-					Ver = new SDKVersion();
-					Ver.Version = System.Version.Parse(ReadVersionFile(VersionFilePath));
-					Ver.Directory = Dir;
-					Versions.Add(Ver);
+					case UnrealTargetPlatform.Win64:
+						return NODE_JS_WIN;
+					case UnrealTargetPlatform.Mac:
+						return NODE_JS_MAC;
+					default:
+						return "";
 				}
 			}
-
-			Versions.Sort();
-			return Versions;
 		}
 
-		static void EnsureConfigCacheIsReady()
+		public static string  SetupEmscriptenTemp()
 		{
-			if (ConfigCache == null)
+			string HTML5Intermediatory = Path.GetFullPath(Path.Combine(BuildConfiguration.RelativeEnginePath, BuildConfiguration.PlatformIntermediateFolder));
+			string TempPath= Path.Combine(HTML5Intermediatory, "EmscriptenTemp");
+			if (!Directory.Exists(TempPath))
 			{
-				var CmdLine = Environment.GetCommandLineArgs();
-				string ProjectPath = UnrealBuildTool.GetUProjectPath();
-				if (String.IsNullOrEmpty(ProjectPath))
-				{ 
-					ParseProjectPath(CmdLine, out ProjectPath);
-				}
-				ConfigCache = new ConfigCacheIni(UnrealTargetPlatform.HTML5, "Engine", ProjectPath);
-				// always pick from the Engine the root directory and NOT the staged engine directory. 
-				// This breaks parse order for user & game config inis so disabled but pulling from the staged directory may still be a problem 
-				//string IniFile = Path.GetFullPath(Path.GetDirectoryName(UnrealBuildTool.GetUBTPath()) + "/../../") + "Config/HTML5/HTML5Engine.ini";
-				//ConfigCache.ParseIniFile(IniFile);
+				Directory.CreateDirectory(TempPath);
 			}
+			return TempPath;
 		}
 
-        public static string EmscriptenSDKPath()
-        {
-			EnsureConfigCacheIsReady();
-			// Search the ini file for Emscripten root first
-			string EmscriptenRoot;
-			string EMCCPath;
-			if (ConfigCache.GetString("/Script/HTML5PlatformEditor.HTML5SDKSettings", "EmscriptenRoot", out EmscriptenRoot) || Environment.GetEnvironmentVariable("EMSCRIPTEN") != null)
-			{
-				string SDKPathString = null;
-				string VersionString = null;
-				var Index = EmscriptenRoot.IndexOf("SDKPath=");
-				if (Index != -1)
-				{
-					Index+=9;
-					var Index2 = EmscriptenRoot.IndexOf("\"", Index);
-					if (Index2 != -1)
-					{
-						SDKPathString = EmscriptenRoot.Substring(Index, Index2 - Index);
-					}
-				}
-				Index = EmscriptenRoot.IndexOf("EmscriptenVersion=");
-				if (Index != -1)
-				{
-					Index += 19;
-					var Index2 = EmscriptenRoot.IndexOf("\"", Index);
-					if (Index2 != -1)
-					{
-						VersionString = EmscriptenRoot.Substring(Index, Index2 - Index);
-					}
-				}
-
-				if (string.IsNullOrEmpty(SDKPathString) && Environment.GetEnvironmentVariable("EMSCRIPTEN") != null)
-				{
-					VersionString = "-1.-1.-1";
-					SDKPathString = Environment.GetEnvironmentVariable("EMSCRIPTEN");
-				}
-
-				if (!string.IsNullOrEmpty(SDKPathString) && !string.IsNullOrEmpty(VersionString))
-				{
-					var SDKVersions = GetInstalledVersions(SDKPathString);
-
-					if (VersionString == "-1.-1.-1" && SDKVersions.Count > 0)
-					{
-						var Ver = SDKVersions[SDKVersions.Count - 1];
-						if (System.IO.Directory.Exists(Ver.Directory))
-						{
-							return Ver.Directory;
-						}
-					}
-					else if (SDKVersions.Count > 0)
-					{
-						var RequiredVersion = System.Version.Parse(VersionString);
-						foreach (var i in SDKVersions)
-						{
-							if (i.Version == RequiredVersion && System.IO.Directory.Exists(i.Directory))
-							{
-								return i.Directory;
-							}
-						}
-					}
-				}
-			}
-
-
-			// Old Method
-			if (bAllowFallbackSDKSettings)
-			{
-				bool ok = ConfigCache.GetString("HTML5SDKPaths", "Emscripten", out EMCCPath);
-				if (ok && System.IO.Directory.Exists(EMCCPath))
-					return EMCCPath;
-
-				// Older method used platform name
-				string PlatformName = "";
-				if (!Utils.IsRunningOnMono)
-				{
-					PlatformName = "Windows";
-				}
-				else
-				{
-					PlatformName = "Mac";
-				}
-				ok = ConfigCache.GetString("HTML5SDKPaths", PlatformName, out EMCCPath);
-
-				if (ok && System.IO.Directory.Exists(EMCCPath))
-					return EMCCPath;
-
-				// try to find SDK Location from env.
-				if (Environment.GetEnvironmentVariable("EMSCRIPTEN") != null
-					&& System.IO.Directory.Exists(Environment.GetEnvironmentVariable("EMSCRIPTEN"))
-					)
-				{
-					return Environment.GetEnvironmentVariable("EMSCRIPTEN");
-				}
-			}
-
-            return ""; 
-        }
-
-        public  static string PythonPath()
-        {
-			// find Python.
-			EnsureConfigCacheIsReady();
-
-			string PythonPath;
-			if (ConfigCache.GetPath("/Script/HTML5PlatformEditor.HTML5SDKSettings", "Python", out PythonPath) && System.IO.File.Exists(PythonPath))
-				return PythonPath;
-
-			if (bAllowFallbackSDKSettings)
-			{
-				// Check the ini first. Check for Python="path"
-				bool ok = ConfigCache.GetString("HTML5SDKPaths", "Python", out PythonPath);
-				if (ok && System.IO.File.Exists(PythonPath))
-					return PythonPath;
-
-				var EmscriptenSettings = ReadEmscriptenSettings();
-
-				// check emscripten generated config file. 
-				if (EmscriptenSettings.ContainsKey("PYTHON")
-					&& System.IO.File.Exists(EmscriptenSettings["PYTHON"])
-					)
-				{
-					return EmscriptenSettings["PYTHON"];
-				}
-
-				// It might be setup as a env variable. 
-				if (Environment.GetEnvironmentVariable("PYTHON") != null
-					&&
-					System.IO.File.Exists(Environment.GetEnvironmentVariable("PYTHON"))
-					)
-				{
-					return Environment.GetEnvironmentVariable("PYTHON");
-				}
-			}
-
-            // it might just be on path. 
-            ProcessStartInfo startInfo = new ProcessStartInfo();
-            Process PythonProcess = new Process();
-
-            startInfo.CreateNoWindow = true;
-            startInfo.RedirectStandardOutput = true;
-            startInfo.RedirectStandardInput = true;
-
-            startInfo.Arguments = "python";
-			startInfo.UseShellExecute = false;
+		public static string SetUpEmscriptenConfigFile()
+		{
+			string UserFolder = Environment.GetFolderPath(Environment.SpecialFolder.UserProfile); 
+			string ConfigFile = Path.Combine(UserFolder, ".emscripten");
 
-			if (!Utils.IsRunningOnMono) 
-			{
-				startInfo.FileName = "C:\\Windows\\System32\\where.exe";
-			}
-			else 
+			if (!File.Exists(ConfigFile) || !File.ReadAllText(ConfigFile).Contains(SDKVersion))
 			{
-				startInfo.FileName = "whereis";
+				var ConfigString = String.Join(
+								Environment.NewLine,
+								"import os",
+								"SPIDERMONKEY_ENGINE = ''",
+								"LLVM_ROOT='" + LLVM_ROOT +"'",
+								"NODE_JS= '" + NODE_JS +"'",
+								"PYTHON= '" + PYTHON +"'",
+								"EMSCRIPTEN_NATIVE_OPTIMIZER='"+Path.Combine(LLVM_ROOT,"optimizer.exe") +"'",
+								"EMSCRIPTEN_ROOT= '" + EMSCRIPTEN_ROOT + "'",
+								"TEMP_DIR= '" + SetupEmscriptenTemp() + "'",
+								"COMPILER_ENGINE = NODE_JS",
+								"JS_ENGINES = [NODE_JS]",
+								"V8_ENGINE = ''",
+								"GENERATEDBYUE4='"+SDKVersion +"'"
+								);
+				File.WriteAllText(ConfigFile, ConfigString.Replace("\\","/"));
 			}
+			return ConfigFile;
+		}
 
-
-			PythonProcess.StartInfo = startInfo;
-            PythonProcess.Start();
-            PythonProcess.WaitForExit();
-
-            string Apps = PythonProcess.StandardOutput.ReadToEnd();
-            Apps = Apps.Replace('\r', '\n');
-            string[] locations = Apps.Split('\n');
-
-            return locations[0];
-        }
-
-		public static string OverriddenSDKVersion()
+		public static string EmscriptenBase()
 		{
-			EnsureConfigCacheIsReady();
-
-			// Check the ini first. Check for ForceSDKVersion="x.xx.x"
-			string ForcedSDKVersion;
-			bool ok = ConfigCache.GetString("HTML5SDKPaths", "ForceSDKVersion", out ForcedSDKVersion);
-			if (ok)
-				return ForcedSDKVersion;
-
-			return "";
+			return EMSCRIPTEN_ROOT; 
+		}
+	
+		public static string EmscriptenVersion()
+		{
+			return SDKVersion; 
 		}
 
-        public static string EmscriptenPackager()
-        {
-            return Path.Combine(EmscriptenSDKPath(), "tools", "file_packager.py"); 
-        }
-
-        public static string EmscriptenCompiler()
-        {
-            return Path.Combine(EmscriptenSDKPath(), "emcc");
-        }
-
-        public static bool IsSDKInstalled()
-        {
-			return File.Exists(GetVersionInfoPath());
-        }
-
-        public static bool IsPythonInstalled()
-        {
-            if (String.IsNullOrEmpty(PythonPath()))
-                return false;
-            return true;
-        }
+		public  static string Python()
+		{
+			return PYTHON; 
+		}
 
-		public static bool IsSDKVersionOverridden()
+		public static string EmscriptenPackager()
 		{
-			if (String.IsNullOrEmpty(OverriddenSDKVersion()))
-				return false;
-			return true;
+			return Path.Combine(EmscriptenBase(), "tools","file_packager.py"); 
 		}
 
-		public static string ReadVersionFile(string path)
+		public static string EmscriptenCompiler()
 		{
-			string VersionInfo = File.ReadAllText(path);
-			VersionInfo = VersionInfo.Trim();
-			return VersionInfo; 
+			return Path.Combine(EmscriptenBase(), "emcc");
 		}
 
-        public static string EmscriptenVersion()
-        {
-			return ReadVersionFile(GetVersionInfoPath());
-        }
+		public static bool IsSDKInstalled()
+		{
+			bool SDK = File.Exists(GetVersionInfoPath());
+
+			switch (BuildHostPlatform.Current.Platform)
+			{
+				case UnrealTargetPlatform.Win64: 
+						return SDK && Directory.Exists(Path.Combine(SDKBase,"Win64"));
+				case UnrealTargetPlatform.Mac: 
+						return SDK && Directory.Exists(Path.Combine(SDKBase,"Mac")) && File.Exists(PYTHON);
+				default:
+						return false; 
+			}
+		}
 
 		static string GetVersionInfoPath()
 		{
-            string BaseSDKPath = HTML5SDKInfo.EmscriptenSDKPath(); 
-            return Path.Combine(BaseSDKPath,"emscripten-version.txt");
+			return Path.Combine(EmscriptenBase(),"emscripten-version.txt");
 		}
-    }
+	}
 }
  
\ No newline at end of file
diff --git a/Engine/Source/Programs/UnrealBuildTool/HTML5/HTML5ToolChain.cs b/Engine/Source/Programs/UnrealBuildTool/HTML5/HTML5ToolChain.cs
--- a/Engine/Source/Programs/UnrealBuildTool/HTML5/HTML5ToolChain.cs
+++ b/Engine/Source/Programs/UnrealBuildTool/HTML5/HTML5ToolChain.cs
@@ -11,34 +11,29 @@ namespace UnrealBuildTool
 {
 	class HTML5ToolChain : VCToolChain
 	{
-        static string EMCCPath;
-        static string PythonPath;
-		static string EMSDKVersionString;
         // Debug options -- TODO: add these to SDK/Project setup?
         static bool   bEnableTracing = false;
 
 		// cache the location of SDK tools
 		public override void RegisterToolChain()
 		{
-            if (HTML5SDKInfo.IsSDKInstalled() && HTML5SDKInfo.IsPythonInstalled())
-            {
-                    EMCCPath = "\"" + HTML5SDKInfo.EmscriptenCompiler() + "\"";
-					PythonPath = HTML5SDKInfo.PythonPath();
-					EMSDKVersionString = HTML5SDKInfo.EmscriptenVersion().Replace(".", "");
-
+			if (HTML5SDKInfo.IsSDKInstalled())
+			{
 					// set some environment variable we'll need
 					//Environment.SetEnvironmentVariable("EMCC_DEBUG", "cache");
 					Environment.SetEnvironmentVariable("EMCC_CORES", "8");
-					Environment.SetEnvironmentVariable("EMCC_FORCE_STDLIBS", "1");
 					Environment.SetEnvironmentVariable("EMCC_OPTIMIZE_NORMALLY", "1");
 					// finally register the toolchain that is now ready to go
-                    Log.TraceVerbose("        Registered for {0}", CPPTargetPlatform.HTML5.ToString());
+					Log.TraceVerbose("        Registered for {0}", CPPTargetPlatform.HTML5.ToString());
 					UEToolChain.RegisterPlatformToolChain(CPPTargetPlatform.HTML5, this);
-
-   
 			}
 		}
-
+		public override void PreBuildSync()
+		{
+			HTML5SDKInfo.SetupEmscriptenTemp();
+			HTML5SDKInfo.SetUpEmscriptenConfigFile();
+		}
+	
 		static string GetSharedArguments_Global(CPPTargetConfiguration TargetConfiguration, string Architecture, bool bEnableShadowVariableWarning)
 		{
             string Result = " ";
@@ -94,7 +89,6 @@ namespace UnrealBuildTool
             {
             	Result += " --tracing";
             }
-
             return Result;
 		}
 		
@@ -206,9 +200,7 @@ namespace UnrealBuildTool
 				}
 
 				Result += " -s CASE_INSENSITIVE_FS=1 ";
-
-
-			}
+		}
 
 			return Result;
 		}
@@ -327,9 +319,6 @@ namespace UnrealBuildTool
 				Arguments += string.Format(" -D__EMSCRIPTEN_TRACING__");
 			}
 
-			Arguments += string.Format(" -D__EMCC_VER__={0}", EMSDKVersionString);
-
-        
 			var BuildPlatform = UEBuildPlatform.GetBuildPlatformForCPPTargetPlatform(CompileEnvironment.Config.Target.Platform);
 
 			foreach (FileItem SourceFile in SourceFiles)
@@ -364,9 +353,9 @@ namespace UnrealBuildTool
 				}
 
 				CompileAction.WorkingDirectory = Path.GetFullPath(".");
-				CompileAction.CommandPath = PythonPath;
+				CompileAction.CommandPath = HTML5SDKInfo.Python();
 
-				CompileAction.CommandArguments = EMCCPath + " " + Arguments + FileArguments + CompileEnvironment.Config.AdditionalArguments;
+				CompileAction.CommandArguments = HTML5SDKInfo.EmscriptenCompiler() + " " + Arguments + FileArguments + CompileEnvironment.Config.AdditionalArguments;
 
                 //System.Console.WriteLine(CompileAction.CommandArguments); 
 				CompileAction.StatusDescription = Path.GetFileName(SourceFile.AbsolutePath);
@@ -474,8 +463,8 @@ namespace UnrealBuildTool
 
 			LinkAction.bCanExecuteRemotely = false;
 			LinkAction.WorkingDirectory = Path.GetFullPath(".");
-			LinkAction.CommandPath = PythonPath;
-			LinkAction.CommandArguments = EMCCPath;
+			LinkAction.CommandPath = HTML5SDKInfo.Python();
+			LinkAction.CommandArguments = HTML5SDKInfo.EmscriptenCompiler();
 			ReponseLines.Add(GetLinkArguments(LinkEnvironment));
 
 			// Add the input files to a response file, and pass the response file on the command-line.
@@ -509,9 +498,9 @@ namespace UnrealBuildTool
 								ReponseLines.Add(string.Format(" --js-library \"{0}\"", Item.AbsolutePath));
 							else
 								ReponseLines.Add(string.Format(" \"{0}\"", Item.AbsolutePath));
-                            LinkAction.PrerequisiteItems.Add(Item);
-                        }
-                    }
+							LinkAction.PrerequisiteItems.Add(Item);
+						}
+					}
 			}
 			// make the file we will create
 
@@ -527,7 +516,9 @@ namespace UnrealBuildTool
      		LinkAction.StatusDescription = Path.GetFileName(OutputFile.AbsolutePath);
 
 			string ResponseFileName = GetResponseFileName(LinkEnvironment, OutputFile);
-			LinkAction.CommandArguments += string.Format(" @\"{0}\"", ResponseFile.Create(ResponseFileName, ReponseLines));
+
+
+			LinkAction.CommandArguments += string.Format(" @\"{0}\""  , ResponseFile.Create(ResponseFileName, ReponseLines));
 
 			LinkAction.OutputEventHandler = new DataReceivedEventHandler(RemoteOutputReceivedEventHandler);
 
diff --git a/Engine/Source/Programs/UnrealBuildTool/HTML5/UEBuildHTML5.cs b/Engine/Source/Programs/UnrealBuildTool/HTML5/UEBuildHTML5.cs
--- a/Engine/Source/Programs/UnrealBuildTool/HTML5/UEBuildHTML5.cs
+++ b/Engine/Source/Programs/UnrealBuildTool/HTML5/UEBuildHTML5.cs
@@ -8,20 +8,11 @@ using System.IO;
 
 namespace UnrealBuildTool
 {
-    class HTML5Platform : UEBuildPlatform
-    {
-        // use -win32 for win32 builds. ( build html5 platform as a win32 binary for debugging )
-        [XmlConfig]
-        public static string HTML5Architecture = "";
-
-
-        static private string TargetPlatformName = "HTML5";
-
-		// This is the SDK version we support in //CarefullyRedist.
-		// May include minor revisions or descriptions that a default install from SDK_Manager won't have.
-		// e.g. 1.600_Patch001, or 1.610.001.  The SDK_Manager always installs minor revision patches straight into
-		// the default major revision folder.
-		static private string ExpectedSDKVersion = "1.25.0";
+	class HTML5Platform : UEBuildPlatform
+	{
+		// use -win32 for win32 builds. ( build html5 platform as a win32 binary for debugging )
+		[XmlConfig]
+		public static string HTML5Architecture = "";
 
 		/** 
 		 * Whether platform supports switching SDKs during runtime
@@ -30,13 +21,13 @@ namespace UnrealBuildTool
 		 */
 		protected override bool PlatformSupportsAutoSDKs()
 		{
-			return true;
+			return false;
 		}
 		
 		public override string GetSDKTargetPlatformName()
-        {
-            return TargetPlatformName;
-        }
+		{
+			return "HTML5";
+		}
 
 		/** 
 		 * Returns SDK string as required by the platform 
@@ -45,9 +36,7 @@ namespace UnrealBuildTool
 		 */
 		protected override string GetRequiredSDKString()
 		{
-			if (HTML5SDKInfo.IsSDKVersionOverridden())
-				return HTML5SDKInfo.OverriddenSDKVersion();
-			return ExpectedSDKVersion;
+			return HTML5SDKInfo.EmscriptenVersion();
 		}
 
 		protected override String GetRequiredScriptVersionString()
@@ -56,52 +45,29 @@ namespace UnrealBuildTool
 		}
 		
 		// The current architecture - affects everything about how UBT operates on HTML5
-        public override string GetActiveArchitecture()
-        {
-            // by default, use an empty architecture (which is really just a modifier to the platform for some paths/names)
-            return HTML5Architecture;
-        }
+		public override string GetActiveArchitecture()
+		{
+			// by default, use an empty architecture (which is really just a modifier to the platform for some paths/names)
+			return HTML5Architecture;
+		}
 
-        // F5 should always try to run the Win32 version
-        public override string ModifyNMakeOutput(string ExeName)
-        {
-            // nmake Run should always run the win32 version
-            return Path.ChangeExtension(ExeName+"-win32", ".exe");
-        }
+		// F5 should always try to run the Win32 version
+		public override string ModifyNMakeOutput(string ExeName)
+		{
+			// nmake Run should always run the win32 version
+			return Path.ChangeExtension(ExeName+"-win32", ".exe");
+		}
 
-        /**
-         *	Whether the required external SDKs are installed for this platform
-         */
-        protected override SDKStatus HasRequiredManualSDKInternal()
-        {
-			// if any autosdk setup has been done then the local process environment is suspect
-			if (HasSetupAutoSDK())
-			{
-				return SDKStatus.Invalid;
-            }
+		/**
+			*	Whether the required external SDKs are installed for this platform
+			*/
+		protected override SDKStatus HasRequiredManualSDKInternal()
+		{
 			if (!HTML5SDKInfo.IsSDKInstalled())
 			{
 				return SDKStatus.Invalid;
 			}
-            try
-            {
-                int InstalledVersion = Convert.ToInt32(HTML5SDKInfo.EmscriptenVersion().Replace(".","")); 
-                int RequiredVersion = Convert.ToInt32(GetRequiredSDKString().Replace(".","")); 
-
-                if (InstalledVersion >= RequiredVersion)
-                {
-                    return SDKStatus.Valid;
-                }
-                else
-                {
-					Console.WriteLine("EMSCRIPTEN sdk " + HTML5SDKInfo.EmscriptenVersion() + " found which is older than " + RequiredVersion + " Please install the latest emscripten SDK");
-                    return SDKStatus.Invalid;
-                }
-            }
-            catch (Exception /*ex*/)
-            {
-                 return SDKStatus.Invalid;
-            }
+			return SDKStatus.Valid;
         }
 
         public override bool CanUseXGE()
diff --git a/Engine/Source/Runtime/Core/Public/HTML5/HTML5PlatformMisc.h b/Engine/Source/Runtime/Core/Public/HTML5/HTML5PlatformMisc.h
--- a/Engine/Source/Runtime/Core/Public/HTML5/HTML5PlatformMisc.h
+++ b/Engine/Source/Runtime/Core/Public/HTML5/HTML5PlatformMisc.h
@@ -52,7 +52,8 @@ struct CORE_API FHTML5Misc : public FGenericPlatformMisc
 #else
 			emscripten_log(255, "DebugBreak() called!");
 			EM_ASM(
-				throw new Error('DebugBreak() called! Check Log for information');
+				var callstack = new Error;
+				throw callstack.stack;
 			);
 #endif
 		}
