Remove 'IsAggregate()' and 'IsInEC()' methods, now that aggregate nodes are handled separately.

diff --git a/Engine/Source/Programs/AutomationTool/GUBP/GUBP.cs b/Engine/Source/Programs/AutomationTool/GUBP/GUBP.cs
--- a/Engine/Source/Programs/AutomationTool/GUBP/GUBP.cs
+++ b/Engine/Source/Programs/AutomationTool/GUBP/GUBP.cs
@@ -135,22 +135,18 @@ public partial class GUBP : BuildCommand
         GUBPNodes[Node].Node.FullNamesOfPseudosependencies.Clear();
     }
 
-    static List<NodeInfo> GetDependencies(NodeInfo NodeToDo, bool bFlat = false, bool ECOnly = false)
+    static List<NodeInfo> GetDependencies(NodeInfo NodeToDo, bool bFlat = false)
     {
         List<NodeInfo> Result = new List<NodeInfo>();
         foreach (NodeInfo Dependency in NodeToDo.Dependencies)
         {
-            bool Usable = Dependency.Node.RunInEC() || !ECOnly;
-            if (Usable)
+            if (!Result.Contains(Dependency))
             {
-                if (!Result.Contains(Dependency))
-                {
-                    Result.Add(Dependency);
-                }
+                Result.Add(Dependency);
             }
-            if (bFlat || !Usable)
+            if (bFlat)
             {
-                foreach (NodeInfo RNode in GetDependencies(Dependency, bFlat, ECOnly))
+                foreach (NodeInfo RNode in GetDependencies(Dependency, bFlat))
                 {
                     if (!Result.Contains(RNode))
                     {
@@ -161,17 +157,13 @@ public partial class GUBP : BuildCommand
         }
         foreach (NodeInfo PseudoDependency in NodeToDo.PseudoDependencies)
         {
-            bool Usable = PseudoDependency.Node.RunInEC() || !ECOnly;
-            if (Usable)
+            if (!Result.Contains(PseudoDependency))
             {
-                if (!Result.Contains(PseudoDependency))
-                {
-                    Result.Add(PseudoDependency);
-                }
+                Result.Add(PseudoDependency);
             }
-            if (bFlat || !Usable)
+            if (bFlat)
             {
-                foreach (NodeInfo RNode in GetDependencies(PseudoDependency, bFlat, ECOnly))
+                foreach (NodeInfo RNode in GetDependencies(PseudoDependency, bFlat))
                 {
                     if (!Result.Contains(RNode))
                     {
@@ -185,7 +177,7 @@ public partial class GUBP : BuildCommand
 
 	static List<NodeInfo> GetECDependencies(NodeInfo NodeToDo, bool bFlat = false)
     {
-        return GetDependencies(NodeToDo, bFlat, true);
+        return GetDependencies(NodeToDo, bFlat);
     }
     bool NodeDependsOn(NodeInfo Rootward, NodeInfo Leafward)
     {
@@ -494,10 +486,6 @@ public partial class GUBP : BuildCommand
         Log("*********** Desired And Dependent Nodes, in order.");
         foreach (NodeInfo NodeToDo in Nodes)
         {
-            if (ECOnly && !NodeToDo.Node.RunInEC())
-            {
-                continue;
-            }
             string EMails = "";
             if (AddEmailProps)
             {
@@ -999,7 +987,7 @@ public partial class GUBP : BuildCommand
 
     void UpdateNodeHistory(Dictionary<string, NodeHistory> GUBPNodesHistory, NodeInfo NodeToDo, string CLString, string StoreName)
     {
-        if (NodeToDo.Node.RunInEC() && !NodeToDo.Node.TriggerNode() && CLString != "")
+        if (!NodeToDo.Node.TriggerNode() && CLString != "")
         {
             string GameNameIfAny = NodeToDo.Node.GameNameIfAnyForTempStorage();
             string NodeStoreWildCard = StoreName.Replace(CLString, "*") + "-" + NodeToDo.Node.GetFullName();
@@ -1528,7 +1516,7 @@ public partial class GUBP : BuildCommand
                 {
                     if (Node.Value.Node.GetFullName().Equals(ExplicitTriggerName, StringComparison.InvariantCultureIgnoreCase))
                     {
-                        if (Node.Value.Node.TriggerNode() && Node.Value.Node.RunInEC())
+                        if (Node.Value.Node.TriggerNode())
                         {
                             Node.Value.Node.SetAsExplicitTrigger();
 							ExplicitTrigger = Node.Value;
@@ -1547,7 +1535,7 @@ public partial class GUBP : BuildCommand
                 {
                     foreach (KeyValuePair<string, NodeInfo> Node in GUBPNodes)
                     {
-                        if (Node.Value.Node.TriggerNode() && Node.Value.Node.RunInEC())
+                        if (Node.Value.Node.TriggerNode())
                         {
                             Node.Value.Node.SetAsExplicitTrigger();
                         }
@@ -1951,7 +1939,7 @@ public partial class GUBP : BuildCommand
 		}
 		foreach (NodeInfo NodeToDo in NodesToDo)
 		{
-			if (!NodeToDo.Node.IsAggregate() && !NodeToDo.Node.IsTest())
+			if (!NodeToDo.Node.IsTest())
 			{
 				List<NodeInfo> ECDependencies = GetECDependencies(NodeToDo);
 				foreach (NodeInfo Dep in ECDependencies)
@@ -1991,26 +1979,20 @@ public partial class GUBP : BuildCommand
 				{
 					Note = "always";
 				}
-				if (GUBPNodes[Node.Name].Node.RunInEC())
+
+				List<NodeInfo> Deps = GetECDependencies(Node);
+				string All = "";
+				foreach (NodeInfo Dep in Deps)
 				{
-					List<NodeInfo> Deps = GetECDependencies(Node);
-					string All = "";
-					foreach (NodeInfo Dep in Deps)
+					if (All != "")
 					{
-						if (All != "")
-						{
-							All += " ";
-						}
-						All += Dep.Name;
+						All += " ";
 					}
-					LogVerbose("  {0}: {1}      {2}", Node.Name, Note, All);
-					FullNodeList.Add(Node.Name, Note);
-					FullNodeDirectDependencies.Add(Node.Name, All);
-				}
-				else
-				{
-					LogVerbose("  {0}: {1} [Aggregate]", Node.Name, Note);
+					All += Dep.Name;
 				}
+				LogVerbose("  {0}: {1}      {2}", Node.Name, Note, All);
+				FullNodeList.Add(Node.Name, Note);
+				FullNodeDirectDependencies.Add(Node.Name, All);
 			}
 		}
 	}
@@ -2263,7 +2245,7 @@ public partial class GUBP : BuildCommand
 		{
 			foreach (NodeInfo NodeToDo in OrdereredToDo)
 			{
-				if (NodeToDo.Node.RunInEC() && !NodeToDo.IsComplete && NodeToDo != ExplicitTrigger && !NodeDependsOn(ExplicitTrigger, NodeToDo)) // if something is already finished, we don't put it into EC
+				if (!NodeToDo.IsComplete && NodeToDo != ExplicitTrigger && !NodeDependsOn(ExplicitTrigger, NodeToDo)) // if something is already finished, we don't put it into EC
 				{
 					throw new AutomationException("We are being asked to process node {0}, however, this is an explicit trigger {1}, so everything before it should already be handled. It seems likely that you waited too long to run the trigger. You will have to do a new build from scratch.", NodeToDo.Name, ExplicitTrigger.Name);
 				}
@@ -2283,7 +2265,7 @@ public partial class GUBP : BuildCommand
 		// sticky nodes are ones that we run on the main agent. We run then first and they must not be intermixed with parallel jobs
 		foreach (NodeInfo NodeToDo in OrdereredToDo)
 		{
-			if (NodeToDo.Node.RunInEC() && !NodeToDo.IsComplete) // if something is already finished, we don't put it into EC
+			if (!NodeToDo.IsComplete) // if something is already finished, we don't put it into EC
 			{
 				bHaveECNodes = true;
 				if (NodeToDo.Node.IsSticky())
@@ -2320,7 +2302,7 @@ public partial class GUBP : BuildCommand
 			List<NodeInfo> StickyChain = new List<NodeInfo>();
 			foreach (NodeInfo NodeToDo in OrdereredToDo)
 			{
-				if (NodeToDo.Node.RunInEC() && !NodeToDo.IsComplete) // if something is already finished, we don't put it into EC  
+				if (!NodeToDo.IsComplete) // if something is already finished, we don't put it into EC  
 				{
 					string MyAgentGroup = NodeToDo.Node.AgentSharingGroup;
 					if (MyAgentGroup != "")
@@ -2345,7 +2327,7 @@ public partial class GUBP : BuildCommand
 			}
 			foreach (NodeInfo NodeToDo in OrdereredToDo)
 			{
-				if (NodeToDo.Node.RunInEC() && !NodeToDo.IsComplete) // if something is already finished, we don't put it into EC  
+				if (!NodeToDo.IsComplete) // if something is already finished, we don't put it into EC  
 				{
 					string EMails;
 					List<string> NodeProps = GetECPropsForNode(NodeToDo, CLString, out EMails);
@@ -2396,7 +2378,7 @@ public partial class GUBP : BuildCommand
 						List<NodeInfo> EcDeps = GetECDependencies(NodeToDo);
 						foreach (NodeInfo Dep in EcDeps)
 						{
-							if (Dep.Node.RunInEC() && !Dep.IsComplete && OrdereredToDo.Contains(Dep)) // if something is already finished, we don't put it into EC
+							if (!Dep.IsComplete && OrdereredToDo.Contains(Dep)) // if something is already finished, we don't put it into EC
 							{
 								if (OrdereredToDo.IndexOf(Dep) > OrdereredToDo.IndexOf(NodeToDo))
 								{
@@ -2433,7 +2415,7 @@ public partial class GUBP : BuildCommand
 								List<NodeInfo> EcDeps = GetECDependencies(Chain);
 								foreach (NodeInfo Dep in EcDeps)
 								{
-									if (GUBPNodes[Dep.Name].Node.RunInEC() && !Dep.IsComplete && OrdereredToDo.Contains(Dep)) // if something is already finished, we don't put it into EC
+									if (!Dep.IsComplete && OrdereredToDo.Contains(Dep)) // if something is already finished, we don't put it into EC
 									{
 										if (OrdereredToDo.IndexOf(Dep) > OrdereredToDo.IndexOf(Chain))
 										{
@@ -2464,7 +2446,7 @@ public partial class GUBP : BuildCommand
 							List<NodeInfo> EcDeps = GetECDependencies(NodeToDo);
 							foreach (NodeInfo Dep in EcDeps)
 							{
-								if (Dep.Node.RunInEC() && !Dep.IsComplete && OrdereredToDo.Contains(Dep)) // if something is already finished, we don't put it into EC
+								if (!Dep.IsComplete && OrdereredToDo.Contains(Dep)) // if something is already finished, we don't put it into EC
 								{
 									if (OrdereredToDo.IndexOf(Dep) > OrdereredToDo.IndexOf(NodeToDo))
 									{
@@ -2714,8 +2696,7 @@ public partial class GUBP : BuildCommand
 
             // this is kinda complicated
             bool SaveSuccessRecords = (IsBuildMachine || bFakeEC) && // no real reason to make these locally except for fakeEC tests
-                (!NodeToDo.Node.TriggerNode() || NodeToDo.Node.IsSticky()) // trigger nodes are run twice, one to start the new workflow and once when it is actually triggered, we will save reconds for the latter
-                && (NodeToDo.Node.RunInEC() || !NodeToDo.Node.IsAggregate()); //aggregates not in EC can be "run" multiple times, so we can't track those
+                (!NodeToDo.Node.TriggerNode() || NodeToDo.Node.IsSticky()); // trigger nodes are run twice, one to start the new workflow and once when it is actually triggered, we will save reconds for the latter
 
             Log("***** Running GUBP Node {0} -> {1} : {2}", NodeToDo.Node.GetFullName(), GameNameIfAny, NodeStoreName);
             if (NodeToDo.IsComplete)
@@ -2780,44 +2761,41 @@ public partial class GUBP : BuildCommand
 						BuildDuration = (DateTime.UtcNow - StartTime).TotalMilliseconds / 1000;
 						
                     }
-                    if (!NodeToDo.Node.IsAggregate())
-                    {
+
+					using(TelemetryStopwatch StoreBuildProductsStopwatch = new TelemetryStopwatch("StoreBuildProducts"))
+					{
 						double StoreDuration = 0.0;
 						DateTime StartTime = DateTime.UtcNow;
-						using(TelemetryStopwatch StoreBuildProductsStopwatch = new TelemetryStopwatch("StoreBuildProducts"))
-						{
-							TempStorage.StoreToTempStorage(CmdEnv, NodeStoreName, NodeToDo.Node.BuildProducts, !bSaveSharedTempStorage, GameNameIfAny, StorageRootIfAny);
-						}
+						TempStorage.StoreToTempStorage(CmdEnv, NodeStoreName, NodeToDo.Node.BuildProducts, !bSaveSharedTempStorage, GameNameIfAny, StorageRootIfAny);
 						StoreDuration = (DateTime.UtcNow - StartTime).TotalMilliseconds / 1000;
 						Log("Took {0} seconds to store build products", StoreDuration);
 						if (IsBuildMachine)
 						{
 							RunECTool(String.Format("setProperty \"/myJobStep/StoreDuration\" \"{0}\"", StoreDuration.ToString()));
 						}
-                        if (ParseParam("StompCheck"))
+					}
+                    if (ParseParam("StompCheck"))
+                    {
+                        foreach (string Dep in NodeToDo.Node.AllDependencies)
                         {
-                            foreach (string Dep in NodeToDo.Node.AllDependencies)
+                            try
                             {
-                                try
-                                {
-                                    bool WasLocal;
-									using(TelemetryStopwatch RetrieveBuildProductsStopwatch = new TelemetryStopwatch("RetrieveBuildProducts"))
-									{
-										TempStorage.RetrieveFromTempStorage(CmdEnv, NodeStoreName, out WasLocal, GameNameIfAny, StorageRootIfAny);
-									}
-									if (!WasLocal)
-									{
-										throw new AutomationException("Retrieve was not local?");
-									}
+                                bool WasLocal;
+								using(TelemetryStopwatch RetrieveBuildProductsStopwatch = new TelemetryStopwatch("RetrieveBuildProducts"))
+								{
+									TempStorage.RetrieveFromTempStorage(CmdEnv, NodeStoreName, out WasLocal, GameNameIfAny, StorageRootIfAny);
+								}
+								if (!WasLocal)
+								{
+									throw new AutomationException("Retrieve was not local?");
+								}
 																	                                    
-                                }
-                                catch(Exception Ex)
-                                {
-                                    throw new AutomationException("Node {0} stomped Node {1}   Ex: {2}", NodeToDo.Name, Dep, LogUtils.FormatException(Ex));
-                                }
+                            }
+                            catch(Exception Ex)
+                            {
+                                throw new AutomationException("Node {0} stomped Node {1}   Ex: {2}", NodeToDo.Name, Dep, LogUtils.FormatException(Ex));
                             }
                         }
-
                     }
                 }
                 catch (Exception Ex)
diff --git a/Engine/Source/Programs/AutomationTool/GUBP/LegacyNodes.cs b/Engine/Source/Programs/AutomationTool/GUBP/LegacyNodes.cs
--- a/Engine/Source/Programs/AutomationTool/GUBP/LegacyNodes.cs
+++ b/Engine/Source/Programs/AutomationTool/GUBP/LegacyNodes.cs
@@ -57,19 +57,11 @@ partial class GUBP
         public virtual bool SendSuccessEmail()
         {
             return false;
-        }
-        public virtual bool RunInEC()
-        {
-            return true;
         }
 		public virtual bool IsTest()
 		{
 			return false;
 		}
-        public virtual bool IsAggregate()
-        {
-            return false;
-        }
 		public virtual string NodeHostPlatform()
 		{
 			return "";
