Remove GetECDependencies (now the same as GetDependencies).

diff --git a/Engine/Source/Programs/AutomationTool/GUBP/GUBP.cs b/Engine/Source/Programs/AutomationTool/GUBP/GUBP.cs
--- a/Engine/Source/Programs/AutomationTool/GUBP/GUBP.cs
+++ b/Engine/Source/Programs/AutomationTool/GUBP/GUBP.cs
@@ -175,10 +175,6 @@ public partial class GUBP : BuildCommand
         return Result;
     }
 
-	static List<NodeInfo> GetECDependencies(NodeInfo NodeToDo, bool bFlat = false)
-    {
-        return GetDependencies(NodeToDo, bFlat);
-    }
     bool NodeDependsOn(NodeInfo Rootward, NodeInfo Leafward)
     {
         List<NodeInfo> Deps = GetDependencies(Leafward, true);
@@ -578,7 +574,7 @@ public partial class GUBP : BuildCommand
             }
             if (bShowECDependencies)
             {
-                foreach (NodeInfo Dep in GetECDependencies(NodeToDo))
+                foreach (NodeInfo Dep in GetDependencies(NodeToDo))
                 {
                     Log("           {0}", Dep.Name);
                 }
@@ -1941,7 +1937,7 @@ public partial class GUBP : BuildCommand
 		{
 			if (!NodeToDo.Node.IsTest())
 			{
-				List<NodeInfo> ECDependencies = GetECDependencies(NodeToDo);
+				List<NodeInfo> ECDependencies = GetDependencies(NodeToDo);
 				foreach (NodeInfo Dep in ECDependencies)
 				{
 					string CurrentValue;
@@ -1980,7 +1976,7 @@ public partial class GUBP : BuildCommand
 					Note = "always";
 				}
 
-				List<NodeInfo> Deps = GetECDependencies(Node);
+				List<NodeInfo> Deps = GetDependencies(Node);
 				string All = "";
 				foreach (NodeInfo Dep in Deps)
 				{
@@ -2375,7 +2371,7 @@ public partial class GUBP : BuildCommand
 
 					PreconditionParentPath = ParentPath;
 					{
-						List<NodeInfo> EcDeps = GetECDependencies(NodeToDo);
+						List<NodeInfo> EcDeps = GetDependencies(NodeToDo);
 						foreach (NodeInfo Dep in EcDeps)
 						{
 							if (!Dep.IsComplete && OrdereredToDo.Contains(Dep)) // if something is already finished, we don't put it into EC
@@ -2412,7 +2408,7 @@ public partial class GUBP : BuildCommand
 							// to avoid idle agents (and also EC doesn't actually reserve our agent!), we promote all dependencies to the first one
 							foreach (NodeInfo Chain in MyChain)
 							{
-								List<NodeInfo> EcDeps = GetECDependencies(Chain);
+								List<NodeInfo> EcDeps = GetDependencies(Chain);
 								foreach (NodeInfo Dep in EcDeps)
 								{
 									if (!Dep.IsComplete && OrdereredToDo.Contains(Dep)) // if something is already finished, we don't put it into EC
@@ -2443,7 +2439,7 @@ public partial class GUBP : BuildCommand
 						}
 						else
 						{
-							List<NodeInfo> EcDeps = GetECDependencies(NodeToDo);
+							List<NodeInfo> EcDeps = GetDependencies(NodeToDo);
 							foreach (NodeInfo Dep in EcDeps)
 							{
 								if (!Dep.IsComplete && OrdereredToDo.Contains(Dep)) // if something is already finished, we don't put it into EC
