Fix path to assembly used to spawn child instances of UAT.

diff --git a/Engine/Source/Programs/AutomationTool/AutomationUtils/Utils.cs b/Engine/Source/Programs/AutomationTool/AutomationUtils/Utils.cs
--- a/Engine/Source/Programs/AutomationTool/AutomationUtils/Utils.cs
+++ b/Engine/Source/Programs/AutomationTool/AutomationUtils/Utils.cs
@@ -726,7 +726,7 @@ namespace AutomationTool
 		{
 			get
 			{
-				return CommandUtils.CombinePaths(new Uri(System.Reflection.Assembly.GetExecutingAssembly().CodeBase).LocalPath);
+				return CommandUtils.CombinePaths(new Uri(System.Reflection.Assembly.GetEntryAssembly().CodeBase).LocalPath);
 			}
 		}
 
