Fix some implicit string conversions from NodeInfo to node names.

diff --git a/Engine/Source/Programs/AutomationTool/GUBP/GUBP.cs b/Engine/Source/Programs/AutomationTool/GUBP/GUBP.cs
--- a/Engine/Source/Programs/AutomationTool/GUBP/GUBP.cs
+++ b/Engine/Source/Programs/AutomationTool/GUBP/GUBP.cs
@@ -1898,7 +1898,7 @@ public partial class GUBP : BuildCommand
 				}
 				else
 				{
-					LogVerbose(" Rejecting {0} because -NoLinux was requested", NodeToDo);
+					LogVerbose(" Rejecting {0} because -NoLinux was requested", NodeToDo.Name);
 				}
 			}
 			NodesToDo = NewNodesToDo;
@@ -2570,7 +2570,7 @@ public partial class GUBP : BuildCommand
 						) // unfinished triggers are never run directly by EC, rather it does another job setup
 					{
 						string Arg = String.Format("gubp -Node={0} -FakeEC {1} {2} {3} {4} {5}",
-							NodeToDo,
+							NodeToDo.Name,
 							bFake ? "-Fake" : "",
 							ParseParam("AllPlatforms") ? "-AllPlatforms" : "",
 							ParseParam("UnfinishedTriggersFirst") ? "-UnfinishedTriggersFirst" : "",
