Make a separate function for finding the RunCondition for nodes in EC. Also fix unused variable warning on mac.

diff --git a/Engine/Source/Programs/AutomationTool/GUBP/GUBP.cs b/Engine/Source/Programs/AutomationTool/GUBP/GUBP.cs
--- a/Engine/Source/Programs/AutomationTool/GUBP/GUBP.cs
+++ b/Engine/Source/Programs/AutomationTool/GUBP/GUBP.cs
@@ -474,7 +474,6 @@ public partial class GUBP : BuildCommand
         bool bShowHistory = !bp.ParseParam("NoHistory") && GUBPNodesHistory != null;
         bool AddEmailProps = bp.ParseParam("ShowEmails");
         bool ECProc = bp.ParseParam("ShowECProc");
-        bool ECOnly = bp.ParseParam("ShowECOnly");
         bool bShowTriggers = true;
         string LastControllingTrigger = "";
         string LastAgentGroup = "";
@@ -1288,10 +1287,6 @@ public partial class GUBP : BuildCommand
     [Help("UserTimeIndex=", "An integer used to determine subsets to run based on DependentCISFrequencyQuantumShift, this one overrides TimeIndex")]
     [Help("PreflightUID=", "A unique integer tag from EC used as part of the tempstorage, builds and label names to distinguish multiple attempts.")]
     [Help("Node=", "Nodes to process, -node=Node1+Node2+Node3, if no nodes or games are specified, defaults to all nodes.")]
-    [Help("SetupNode=", "Like -Node, but only applies with CommanderJobSetupOnly")]
-    [Help("RelatedToNode=", "Nodes to process, -RelatedToNode=Node1+Node2+Node3, use all nodes that either depend on these nodes or these nodes depend on them.")]
-    [Help("SetupRelatedToNode=", "Like -RelatedToNode, but only applies with CommanderJobSetupOnly")]
-    [Help("OnlyNode=", "Nodes to process NO dependencies, -OnlyNode=Node1+Node2+Node3, if no nodes or games are specified, defaults to all nodes.")]
     [Help("TriggerNode=", "Trigger Nodes to process, -triggernode=Node.")]
     [Help("Game=", "Games to process, -game=Game1+Game2+Game3, if no games or nodes are specified, defaults to all nodes.")]
     [Help("ListOnly", "List Nodes in this branch")]
@@ -1310,8 +1305,6 @@ public partial class GUBP : BuildCommand
     [Help("ShowDependencies", "Show node dependencies.")]
     [Help("ShowECDependencies", "Show EC node dependencies instead.")]
     [Help("ShowECProc", "Show EC proc names.")]
-    [Help("BuildRocket", "Build in rocket mode.")]
-    [Help("ShowECOnly", "Only show EC nodes.")]
     [Help("ECProject", "From EC, the name of the project, used to get a version number.")]
     [Help("CIS", "This is a CIS run, assign TimeIndex based on the history.")]
     [Help("ForceIncrementalCompile", "make sure all compiles are incremental")]
@@ -2357,19 +2350,8 @@ public partial class GUBP : BuildCommand
 						Args = Args + ", {actualParameterName => 'EmailsForTrigger', value => \'" + EMails + "\'}";
 					}
 					Args = Args + "]";
-					string PreCondition = "";
-					string RunCondition = "";
-					List<string> UncompletedEcDeps = new List<string>();
-					List<string> PreConditionUncompletedEcDeps = new List<string>();
-					string MyAgentGroup = NodeToDo.Node.AgentSharingGroup;
-					bool bDoNestedJobstep = false;
-					bool bDoFirstNestedJobstep = false;
 
-
-					string NodeParentPath = ParentPath;
-					string PreconditionParentPath;
-
-					PreconditionParentPath = ParentPath;
+					List<NodeInfo> UncompletedEcDeps = new List<NodeInfo>();
 					{
 						List<NodeInfo> EcDeps = GetDependencies(NodeToDo);
 						foreach (NodeInfo Dep in EcDeps)
@@ -2380,14 +2362,26 @@ public partial class GUBP : BuildCommand
 								{
 									throw new AutomationException("Topological sort error, node {0} has a dependency of {1} which sorted after it.", NodeToDo.Name, Dep.Name);
 								}
-								UncompletedEcDeps.Add(Dep.Name);
+								UncompletedEcDeps.Add(Dep);
 							}
 						}
 					}
 
-					foreach (string Dep in UncompletedEcDeps)
+					string PreCondition = "";
+					List<string> PreConditionUncompletedEcDeps = new List<string>();
+					string MyAgentGroup = NodeToDo.Node.AgentSharingGroup;
+					bool bDoNestedJobstep = false;
+					bool bDoFirstNestedJobstep = false;
+
+
+					string NodeParentPath = ParentPath;
+					string PreconditionParentPath;
+
+					PreconditionParentPath = ParentPath;
+
+					foreach (NodeInfo Dep in UncompletedEcDeps)
 					{
-						PreConditionUncompletedEcDeps.Add(Dep);
+						PreConditionUncompletedEcDeps.Add(Dep.Name);
 					}
 					if (MyAgentGroup != "")
 					{
@@ -2477,23 +2471,7 @@ public partial class GUBP : BuildCommand
 						PreCondition = PreCondition + ") true;]\"";
 					}
 
-					if (UncompletedEcDeps.Count > 0)
-					{
-						RunCondition = "\"\\$\" . \"[/javascript if(";
-						int Index = 0;
-						foreach (string Dep in UncompletedEcDeps)
-						{
-							RunCondition = RunCondition + "((\'\\$\" . \"[" + GetJobStep(PreconditionParentPath, Dep) + "/outcome]\' == \'success\') || ";
-							RunCondition = RunCondition + "(\'\\$\" . \"[" + GetJobStep(PreconditionParentPath, Dep) + "/outcome]\' == \'warning\'))";
-
-							Index++;
-							if (Index != UncompletedEcDeps.Count)
-							{
-								RunCondition = RunCondition + " && ";
-							}
-						}
-						RunCondition = RunCondition + ")true; else false;]\"";
-					}
+					string RunCondition = GetRunConditionForNode(UncompletedEcDeps, PreconditionParentPath);
 
 					if (bDoNestedJobstep)
 					{
@@ -2624,13 +2602,29 @@ public partial class GUBP : BuildCommand
 			CommandUtils.WriteAllLines(BranchJobDefFile, ECProps.ToArray());
 			RunECTool(String.Format("setProperty \"/myJob/BranchJobDefFile\" \"{0}\"", BranchJobDefFile.Replace("\\", "\\\\")));
 		}
-		if (bFakeEC)
+	}
+
+	private string GetRunConditionForNode(List<NodeInfo> UncompletedEcDeps, string PreconditionParentPath)
+	{
+		string RunCondition = "";
+		if (UncompletedEcDeps.Count > 0)
 		{
-			foreach (string Args in FakeECArgs)
+			RunCondition = "\"\\$\" . \"[/javascript if(";
+			int Index = 0;
+			foreach (NodeInfo Dep in UncompletedEcDeps)
 			{
-				RunUAT(CmdEnv, Args);
+				RunCondition = RunCondition + "((\'\\$\" . \"[" + GetJobStep(PreconditionParentPath, Dep.Name) + "/outcome]\' == \'success\') || ";
+				RunCondition = RunCondition + "(\'\\$\" . \"[" + GetJobStep(PreconditionParentPath, Dep.Name) + "/outcome]\' == \'warning\'))";
+
+				Index++;
+				if (Index != UncompletedEcDeps.Count)
+				{
+					RunCondition = RunCondition + " && ";
+				}
 			}
+			RunCondition = RunCondition + ")true; else false;]\"";
 		}
+		return RunCondition;
 	}
 
 	void ExecuteNodes(List<NodeInfo> OrdereredToDo, bool bFake, bool bFakeEC, bool bSaveSharedTempStorage, Dictionary<string, NodeHistory> GUBPNodesHistory, string CLString, string StoreName, string FakeFail)
