Refactor out the code for generating precondition strings for each node.

diff --git a/Engine/Source/Programs/AutomationTool/GUBP/GUBP.cs b/Engine/Source/Programs/AutomationTool/GUBP/GUBP.cs
--- a/Engine/Source/Programs/AutomationTool/GUBP/GUBP.cs
+++ b/Engine/Source/Programs/AutomationTool/GUBP/GUBP.cs
@@ -2367,111 +2367,26 @@ public partial class GUBP : BuildCommand
 						}
 					}
 
-					string PreCondition = "";
-					List<string> PreConditionUncompletedEcDeps = new List<string>();
+					string PreCondition = GetPreConditionForNode(OrdereredToDo, ParentPath, bHasNoop, AgentGroupChains, StickyChain, NodeToDo, UncompletedEcDeps);
+					string RunCondition = GetRunConditionForNode(UncompletedEcDeps, ParentPath);
+
 					string MyAgentGroup = NodeToDo.Node.AgentSharingGroup;
 					bool bDoNestedJobstep = false;
 					bool bDoFirstNestedJobstep = false;
 
-
 					string NodeParentPath = ParentPath;
-					string PreconditionParentPath;
-
-					PreconditionParentPath = ParentPath;
-
-					foreach (NodeInfo Dep in UncompletedEcDeps)
-					{
-						PreConditionUncompletedEcDeps.Add(Dep.Name);
-					}
 					if (MyAgentGroup != "")
 					{
 						bDoNestedJobstep = true;
 						NodeParentPath = ParentPath + "/jobSteps[" + MyAgentGroup + "]";
 
-						PreConditionUncompletedEcDeps = new List<string>();
-
 						List<NodeInfo> MyChain = AgentGroupChains[MyAgentGroup];
 						int MyIndex = MyChain.IndexOf(NodeToDo);
-						if (MyIndex > 0)
-						{
-							PreConditionUncompletedEcDeps.Add(MyChain[MyIndex - 1].Name);
-						}
-						else
+						if (MyIndex <= 0)
 						{
 							bDoFirstNestedJobstep = bDoNestedJobstep;
-							// to avoid idle agents (and also EC doesn't actually reserve our agent!), we promote all dependencies to the first one
-							foreach (NodeInfo Chain in MyChain)
-							{
-								List<NodeInfo> EcDeps = GetDependencies(Chain);
-								foreach (NodeInfo Dep in EcDeps)
-								{
-									if (!Dep.IsComplete && OrdereredToDo.Contains(Dep)) // if something is already finished, we don't put it into EC
-									{
-										if (OrdereredToDo.IndexOf(Dep) > OrdereredToDo.IndexOf(Chain))
-										{
-											throw new AutomationException("Topological sort error, node {0} has a dependency of {1} which sorted after it.", Chain.Name, Dep.Name);
-										}
-										if (!MyChain.Contains(Dep) && !PreConditionUncompletedEcDeps.Contains(Dep.Name))
-										{
-											PreConditionUncompletedEcDeps.Add(Dep.Name);
-										}
-									}
-								}
-							}
-						}
-					}
-					if (NodeToDo.Node.IsSticky())
-					{
-						List<NodeInfo> MyChain = StickyChain;
-						int MyIndex = MyChain.IndexOf(NodeToDo);
-						if (MyIndex > 0)
-						{
-							if (!PreConditionUncompletedEcDeps.Contains(MyChain[MyIndex - 1].Name) && !MyChain[MyIndex - 1].IsComplete)
-							{
-								PreConditionUncompletedEcDeps.Add(MyChain[MyIndex - 1].Name);
-							}
-						}
-						else
-						{
-							List<NodeInfo> EcDeps = GetDependencies(NodeToDo);
-							foreach (NodeInfo Dep in EcDeps)
-							{
-								if (!Dep.IsComplete && OrdereredToDo.Contains(Dep)) // if something is already finished, we don't put it into EC
-								{
-									if (OrdereredToDo.IndexOf(Dep) > OrdereredToDo.IndexOf(NodeToDo))
-									{
-										throw new AutomationException("Topological sort error, node {0} has a dependency of {1} which sorted after it.", NodeToDo.Name, Dep.Name);
-									}
-									if (!MyChain.Contains(Dep) && !PreConditionUncompletedEcDeps.Contains(Dep.Name))
-									{
-										PreConditionUncompletedEcDeps.Add(Dep.Name);
-									}
-								}
-							}
 						}
 					}
-					if (bHasNoop && PreConditionUncompletedEcDeps.Count == 0)
-					{
-						PreConditionUncompletedEcDeps.Add("Noop");
-					}
-					if (PreConditionUncompletedEcDeps.Count > 0)
-					{
-						PreCondition = "\"\\$\" . \"[/javascript if(";
-						// these run "parallel", but we add preconditions to serialize them
-						int Index = 0;
-						foreach (string Dep in PreConditionUncompletedEcDeps)
-						{
-							PreCondition = PreCondition + "getProperty('" + GetJobStep(PreconditionParentPath, Dep) + "/status\') == \'completed\'";
-							Index++;
-							if (Index != PreConditionUncompletedEcDeps.Count)
-							{
-								PreCondition = PreCondition + " && ";
-							}
-						}
-						PreCondition = PreCondition + ") true;]\"";
-					}
-
-					string RunCondition = GetRunConditionForNode(UncompletedEcDeps, PreconditionParentPath);
 
 					if (bDoNestedJobstep)
 					{
@@ -2505,7 +2420,7 @@ public partial class GUBP : BuildCommand
 								{
 									NestArgs = NestArgs + ", precondition  => ";
 									NestArgs = NestArgs + "\"\\$\" . \"[/javascript if(";
-									NestArgs = NestArgs + "getProperty('" + PreconditionParentPath + "/jobSteps[" + MyAgentGroup + "]/jobSteps[" + MyAgentGroup + "_GetPool]/status') == 'completed'";
+									NestArgs = NestArgs + "getProperty('" + ParentPath + "/jobSteps[" + MyAgentGroup + "]/jobSteps[" + MyAgentGroup + "_GetPool]/status') == 'completed'";
 									NestArgs = NestArgs + ") true;]\"";
 								}
 								NestArgs = NestArgs + "});";
@@ -2513,7 +2428,7 @@ public partial class GUBP : BuildCommand
 							}
 							{
 								PreCondition = "\"\\$\" . \"[/javascript if(";
-								PreCondition = PreCondition + "getProperty('" + PreconditionParentPath + "/jobSteps[" + MyAgentGroup + "]/jobSteps[" + MyAgentGroup + "_GetAgent]/status') == 'completed'";
+								PreCondition = PreCondition + "getProperty('" + ParentPath + "/jobSteps[" + MyAgentGroup + "]/jobSteps[" + MyAgentGroup + "_GetAgent]/status') == 'completed'";
 								PreCondition = PreCondition + ") true;]\"";
 							}
 						}
@@ -2576,7 +2491,7 @@ public partial class GUBP : BuildCommand
 						{
 							string RelPreCondition = "\"\\$\" . \"[/javascript if(";
 							// this runs "parallel", but we a precondition to serialize it
-							RelPreCondition = RelPreCondition + "getProperty('" + PreconditionParentPath + "/jobSteps[" + NodeToDo.Name + "]/status') == 'completed'";
+							RelPreCondition = RelPreCondition + "getProperty('" + ParentPath + "/jobSteps[" + NodeToDo.Name + "]/status') == 'completed'";
 							RelPreCondition = RelPreCondition + ") true;]\"";
 							// we need to release the resource
 							string RelArgs = String.Format("{0}, subprocedure => 'GUBP_Release_AgentShare', parallel => '1', jobStepName => 'Release_{1}', actualParameter => [{{actualParameterName => 'AgentSharingGroup', valued => '{2}'}}], releaseMode => 'release', precondition => '{3}'",
@@ -2604,6 +2519,93 @@ public partial class GUBP : BuildCommand
 		}
 	}
 
+	private string GetPreConditionForNode(List<NodeInfo> OrdereredToDo, string PreconditionParentPath, bool bHasNoop, Dictionary<string, List<NodeInfo>> AgentGroupChains, List<NodeInfo> StickyChain, NodeInfo NodeToDo, List<NodeInfo> UncompletedEcDeps)
+	{
+		List<NodeInfo> PreConditionUncompletedEcDeps = new List<NodeInfo>();
+		if(NodeToDo.Node.AgentSharingGroup == "")
+		{
+			PreConditionUncompletedEcDeps = new List<NodeInfo>(UncompletedEcDeps);
+		}
+		else 
+		{
+			List<NodeInfo> MyChain = AgentGroupChains[NodeToDo.Node.AgentSharingGroup];
+			int MyIndex = MyChain.IndexOf(NodeToDo);
+			if (MyIndex > 0)
+			{
+				PreConditionUncompletedEcDeps.Add(MyChain[MyIndex - 1]);
+			}
+			else
+			{
+				// to avoid idle agents (and also EC doesn't actually reserve our agent!), we promote all dependencies to the first one
+				foreach (NodeInfo Chain in MyChain)
+				{
+					List<NodeInfo> EcDeps = GetDependencies(Chain);
+					foreach (NodeInfo Dep in EcDeps)
+					{
+						if (!Dep.IsComplete && OrdereredToDo.Contains(Dep)) // if something is already finished, we don't put it into EC
+						{
+							if (OrdereredToDo.IndexOf(Dep) > OrdereredToDo.IndexOf(Chain))
+							{
+								throw new AutomationException("Topological sort error, node {0} has a dependency of {1} which sorted after it.", Chain.Name, Dep.Name);
+							}
+							if (!MyChain.Contains(Dep) && !PreConditionUncompletedEcDeps.Contains(Dep))
+							{
+								PreConditionUncompletedEcDeps.Add(Dep);
+							}
+						}
+					}
+				}
+			}
+		}
+		if (NodeToDo.Node.IsSticky())
+		{
+			List<NodeInfo> MyChain = StickyChain;
+			int MyIndex = MyChain.IndexOf(NodeToDo);
+			if (MyIndex > 0)
+			{
+				if (!PreConditionUncompletedEcDeps.Contains(MyChain[MyIndex - 1]) && !MyChain[MyIndex - 1].IsComplete)
+				{
+					PreConditionUncompletedEcDeps.Add(MyChain[MyIndex - 1]);
+				}
+			}
+			else
+			{
+				List<NodeInfo> EcDeps = GetDependencies(NodeToDo);
+				foreach (NodeInfo Dep in EcDeps)
+				{
+					if (!Dep.IsComplete && OrdereredToDo.Contains(Dep)) // if something is already finished, we don't put it into EC
+					{
+						if (OrdereredToDo.IndexOf(Dep) > OrdereredToDo.IndexOf(NodeToDo))
+						{
+							throw new AutomationException("Topological sort error, node {0} has a dependency of {1} which sorted after it.", NodeToDo.Name, Dep.Name);
+						}
+						if (!MyChain.Contains(Dep) && !PreConditionUncompletedEcDeps.Contains(Dep))
+						{
+							PreConditionUncompletedEcDeps.Add(Dep);
+						}
+					}
+				}
+			}
+		}
+
+		List<string> JobStepNames = new List<string>();
+		if (bHasNoop && PreConditionUncompletedEcDeps.Count == 0)
+		{
+			JobStepNames.Add(GetJobStep(PreconditionParentPath, "Noop"));
+		}
+		else
+		{
+			JobStepNames.AddRange(PreConditionUncompletedEcDeps.Select(x => GetJobStep(PreconditionParentPath, x.Name)));
+		}
+
+		string PreCondition = "";
+		if (JobStepNames.Count > 0)
+		{
+			PreCondition = String.Format("\"\\$\" . \"[/javascript if({0}) true;]\"", String.Join(" && ", JobStepNames.Select(x => String.Format("getProperty('{0}/status') == 'completed'", x))));
+		}
+		return PreCondition;
+	}
+
 	private string GetRunConditionForNode(List<NodeInfo> UncompletedEcDeps, string PreconditionParentPath)
 	{
 		string RunCondition = "";
