Remove direct references to GUBPNodes dictionary when setting up EC.

diff --git a/Engine/Source/Programs/AutomationTool/GUBP/GUBP.cs b/Engine/Source/Programs/AutomationTool/GUBP/GUBP.cs
--- a/Engine/Source/Programs/AutomationTool/GUBP/GUBP.cs
+++ b/Engine/Source/Programs/AutomationTool/GUBP/GUBP.cs
@@ -967,15 +967,18 @@ public partial class GUBP : BuildCommand
         return OrdereredToDo;
     }
 
-    string GetJobStepPath(string Dep)
+    string GetJobStepPath(NodeInfo Dep)
     {
-        if (Dep != "Noop" && GUBPNodes[Dep].Node.AgentSharingGroup != "")
+		if (Dep.Node.AgentSharingGroup == "")
+		{
+			return "jobSteps[" + Dep.Name + "]";
+		}
+        else
         {
-            return "jobSteps[" + GUBPNodes[Dep].Node.AgentSharingGroup + "]/jobSteps[" + Dep + "]";
+            return "jobSteps[" + Dep.Node.AgentSharingGroup + "]/jobSteps[" + Dep.Name + "]";
         }
-        return "jobSteps[" + Dep + "]";
     }
-    string GetJobStep(string ParentPath, string Dep)
+    string GetJobStep(string ParentPath, NodeInfo Dep)
     {
         return ParentPath + "/" + GetJobStepPath(Dep);
     }
@@ -1235,7 +1238,7 @@ public partial class GUBP : BuildCommand
 			ECProps.Add(string.Format("AgentRequirementString/{0}={1}", NodeToDo.Name, AgentReq));
 			ECProps.Add(string.Format("RequiredMemory/{0}={1}", NodeToDo.Name, NodeToDo.Node.AgentMemoryRequirement(this)));
 			ECProps.Add(string.Format("Timeouts/{0}={1}", NodeToDo.Name, NodeToDo.Node.TimeoutInMinutes()));
-			ECProps.Add(string.Format("JobStepPath/{0}={1}", NodeToDo.Name, GetJobStepPath(NodeToDo.Name)));
+			ECProps.Add(string.Format("JobStepPath/{0}={1}", NodeToDo.Name, GetJobStepPath(NodeToDo)));
 		}
 		
         return ECProps;
@@ -2591,11 +2594,11 @@ public partial class GUBP : BuildCommand
 		List<string> JobStepNames = new List<string>();
 		if (bHasNoop && PreConditionUncompletedEcDeps.Count == 0)
 		{
-			JobStepNames.Add(GetJobStep(PreconditionParentPath, "Noop"));
+			JobStepNames.Add(PreconditionParentPath + "/jobSteps[Noop]");
 		}
 		else
 		{
-			JobStepNames.AddRange(PreConditionUncompletedEcDeps.Select(x => GetJobStep(PreconditionParentPath, x.Name)));
+			JobStepNames.AddRange(PreConditionUncompletedEcDeps.Select(x => GetJobStep(PreconditionParentPath, x)));
 		}
 
 		string PreCondition = "";
@@ -2615,8 +2618,8 @@ public partial class GUBP : BuildCommand
 			int Index = 0;
 			foreach (NodeInfo Dep in UncompletedEcDeps)
 			{
-				RunCondition = RunCondition + "((\'\\$\" . \"[" + GetJobStep(PreconditionParentPath, Dep.Name) + "/outcome]\' == \'success\') || ";
-				RunCondition = RunCondition + "(\'\\$\" . \"[" + GetJobStep(PreconditionParentPath, Dep.Name) + "/outcome]\' == \'warning\'))";
+				RunCondition = RunCondition + "((\'\\$\" . \"[" + GetJobStep(PreconditionParentPath, Dep) + "/outcome]\' == \'success\') || ";
+				RunCondition = RunCondition + "(\'\\$\" . \"[" + GetJobStep(PreconditionParentPath, Dep) + "/outcome]\' == \'warning\'))";
 
 				Index++;
 				if (Index != UncompletedEcDeps.Count)
