Remove global node history cache, since it's only updated in the only places where it's used. Also move all EC setup under a single function.

diff --git a/Engine/Source/Programs/AutomationTool/GUBP/GUBP.cs b/Engine/Source/Programs/AutomationTool/GUBP/GUBP.cs
--- a/Engine/Source/Programs/AutomationTool/GUBP/GUBP.cs
+++ b/Engine/Source/Programs/AutomationTool/GUBP/GUBP.cs
@@ -181,6 +181,10 @@ public partial class GUBP : BuildCommand
         return Deps.Contains(Rootward);
     }
 
+	/// <summary>
+	/// Recursively update the ControllingTriggers array for each of the nodes passed in. 
+	/// </summary>
+	/// <param name="NodesToDo">Nodes to update</param>
 	static void FindControllingTriggers(IEnumerable<NodeInfo> NodesToDo)
 	{
 		foreach(NodeInfo NodeToDo in NodesToDo)
@@ -189,6 +193,10 @@ public partial class GUBP : BuildCommand
 		}
 	}
 
+	/// <summary>
+	/// Recursively find the controlling triggers for the given node.
+	/// </summary>
+	/// <param name="NodeToDo">Node to find the controlling triggers for</param>
     static void FindControllingTriggers(NodeInfo NodeToDo)
     {
 		if(NodeToDo.ControllingTriggers == null)
@@ -464,14 +472,9 @@ public partial class GUBP : BuildCommand
 
     }
 
-    void PrintNodes(GUBP bp, List<NodeInfo> Nodes, IEnumerable<AggregateInfo> Aggregates, Dictionary<string, NodeHistory> GUBPNodesHistory, List<NodeInfo> UnfinishedTriggers = null)
+    void PrintNodes(GUBP bp, List<NodeInfo> Nodes, IEnumerable<AggregateInfo> Aggregates, List<NodeInfo> UnfinishedTriggers = null)
     {
-        bool bShowAllChanges = bp.ParseParam("AllChanges") && GUBPNodesHistory != null;
-        bool bShowChanges = (bp.ParseParam("Changes") && GUBPNodesHistory != null) || bShowAllChanges;
-        bool bShowDetailedHistory = (bp.ParseParam("History") && GUBPNodesHistory != null) || bShowChanges;
         bool bShowDependencies = bp.ParseParam("ShowDependencies");
-        bool bShowECDependencies = bp.ParseParam("ShowECDependencies");
-        bool bShowHistory = !bp.ParseParam("NoHistory") && GUBPNodesHistory != null;
         bool AddEmailProps = bp.ParseParam("ShowEmails");
         bool ECProc = bp.ParseParam("ShowECProc");
         bool bShowTriggers = true;
@@ -544,22 +547,7 @@ public partial class GUBP : BuildCommand
                 EMails,
                 ECProc ? NodeToDo.Node.ECProcedure() : ""
                 );
-            if (bShowHistory && GUBPNodesHistory.ContainsKey(NodeToDo.Name))
-            {
-                NodeHistory History = GUBPNodesHistory[NodeToDo.Name];
 
-                if (bShowDetailedHistory)
-                {
-                    PrintDetailedChanges(History, bShowAllChanges);
-                }
-                else
-                {
-                    Log("         Last Success: {0}", History.LastSucceeded);
-                    Log("         Last Fail   : {0}", History.LastFailed);
-                    Log("          Fails Since: {0}", History.FailedString);
-                    Log("     InProgress Since: {0}", History.InProgressString);					
-                }
-            }
             if (bShowDependencies)
             {
                 foreach (NodeInfo Dep in NodeToDo.Dependencies)
@@ -571,13 +559,6 @@ public partial class GUBP : BuildCommand
                     Log("           pdep> {0}", Dep.Name);
                 }
             }
-            if (bShowECDependencies)
-            {
-                foreach (NodeInfo Dep in GetDependencies(NodeToDo))
-                {
-                    Log("           {0}", Dep.Name);
-                }
-            }
         }
 
 		AggregateInfo[] MatchingAggregates = Aggregates.Where(x => x.Dependencies.All(y => Nodes.Contains(y))).ToArray();
@@ -679,7 +660,7 @@ public partial class GUBP : BuildCommand
         return AltHostPlatform;
     }
 
-    List<int> ConvertCLToIntList(List<string> Strings)
+    static List<int> ConvertCLToIntList(List<string> Strings)
     {
         List<int> Result = new List<int>();
         foreach (string ThisString in Strings)
@@ -983,13 +964,15 @@ public partial class GUBP : BuildCommand
         return ParentPath + "/" + GetJobStepPath(Dep);
     }
 
-    void UpdateNodeHistory(Dictionary<string, NodeHistory> GUBPNodesHistory, NodeInfo NodeToDo, string CLString, string StoreName)
+    static NodeHistory FindNodeHistory(NodeInfo NodeToDo, string CLString, string StoreName)
     {
+		NodeHistory History = null;
+
         if (!NodeToDo.Node.TriggerNode() && CLString != "")
         {
             string GameNameIfAny = NodeToDo.Node.GameNameIfAnyForTempStorage();
             string NodeStoreWildCard = StoreName.Replace(CLString, "*") + "-" + NodeToDo.Node.GetFullName();
-            NodeHistory History = new NodeHistory();
+            History = new NodeHistory();
 
             History.AllStarted = ConvertCLToIntList(TempStorage.FindTempStorageManifests(CmdEnv, NodeStoreWildCard + StartedTempStorageSuffix, false, true, GameNameIfAny));
             History.AllSucceeded = ConvertCLToIntList(TempStorage.FindTempStorageManifests(CmdEnv, NodeStoreWildCard + SucceededTempStorageSuffix, false, true, GameNameIfAny));
@@ -1018,25 +1001,21 @@ public partial class GUBP : BuildCommand
                         History.InProgress.Add(Started);
                         History.InProgressString = GUBPNode.MergeSpaceStrings(History.InProgressString, String.Format("{0}", Started));
                     }
-                }                
-            }
-			if (GUBPNodesHistory.ContainsKey(NodeToDo.Name))
-			{
-				GUBPNodesHistory.Remove(NodeToDo.Name);
+                }
 			}
-			GUBPNodesHistory.Add(NodeToDo.Name, History);
         }
+
+		return History;
     }
-	void GetFailureEmails(Dictionary<string, NodeHistory> GUBPNodesHistory, NodeInfo NodeToDo, string CLString, string StoreName, bool OnlyLateUpdates = false)
+	void GetFailureEmails(NodeInfo NodeToDo, NodeHistory History, string CLString, string StoreName, bool OnlyLateUpdates = false)
 	{
         string EMails;
         string FailCauserEMails = "";
         string EMailNote = "";
         bool SendSuccessForGreenAfterRed = false;
         int NumPeople = 0;
-        if (GUBPNodesHistory.ContainsKey(NodeToDo.Name))
+        if (History != null)
         {
-            NodeHistory History = GUBPNodesHistory[NodeToDo.Name];
 			RunECTool(String.Format("setProperty \"/myWorkflow/LastGreen/{0}\" \"{1}\"", NodeToDo.Name, History.LastSucceeded), true);
 			RunECTool(String.Format("setProperty \"/myWorkflow/LastGreen/{0}\" \"{1}\"", NodeToDo.Name, History.FailedString), true);
 
@@ -1047,7 +1026,7 @@ public partial class GUBP : BuildCommand
                 {
                     if (OnlyLateUpdates)
                     {
-                        LastNonDuplicateFail = FindLastNonDuplicateFail(GUBPNodesHistory, NodeToDo, CLString, StoreName);
+						LastNonDuplicateFail = FindLastNonDuplicateFail(NodeToDo, History, CLString, StoreName);
                         if (LastNonDuplicateFail < P4Env.Changelist)
                         {
                             Log("*** Red-after-red spam reduction, changed CL {0} to CL {1} because the errors didn't change.", P4Env.Changelist, LastNonDuplicateFail);
@@ -1154,9 +1133,8 @@ public partial class GUBP : BuildCommand
         return true;
     }
 
-    int FindLastNonDuplicateFail(Dictionary<string, NodeHistory> GUBPNodesHistory, NodeInfo NodeToDo, string CLString, string StoreName)
+	int FindLastNonDuplicateFail(NodeInfo NodeToDo, NodeHistory History, string CLString, string StoreName)
     {
-        NodeHistory History = GUBPNodesHistory[NodeToDo.Name];
         int Result = P4Env.Changelist;
 
         string GameNameIfAny = NodeToDo.Node.GameNameIfAnyForTempStorage();
@@ -1339,24 +1317,24 @@ public partial class GUBP : BuildCommand
         {
             HostPlatforms.Add(UnrealTargetPlatform.Win64);
         }
-        if (P4Enabled)
-        {
-            BranchName = P4Env.BuildRootP4;
-        }
+		if (P4Enabled)
+		{
+			BranchName = P4Env.BuildRootP4;
+		}
 		else
 		{ 
 			BranchName = ParseParamValue("BranchName", "");
-        }
-        BranchOptions = GetBranchOptions(BranchName);
+		}
+		BranchOptions = GetBranchOptions(BranchName);
         bool WithMac = !BranchOptions.PlatformsToRemove.Contains(UnrealTargetPlatform.Mac);
         if (ParseParam("NoMac"))
         {
             WithMac = false;
         }
         if (WithMac)
-        {
-            HostPlatforms.Add(UnrealTargetPlatform.Mac);
-        }
+		{
+			HostPlatforms.Add(UnrealTargetPlatform.Mac);
+		}
 
 		bool WithLinux = !BranchOptions.PlatformsToRemove.Contains(UnrealTargetPlatform.Linux);
 		bool WithoutLinux = ParseParam("NoLinux");
@@ -1383,9 +1361,6 @@ public partial class GUBP : BuildCommand
         }
         CL = ParseParamInt("CL", 0);
         bool bCleanLocalTempStorage = ParseParam("CleanLocal");
-        bool bChanges = ParseParam("Changes") || ParseParam("AllChanges");
-        bool bHistory = ParseParam("History") || bChanges;
-        bool bListOnly = ParseParam("ListOnly") || bHistory;
         bool bSkipTriggers = ParseParam("SkipTriggers");
         bool bFake = ParseParam("fake");
         bool bFakeEC = ParseParam("FakeEC");
@@ -1397,10 +1372,6 @@ public partial class GUBP : BuildCommand
 
         bool bSaveSharedTempStorage = false;
 
-        if (bHistory && !P4Enabled)
-        {
-            throw new AutomationException("-Changes and -History require -P4.");
-        }
         bool LocalOnly = true;
         string CLString = "";
         if (String.IsNullOrEmpty(StoreName))
@@ -1442,15 +1413,12 @@ public partial class GUBP : BuildCommand
         }
 
         bool CommanderSetup = ParseParam("CommanderJobSetupOnly");
-        string ExplicitTriggerName = "";
-        if (CommanderSetup)
-        {
-            ExplicitTriggerName = ParseParamValue("TriggerNode");
-            if (ExplicitTriggerName == null)
-            {
-                ExplicitTriggerName = "";
-            }
-        }
+
+		string ExplicitTriggerName = "";
+		if (CommanderSetup)
+		{
+			ExplicitTriggerName = ParseParamValue("TriggerNode", "");
+		}
 
         if (ParseParam("CIS") && ExplicitTriggerName == "" && CommanderSetup) // explicit triggers will already have a time index assigned
         {
@@ -1458,23 +1426,6 @@ public partial class GUBP : BuildCommand
             Log("Setting TimeIndex to {0}", TimeIndex);
         }
 
-        LogVerbose("************************* CL:			                    {0}", CL);
-		LogVerbose("************************* P4Enabled:			            {0}", P4Enabled);
-        foreach (UnrealTargetPlatform HostPlatform in HostPlatforms)
-        {
-			LogVerbose("*************************   HostPlatform:		            {0}", HostPlatform.ToString()); 
-        }
-		LogVerbose("************************* StoreName:		                {0}", StoreName.ToString());
-		LogVerbose("************************* bCleanLocalTempStorage:		    {0}", bCleanLocalTempStorage);
-		LogVerbose("************************* bSkipTriggers:		            {0}", bSkipTriggers);
-		LogVerbose("************************* bSaveSharedTempStorage:		    {0}", bSaveSharedTempStorage);
-		LogVerbose("************************* bSignBuildProducts:		        {0}", bSignBuildProducts);
-		LogVerbose("************************* bFake:           		        {0}", bFake);
-		LogVerbose("************************* bFakeEC:           		        {0}", bFakeEC);
-		LogVerbose("************************* bHistory:           		        {0}", bHistory);
-		LogVerbose("************************* TimeIndex:           		    {0}", TimeIndex);        
-
-
         GUBPNodes = new Dictionary<string, NodeInfo>();
 		GUBPAggregates = new Dictionary<string, AggregateInfo>();
         Branch = new BranchInfo(HostPlatforms);        
@@ -1536,28 +1487,11 @@ public partial class GUBP : BuildCommand
             }
         }
 
-        Dictionary<string, NodeHistory> GUBPNodesHistory = new Dictionary<string, NodeHistory>();
-
-        /*if (CLString != "" && StoreName.Contains(CLString) && !ParseParam("NoHistory"))
-        {
-            Log("******* Updating history");
-            var StartTime = DateTime.UtcNow;
-            foreach (var Node in NodesToDo)
-            {
-                if (!NodeIsAlreadyComplete(Node, LocalOnly))
-                {
-                    UpdateNodeHistory(Node, CLString);
-                }				
-            }
-            var BuildDuration = (DateTime.UtcNow - StartTime).TotalMilliseconds;
-            Log("Took {0}s to get history for {1} nodes", BuildDuration / 1000, NodesToDo.Count);
-        }*/
-
         List<NodeInfo> OrderedToDo = TopologicalSort(NodesToDo, ExplicitTrigger, false, false);
 
 		List<NodeInfo> UnfinishedTriggers = FindUnfinishedTriggers(bSkipTriggers, ExplicitTrigger, OrderedToDo);
 
-		PrintNodes(this, OrderedToDo, GUBPAggregates.Values, GUBPNodesHistory, UnfinishedTriggers);
+		PrintNodes(this, OrderedToDo, GUBPAggregates.Values, UnfinishedTriggers);
 
         //check sorting
 		CheckSortOrder(OrderedToDo);
@@ -1565,26 +1499,15 @@ public partial class GUBP : BuildCommand
         string FakeFail = ParseParamValue("FakeFail");
         if(CommanderSetup)
         {
-			Dictionary<string, string> FullNodeDependedOnBy = GetFullNodeDependedOnBy(NodesToDo);
-
-			List<AggregateInfo> SeparatePromotables = FindPromotables(GUBPAggregates.Values);
-			Dictionary<NodeInfo, List<AggregateInfo>> DependentPromotions = FindDependentPromotables(GUBPNodes.Values, SeparatePromotables);
-
-			Dictionary<string, string> FullNodeList;
-			Dictionary<string, string> FullNodeDirectDependencies;
-			GetFullNodeListAndDirectDependencies(out FullNodeList, out FullNodeDirectDependencies);
-
-			Dictionary<string, int> FullNodeListSortKey = GetDisplayOrder(FullNodeList.Keys.ToList(), FullNodeDirectDependencies, GUBPNodes);
-
-			DoCommanderSetup(TimeIndex, bSkipTriggers, bFake, bFakeEC, CLString, ExplicitTrigger, FullNodeList, FullNodeDirectDependencies, FullNodeDependedOnBy, DependentPromotions, SeparatePromotables, FullNodeListSortKey, GUBPNodesHistory, OrderedToDo, UnfinishedTriggers, FakeFail);
+			DoCommanderSetup(NodesToDo, OrderedToDo, TimeIndex, bSkipTriggers, bFake, bFakeEC, CLString, ExplicitTrigger, UnfinishedTriggers, FakeFail);
         }
 		else if(ParseParam("SaveGraph"))
 		{
 			SaveGraphVisualization(OrderedToDo);
 		}
-		else if(!bListOnly)
+		else if(!ParseParam("ListOnly"))
 		{
-			ExecuteNodes(OrderedToDo, bFake, bFakeEC, bSaveSharedTempStorage, GUBPNodesHistory, CLString, StoreName, FakeFail);
+			ExecuteNodes(OrderedToDo, bFake, bFakeEC, bSaveSharedTempStorage, CLString, StoreName, FakeFail);
 		}
         PrintRunTime();
 	}
@@ -2155,8 +2078,18 @@ public partial class GUBP : BuildCommand
 		}
 	}
 
-	private void DoCommanderSetup(int TimeIndex, bool bSkipTriggers, bool bFake, bool bFakeEC, string CLString, NodeInfo ExplicitTrigger, Dictionary<string, string> FullNodeList, Dictionary<string, string> FullNodeDirectDependencies, Dictionary<string, string> FullNodeDependedOnBy, Dictionary<NodeInfo, List<AggregateInfo>> DependentPromotions, List<AggregateInfo> SeparatePromotables, Dictionary<string, int> FullNodeListSortKey, Dictionary<string, NodeHistory> GUBPNodesHistory, List<NodeInfo> OrdereredToDo, List<NodeInfo> UnfinishedTriggers, string FakeFail)
+	private void DoCommanderSetup(HashSet<NodeInfo> NodesToDo, List<NodeInfo> OrdereredToDo, int TimeIndex, bool bSkipTriggers, bool bFake, bool bFakeEC, string CLString, NodeInfo ExplicitTrigger, List<NodeInfo> UnfinishedTriggers, string FakeFail)
 	{
+		Dictionary<string, string> FullNodeDependedOnBy = GetFullNodeDependedOnBy(NodesToDo);
+
+		List<AggregateInfo> SeparatePromotables = FindPromotables(GUBPAggregates.Values);
+		Dictionary<NodeInfo, List<AggregateInfo>> DependentPromotions = FindDependentPromotables(GUBPNodes.Values, SeparatePromotables);
+
+		Dictionary<string, string> FullNodeList;
+		Dictionary<string, string> FullNodeDirectDependencies;
+		GetFullNodeListAndDirectDependencies(out FullNodeList, out FullNodeDirectDependencies);
+
+		Dictionary<string, int> FullNodeListSortKey = GetDisplayOrder(FullNodeList.Keys.ToList(), FullNodeDirectDependencies, GUBPNodes);
 
 		if (OrdereredToDo.Count == 0)
 		{
@@ -2230,7 +2163,7 @@ public partial class GUBP : BuildCommand
 		using(TelemetryStopwatch PrintNodesTimer = new TelemetryStopwatch("SetupCommanderPrint"))
 		{
 			Log("*********** EC Nodes, in order.");
-			PrintNodes(this, OrdereredToDo, GUBPAggregates.Values, GUBPNodesHistory, UnfinishedTriggers);
+			PrintNodes(this, OrdereredToDo, GUBPAggregates.Values, UnfinishedTriggers);
 		}
 		// here we are just making sure everything before the explicit trigger is completed.
 		if (ExplicitTrigger != null)
@@ -2632,7 +2565,7 @@ public partial class GUBP : BuildCommand
 		return RunCondition;
 	}
 
-	void ExecuteNodes(List<NodeInfo> OrdereredToDo, bool bFake, bool bFakeEC, bool bSaveSharedTempStorage, Dictionary<string, NodeHistory> GUBPNodesHistory, string CLString, string StoreName, string FakeFail)
+	void ExecuteNodes(List<NodeInfo> OrdereredToDo, bool bFake, bool bFakeEC, bool bSaveSharedTempStorage, string CLString, string StoreName, string FakeFail)
 	{
         Dictionary<string, NodeInfo> BuildProductToNodeMap = new Dictionary<string, NodeInfo>();
 		foreach (NodeInfo NodeToDo in OrdereredToDo)
@@ -2795,11 +2728,13 @@ public partial class GUBP : BuildCommand
                 }
                 catch (Exception Ex)
                 {
+					NodeHistory History = null;
+
                     if (SaveSuccessRecords)
                     {
 						using(TelemetryStopwatch UpdateNodeHistoryStopwatch = new TelemetryStopwatch("UpdateNodeHistory"))
 						{
-							UpdateNodeHistory(GUBPNodesHistory, NodeToDo, CLString, StoreName);
+							History = FindNodeHistory(NodeToDo, CLString, StoreName);
 						}
 						using(TelemetryStopwatch SaveNodeStatusStopwatch = new TelemetryStopwatch("SaveNodeStatus"))
 						{
@@ -2814,7 +2749,7 @@ public partial class GUBP : BuildCommand
 						{
 							using(TelemetryStopwatch GetFailEmailsStopwatch = new TelemetryStopwatch("GetFailEmails"))
 							{
-								GetFailureEmails(GUBPNodesHistory, NodeToDo, CLString, StoreName);
+								GetFailureEmails(NodeToDo, History, CLString, StoreName);
 							}
 						}
 						UpdateECBuildTime(NodeToDo, BuildDuration);
@@ -2823,9 +2758,8 @@ public partial class GUBP : BuildCommand
                     Log("{0}", ExceptionToString(Ex));
 
 
-                    if (GUBPNodesHistory.ContainsKey(NodeToDo.Name))
+                    if (History != null)
                     {
-                        NodeHistory History = GUBPNodesHistory[NodeToDo.Name];
                         Log("Changes since last green *********************************");
                         Log("");
                         Log("");
@@ -2867,9 +2801,10 @@ public partial class GUBP : BuildCommand
                 }
                 if (SaveSuccessRecords) 
                 {
+					NodeHistory History = null;
 					using(TelemetryStopwatch UpdateNodeHistoryStopwatch = new TelemetryStopwatch("UpdateNodeHistory"))
 					{
-						UpdateNodeHistory(GUBPNodesHistory, NodeToDo, CLString, StoreName);
+						History = FindNodeHistory(NodeToDo, CLString, StoreName);
 					}
 					using(TelemetryStopwatch SaveNodeStatusStopwatch = new TelemetryStopwatch("SaveNodeStatus"))
 					{
@@ -2884,7 +2819,7 @@ public partial class GUBP : BuildCommand
 					{
 						using(TelemetryStopwatch GetFailEmailsStopwatch = new TelemetryStopwatch("GetFailEmails"))
 						{
-							GetFailureEmails(GUBPNodesHistory, NodeToDo, CLString, StoreName);
+							GetFailureEmails(NodeToDo, History, CLString, StoreName);
 						}
 					}
 					UpdateECBuildTime(NodeToDo, BuildDuration);
