Move time quantum to a parameter, rather than querying the branch options directly.

diff --git a/Engine/Source/Programs/AutomationTool/GUBP/GUBP.cs b/Engine/Source/Programs/AutomationTool/GUBP/GUBP.cs
--- a/Engine/Source/Programs/AutomationTool/GUBP/GUBP.cs
+++ b/Engine/Source/Programs/AutomationTool/GUBP/GUBP.cs
@@ -1,4 +1,5 @@
 ﻿// Copyright 1998-2015 Epic Games, Inc. All Rights Reserved.
+
 using System;
 using System.Collections.Generic;
 using System.Text;
@@ -250,28 +251,23 @@ public partial class GUBP : BuildCommand
 		}
 	}
 
-    string CISFrequencyQuantumShiftString(NodeInfo NodeToDo)
+    static string CISFrequencyQuantumShiftString(NodeInfo NodeToDo, int TimeQuantum)
     {
-        string FrequencyString = "";
         int Quantum = NodeToDo.FrequencyShift;
-		if (Quantum > 0)
+		if (Quantum == 0)
         {
-			int TimeQuantum = 20;
-			if(BranchOptions.QuantumOverride != 0)
-			{
-				TimeQuantum = BranchOptions.QuantumOverride;
-			}
-            int Minutes = TimeQuantum * (1 << Quantum);
-            if (Minutes < 60)
-            {
-                FrequencyString = string.Format(" ({0}m)", Minutes);
-            }
-            else
-            {
-                FrequencyString = string.Format(" ({0}h{1}m)", Minutes / 60, Minutes % 60);
-            }
+			return "";
+		}
+
+        int Minutes = TimeQuantum * (1 << Quantum);
+        if (Minutes < 60)
+        {
+            return String.Format(" ({0}m)", Minutes);
+        }
+        else
+        {
+            return String.Format(" ({0}h{1}m)", Minutes / 60, Minutes % 60);
         }
-        return FrequencyString;
     }
 
     int ComputeDependentCISFrequencyQuantumShift(NodeInfo NodeToDo, Dictionary<NodeInfo, int> FrequencyOverrides)
@@ -354,7 +350,7 @@ public partial class GUBP : BuildCommand
             return RunAndLog("ectool", "--timeout 900 " + Args, Options: Opts);
         }
     }
-    void WriteECPerl(List<string> Args)
+    static void WriteECPerl(List<string> Args)
     {
         Args.Add("$batch->submit();");
         string ECPerlFile = CommandUtils.CombinePaths(CommandUtils.CmdEnv.LogFolder, "jobsteps.pl");
@@ -374,7 +370,7 @@ public partial class GUBP : BuildCommand
         return HackFrequency(bp, BranchName, NodeToDo, BaseFrequency);
     }
 
-    List<P4Connection.ChangeRecord> GetChanges(int LastOutputForChanges, int TopCL, int LastGreen)
+    static List<P4Connection.ChangeRecord> GetChanges(int LastOutputForChanges, int TopCL, int LastGreen)
     {
         List<P4Connection.ChangeRecord> Result = new List<P4Connection.ChangeRecord>();
         if (TopCL > LastGreen)
@@ -409,7 +405,7 @@ public partial class GUBP : BuildCommand
         return Result;
     }
 
-    int PrintChanges(int LastOutputForChanges, int TopCL, int LastGreen)
+    static int PrintChanges(int LastOutputForChanges, int TopCL, int LastGreen)
     {
         List<P4Connection.ChangeRecord> ChangeRecords = GetChanges(LastOutputForChanges, TopCL, LastGreen);
         foreach (P4Connection.ChangeRecord Record in ChangeRecords)
@@ -424,7 +420,7 @@ public partial class GUBP : BuildCommand
         return TopCL;
     }
 
-    void PrintDetailedChanges(NodeHistory History, bool bShowAllChanges = false)
+    static void PrintDetailedChanges(NodeHistory History, bool bShowAllChanges = false)
     {
         DateTime StartTime = DateTime.UtcNow;
 
@@ -472,7 +468,7 @@ public partial class GUBP : BuildCommand
 
     }
 
-    void PrintNodes(GUBP bp, List<NodeInfo> Nodes, IEnumerable<AggregateInfo> Aggregates, List<NodeInfo> UnfinishedTriggers = null)
+    void PrintNodes(GUBP bp, List<NodeInfo> Nodes, IEnumerable<AggregateInfo> Aggregates, List<NodeInfo> UnfinishedTriggers, int TimeQuantum)
     {
         bool bShowDependencies = bp.ParseParam("ShowDependencies");
         bool AddEmailProps = bp.ParseParam("ShowEmails");
@@ -533,7 +529,7 @@ public partial class GUBP : BuildCommand
             {
                 MemoryReq = "";                
             }
-            string FrequencyString = CISFrequencyQuantumShiftString(NodeToDo);
+            string FrequencyString = CISFrequencyQuantumShiftString(NodeToDo, TimeQuantum);
 
             Log("      {0}{1}{2}{3}{4}{5}{6} {7}  {8}",
                 (LastAgentGroup != "" ? "  " : ""),
@@ -577,7 +573,7 @@ public partial class GUBP : BuildCommand
 		}
     }
 
-    void SaveGraphVisualization(List<NodeInfo> Nodes)
+    static void SaveGraphVisualization(List<NodeInfo> Nodes)
     {
         List<GraphNode> GraphNodes = new List<GraphNode>();
 
@@ -722,7 +718,7 @@ public partial class GUBP : BuildCommand
         return Result;
     }
 
-    List<NodeInfo> TopologicalSort(HashSet<NodeInfo> NodesToDo, NodeInfo ExplicitTrigger, bool SubSort, bool DoNotConsiderCompletion)
+    static List<NodeInfo> TopologicalSort(HashSet<NodeInfo> NodesToDo, NodeInfo ExplicitTrigger, bool SubSort, bool DoNotConsiderCompletion)
     {
         DateTime StartTime = DateTime.UtcNow;
 
@@ -948,7 +944,7 @@ public partial class GUBP : BuildCommand
         return OrdereredToDo;
     }
 
-    string GetJobStepPath(NodeInfo Dep)
+    static string GetJobStepPath(NodeInfo Dep)
     {
 		if (Dep.Node.AgentSharingGroup == "")
 		{
@@ -959,7 +955,7 @@ public partial class GUBP : BuildCommand
             return "jobSteps[" + Dep.Node.AgentSharingGroup + "]/jobSteps[" + Dep.Name + "]";
         }
     }
-    string GetJobStep(string ParentPath, NodeInfo Dep)
+    static string GetJobStep(string ParentPath, NodeInfo Dep)
     {
         return ParentPath + "/" + GetJobStepPath(Dep);
     }
@@ -1489,9 +1485,15 @@ public partial class GUBP : BuildCommand
 
         List<NodeInfo> OrderedToDo = TopologicalSort(NodesToDo, ExplicitTrigger, false, false);
 
+		int TimeQuantum = 20;
+		if(BranchOptions.QuantumOverride != 0)
+		{
+			TimeQuantum = BranchOptions.QuantumOverride;
+		}
+
 		List<NodeInfo> UnfinishedTriggers = FindUnfinishedTriggers(bSkipTriggers, ExplicitTrigger, OrderedToDo);
 
-		PrintNodes(this, OrderedToDo, GUBPAggregates.Values, UnfinishedTriggers);
+		PrintNodes(this, OrderedToDo, GUBPAggregates.Values, UnfinishedTriggers, TimeQuantum);
 
         //check sorting
 		CheckSortOrder(OrderedToDo);
@@ -1499,7 +1501,7 @@ public partial class GUBP : BuildCommand
         string FakeFail = ParseParamValue("FakeFail");
         if(CommanderSetup)
         {
-			DoCommanderSetup(NodesToDo, OrderedToDo, TimeIndex, bSkipTriggers, bFake, bFakeEC, CLString, ExplicitTrigger, UnfinishedTriggers, FakeFail);
+			DoCommanderSetup(NodesToDo, OrderedToDo, TimeIndex, TimeQuantum, bSkipTriggers, bFake, bFakeEC, CLString, ExplicitTrigger, UnfinishedTriggers, FakeFail);
         }
 		else if(ParseParam("SaveGraph"))
 		{
@@ -1874,21 +1876,22 @@ public partial class GUBP : BuildCommand
 		return FullNodeDependedOnBy;
 	}
 
-	private void GetFullNodeListAndDirectDependencies(out Dictionary<string, string> FullNodeList, out Dictionary<string, string> FullNodeDirectDependencies)
+	private static void GetFullNodeListAndDirectDependencies(IEnumerable<NodeInfo> Nodes, int TimeQuantum, out Dictionary<string, string> FullNodeList, out Dictionary<string, string> FullNodeDirectDependencies)
 	{
 		using(TelemetryStopwatch Timer = new TelemetryStopwatch("GetDependencies"))
 		{
 			FullNodeList = new Dictionary<string,string>();
 			FullNodeDirectDependencies = new Dictionary<string,string>();
 
-			Log("******* {0} GUBP Nodes", GUBPNodes.Count);
-			List<NodeInfo> SortedNodes = TopologicalSort(new HashSet<NodeInfo>(GUBPNodes.Values), null, SubSort: false, DoNotConsiderCompletion: true);
+			List<NodeInfo> SortedNodes = TopologicalSort(new HashSet<NodeInfo>(Nodes), null, SubSort: false, DoNotConsiderCompletion: true);
+			Log("******* {0} GUBP Nodes", SortedNodes.Count);
+
 			foreach (NodeInfo Node in SortedNodes)
 			{
 				string Note = Node.ControllingTriggerDotName;
 				if (Note == "")
 				{
-					Note = CISFrequencyQuantumShiftString(Node);
+					Note = CISFrequencyQuantumShiftString(Node, TimeQuantum);
 				}
 				if (Note == "")
 				{
@@ -2078,7 +2081,7 @@ public partial class GUBP : BuildCommand
 		}
 	}
 
-	private void DoCommanderSetup(HashSet<NodeInfo> NodesToDo, List<NodeInfo> OrdereredToDo, int TimeIndex, bool bSkipTriggers, bool bFake, bool bFakeEC, string CLString, NodeInfo ExplicitTrigger, List<NodeInfo> UnfinishedTriggers, string FakeFail)
+	private void DoCommanderSetup(HashSet<NodeInfo> NodesToDo, List<NodeInfo> OrdereredToDo, int TimeIndex, int TimeQuantum, bool bSkipTriggers, bool bFake, bool bFakeEC, string CLString, NodeInfo ExplicitTrigger, List<NodeInfo> UnfinishedTriggers, string FakeFail)
 	{
 		Dictionary<string, string> FullNodeDependedOnBy = GetFullNodeDependedOnBy(NodesToDo);
 
@@ -2087,7 +2090,7 @@ public partial class GUBP : BuildCommand
 
 		Dictionary<string, string> FullNodeList;
 		Dictionary<string, string> FullNodeDirectDependencies;
-		GetFullNodeListAndDirectDependencies(out FullNodeList, out FullNodeDirectDependencies);
+		GetFullNodeListAndDirectDependencies(GUBPNodes.Values, TimeQuantum, out FullNodeList, out FullNodeDirectDependencies);
 
 		Dictionary<string, int> FullNodeListSortKey = GetDisplayOrder(FullNodeList.Keys.ToList(), FullNodeDirectDependencies, GUBPNodes);
 
@@ -2163,7 +2166,7 @@ public partial class GUBP : BuildCommand
 		using(TelemetryStopwatch PrintNodesTimer = new TelemetryStopwatch("SetupCommanderPrint"))
 		{
 			Log("*********** EC Nodes, in order.");
-			PrintNodes(this, OrdereredToDo, GUBPAggregates.Values, UnfinishedTriggers);
+			PrintNodes(this, OrdereredToDo, GUBPAggregates.Values, UnfinishedTriggers, TimeQuantum);
 		}
 		// here we are just making sure everything before the explicit trigger is completed.
 		if (ExplicitTrigger != null)
