Add Obsolete file Manifest for deployment.  Platforms can delete unreferenced files easily now to clean up the deployment area.

diff --git a/Engine/Source/Programs/AutomationTool/AutomationUtils/DeploymentContext.cs b/Engine/Source/Programs/AutomationTool/AutomationUtils/DeploymentContext.cs
--- a/Engine/Source/Programs/AutomationTool/AutomationUtils/DeploymentContext.cs
+++ b/Engine/Source/Programs/AutomationTool/AutomationUtils/DeploymentContext.cs
@@ -175,6 +175,12 @@ public class DeploymentContext //: ProjectParams
 	static public readonly string UFSDeployDeltaFileName			= "Manifest_DeltaUFSFiles.txt";	
 	static public readonly string NonUFSDeployDeltaFileName			= "Manifest_DeltaNonUFSFiles.txt";
 
+	/// <summary>
+	/// Filename for the manifest of files to delete during deployment.
+	/// </summary>
+	static public readonly string UFSDeployObsoleteFileName = "Manifest_ObsoleteUFSFiles.txt";
+	static public readonly string NonUFSDeployObsoleteFileName = "Manifest_ObsoleteNonUFSFiles.txt";
+
 	/// <summary>
 	/// Filename for the manifest of files currently deployed on a device.
 	/// </summary>
@@ -710,4 +716,14 @@ public class DeploymentContext //: ProjectParams
 	{
 		return Path.Combine(StageDirectory, NonUFSDeployDeltaFileName);
 	}
+
+	public String GetUFSDeploymentObsoletePath()
+	{
+		return Path.Combine(StageDirectory, UFSDeployObsoleteFileName);
+	}
+
+	public String GetNonUFSDeploymentObsoletePath()
+	{
+		return Path.Combine(StageDirectory, NonUFSDeployObsoleteFileName);
+	}
 }
diff --git a/Engine/Source/Programs/AutomationTool/Scripts/CopyBuildToStagingDirectory.Automation.cs b/Engine/Source/Programs/AutomationTool/Scripts/CopyBuildToStagingDirectory.Automation.cs
--- a/Engine/Source/Programs/AutomationTool/Scripts/CopyBuildToStagingDirectory.Automation.cs
+++ b/Engine/Source/Programs/AutomationTool/Scripts/CopyBuildToStagingDirectory.Automation.cs
@@ -1342,6 +1342,29 @@ public partial class Project : CommandUtils
 		return StagedFiles;
 	}
 
+	protected static void WriteObsoleteManifest(ProjectParams Params, DeploymentContext SC, Dictionary<string, string> DeployedFiles, Dictionary<string, string> StagedFiles, string ObsoleteManifest)
+	{
+		List<string> ObsoleteFiles = new List<string>();
+
+		// any file that has been deployed, but is no longer in the staged files is obsolete and should be deleted.
+		foreach (KeyValuePair<string, string> File in DeployedFiles)
+		{
+			if (!StagedFiles.ContainsKey(File.Key))
+			{
+				ObsoleteFiles.Add(File.Key);
+			}
+		}
+
+		// write out to the deltamanifest.json
+		string ManifestFile = CombinePaths(SC.StageDirectory, ObsoleteManifest);
+		StreamWriter Writer = System.IO.File.CreateText(ManifestFile);
+		foreach (string Line in ObsoleteFiles)
+		{
+			Writer.WriteLine(Line);
+		}
+		Writer.Close();
+	}	
+
 	protected static void WriteDeltaManifest(ProjectParams Params, DeploymentContext SC, Dictionary<string, string> DeployedFiles, Dictionary<string, string> StagedFiles, string DeltaManifest)
 	{
 		List<string> CRCFiles = SC.StageTargetPlatform.GetFilesForCRCCheck();
@@ -1558,13 +1581,12 @@ public partial class Project : CommandUtils
 					CreateStagingManifest(Params, SC);
 					ApplyStagingManifest(Params, SC);
 
-					if (Params.IterativeDeploy)
 					{
 						// get the deployed file data
 						Dictionary<string, string> DeployedUFSFiles = new Dictionary<string, string>();
 						Dictionary<string, string> DeployedNonUFSFiles = new Dictionary<string, string>();
 						List<string> UFSManifests;
-						List<string> NonUFSManifests; 
+						List<string> NonUFSManifests;
 						if (SC.StageTargetPlatform.RetrieveDeployedManifests(Params, SC, out UFSManifests, out NonUFSManifests))
 						{
 							DeployedUFSFiles = ReadDeployedManifest(Params, SC, UFSManifests);
@@ -1575,9 +1597,16 @@ public partial class Project : CommandUtils
 						Dictionary<string, string> StagedUFSFiles = ReadStagedManifest(Params, SC, DeploymentContext.UFSDeployedManifestFileName);
 						Dictionary<string, string> StagedNonUFSFiles = ReadStagedManifest(Params, SC, DeploymentContext.NonUFSDeployedManifestFileName);
 
-						// write out the delta file data
-						WriteDeltaManifest(Params, SC, DeployedUFSFiles, StagedUFSFiles, DeploymentContext.UFSDeployDeltaFileName);
-						WriteDeltaManifest(Params, SC, DeployedNonUFSFiles, StagedNonUFSFiles, DeploymentContext.NonUFSDeployDeltaFileName);
+						WriteObsoleteManifest(Params, SC, DeployedUFSFiles, StagedUFSFiles, DeploymentContext.UFSDeployObsoleteFileName);
+						WriteObsoleteManifest(Params, SC, DeployedNonUFSFiles, StagedNonUFSFiles, DeploymentContext.NonUFSDeployObsoleteFileName);
+
+						if (Params.IterativeDeploy)
+						{
+							
+							// write out the delta file data
+							WriteDeltaManifest(Params, SC, DeployedUFSFiles, StagedUFSFiles, DeploymentContext.UFSDeployDeltaFileName);
+							WriteDeltaManifest(Params, SC, DeployedNonUFSFiles, StagedNonUFSFiles, DeploymentContext.NonUFSDeployDeltaFileName);
+						}
 					}
 				}
 			}
