Internally pass email lists around as arrays rather than space delimited strings.

diff --git a/Engine/Source/Programs/AutomationTool/GUBP/GUBP.cs b/Engine/Source/Programs/AutomationTool/GUBP/GUBP.cs
--- a/Engine/Source/Programs/AutomationTool/GUBP/GUBP.cs
+++ b/Engine/Source/Programs/AutomationTool/GUBP/GUBP.cs
@@ -356,14 +356,14 @@ public partial class GUBP : BuildCommand
         string ECPerlFile = CommandUtils.CombinePaths(CommandUtils.CmdEnv.LogFolder, "jobsteps.pl");
         WriteAllLines_NoExceptions(ECPerlFile, Args.ToArray());
     }
-    string GetEMailListForNode(GUBP bp, NodeInfo NodeToDo, string Emails, string Causers)
+    string[] GetEMailListForNode(GUBP bp, NodeInfo NodeToDo, string Causers)
     {        
         string BranchForEmail = "";
         if (P4Enabled)
         {
             BranchForEmail = P4Env.BuildRootP4;
         }
-        return HackEmails(Emails, Causers, BranchForEmail, NodeToDo);
+        return HackEmails(Causers, BranchForEmail, NodeToDo);
     }
     int GetFrequencyForNode(GUBP bp, NodeInfo NodeToDo, int BaseFrequency)
     {
@@ -480,10 +480,10 @@ public partial class GUBP : BuildCommand
         Log("*********** Desired And Dependent Nodes, in order.");
         foreach (NodeInfo NodeToDo in Nodes)
         {
-            string EMails = "";
+            string[] EMails = new string[0];
             if (AddEmailProps)
             {
-                EMails = GetEMailListForNode(bp, NodeToDo, "", "");
+                EMails = GetEMailListForNode(bp, NodeToDo, "");
             }
             if (bShowTriggers)
             {
@@ -540,7 +540,7 @@ public partial class GUBP : BuildCommand
                 NodeToDo.Node.IsSticky() ? " - (Sticky)" : "",				
                 Agent,
                 MemoryReq,
-                EMails,
+                String.Join(" ", EMails),
                 ECProc ? NodeToDo.Node.ECProcedure() : ""
                 );
 
@@ -1005,7 +1005,6 @@ public partial class GUBP : BuildCommand
     }
 	void GetFailureEmails(NodeInfo NodeToDo, NodeHistory History, string CLString, string StoreName, bool OnlyLateUpdates = false)
 	{
-        string EMails;
         string FailCauserEMails = "";
         string EMailNote = "";
         bool SendSuccessForGreenAfterRed = false;
@@ -1082,19 +1081,13 @@ public partial class GUBP : BuildCommand
 		RunECTool(String.Format("setProperty \"/myWorkflow/FailCausers/{0}\" \"{1}\"", NodeToDo.Name, FailCauserEMails));
 		RunECTool(String.Format("setProperty \"/myWorkflow/EmailNotes/{0}\" \"{1}\"", NodeToDo.Name, EMailNote));
         {
-            string AdditionalEmails = "";
 			string Causers = "";
             if (ParseParam("CIS") && !NodeToDo.Node.SendSuccessEmail() && !NodeToDo.Node.TriggerNode())
             {
 				Causers = FailCauserEMails;
            }
-            string AddEmails = ParseParamValue("AddEmails");
-            if (!String.IsNullOrEmpty(AddEmails))
-            {
-                AdditionalEmails = GUBPNode.MergeSpaceStrings(AddEmails, AdditionalEmails);
-            }
-			EMails = GetEMailListForNode(this, NodeToDo, AdditionalEmails, Causers);
-			RunECTool(String.Format("setProperty \"/myWorkflow/FailEmails/{0}\" \"{1}\"", NodeToDo.Name, EMails));            
+			string[] EMails = GetEMailListForNode(this, NodeToDo, Causers);
+			RunECTool(String.Format("setProperty \"/myWorkflow/FailEmails/{0}\" \"{1}\"", NodeToDo.Name, String.Join(" ", EMails)));            
         }
 		if (NodeToDo.Node.SendSuccessEmail() || SendSuccessForGreenAfterRed)
 		{
@@ -1188,19 +1181,10 @@ public partial class GUBP : BuildCommand
         }
         return Result;
     }
-    List<string> GetECPropsForNode(NodeInfo NodeToDo, string CLString, out string EMails, bool OnlyLateUpdates = false)
+    List<string> GetECPropsForNode(NodeInfo NodeToDo, string CLString, string[] EMails, bool OnlyLateUpdates = false)
     {
         List<string> ECProps = new List<string>();
-        EMails = "";		
-		string AdditonalEmails = "";
-		string Causers = "";				
-		string AddEmails = ParseParamValue("AddEmails");
-		if (!String.IsNullOrEmpty(AddEmails))
-		{
-			AdditonalEmails = GUBPNode.MergeSpaceStrings(AddEmails, AdditonalEmails);
-		}
-		EMails = GetEMailListForNode(this, NodeToDo, AdditonalEmails, Causers);
-		ECProps.Add("FailEmails/" + NodeToDo.Name + "=" + EMails);
+		ECProps.Add("FailEmails/" + NodeToDo.Name + "=" + String.Join(" ", EMails));
 	
 		if (!OnlyLateUpdates)
 		{
@@ -1223,8 +1207,8 @@ public partial class GUBP : BuildCommand
         try
         {
             Log("Updating node props for node {0}", NodeToDo.Name);
-			string EMails = "";
-            List<string> Props = GetECPropsForNode(NodeToDo, CLString, out EMails, true);
+			string[] EMails = GetEMailListForNode(this, NodeToDo, "");
+            List<string> Props = GetECPropsForNode(NodeToDo, CLString, EMails, true);
             foreach (string Prop in Props)
             {
                 string[] Parts = Prop.Split("=".ToCharArray());
@@ -2257,8 +2241,8 @@ public partial class GUBP : BuildCommand
 			{
 				if (!NodeToDo.IsComplete) // if something is already finished, we don't put it into EC  
 				{
-					string EMails;
-					List<string> NodeProps = GetECPropsForNode(NodeToDo, CLString, out EMails);
+					string[] EMails = GetEMailListForNode(this, NodeToDo, "");
+					List<string> NodeProps = GetECPropsForNode(NodeToDo, CLString, EMails);
 					ECProps.AddRange(NodeProps);
 
 					bool Sticky = NodeToDo.Node.IsSticky();
@@ -2284,9 +2268,9 @@ public partial class GUBP : BuildCommand
 						Args = Args + ProcedureParams;
 					}
 
-					if ((Procedure == "GUBP_UAT_Trigger" || Procedure == "GUBP_Hardcoded_Trigger") && !String.IsNullOrEmpty(EMails))
+					if ((Procedure == "GUBP_UAT_Trigger" || Procedure == "GUBP_Hardcoded_Trigger") && EMails.Length > 0)
 					{
-						Args = Args + ", {actualParameterName => 'EmailsForTrigger', value => \'" + EMails + "\'}";
+						Args = Args + ", {actualParameterName => 'EmailsForTrigger', value => \'" + String.Join(" ", EMails) + "\'}";
 					}
 					Args = Args + "]";
 
diff --git a/Engine/Source/Programs/AutomationTool/GUBP/LegacyBranchSetup.cs b/Engine/Source/Programs/AutomationTool/GUBP/LegacyBranchSetup.cs
--- a/Engine/Source/Programs/AutomationTool/GUBP/LegacyBranchSetup.cs
+++ b/Engine/Source/Programs/AutomationTool/GUBP/LegacyBranchSetup.cs
@@ -126,10 +126,6 @@ partial class GUBP
         public virtual List<string> ModifyEmail(string Email, GUBP bp, string Branch, string NodeName)
         {
             return new List<string>{Email};
-        }
-        public virtual List<string> FinalizeEmails(List<string> Emails, GUBP bp, string Branch, string NodeName)
-        {
-            return Emails;
         }
 		public virtual bool VetoEmailingCausers(GUBP bp, string Branch, string NodeName)
 		{
@@ -138,13 +134,20 @@ partial class GUBP
     }
 
     private static List<GUBPEmailHacker> EmailHackers;
-    private string HackEmails(string Emails, string Causers, string Branch, NodeInfo NodeInfo)
+    private string[] HackEmails(string Causers, string Branch, NodeInfo NodeInfo)
     {
-        string OnlyEmail = ParseParamValue("OnlyEmail");
+		string OnlyEmail = ParseParamValue("OnlyEmail");
         if (!String.IsNullOrEmpty(OnlyEmail))
         {
-            return OnlyEmail;
+            return new string[]{ OnlyEmail };
         }
+
+		string EmailOnly = ParseParamValue("EmailOnly");
+		if (!String.IsNullOrEmpty(EmailOnly))
+		{
+			return new string[]{ EmailOnly };
+		}
+
         string EmailHint = ParseParamValue("EmailHint");
         if(EmailHint == null)
         {
@@ -170,7 +173,14 @@ partial class GUBP
                 }
             }
         }
-        List<string> Result = new List<string>(Emails.Split(' '));
+        List<string> Result = new List<string>();
+
+        string AddEmails = ParseParamValue("AddEmails");
+        if (!String.IsNullOrEmpty(AddEmails))
+        {
+			Result.AddRange(AddEmails.Split(new char[]{ ' ' }, StringSplitOptions.RemoveEmptyEntries));
+        }
+
 		if(!EmailHackers.Any(x => x.VetoEmailingCausers(this, Branch, NodeInfo.Name)))
 		{
 			Result.AddRange(Causers.Split(' '));
@@ -188,18 +198,7 @@ partial class GUBP
             }
             Result = NewResult;
         }
-        foreach (var EmailHacker in EmailHackers)
-        {
-            Result = EmailHacker.FinalizeEmails(Result, this, Branch, NodeInfo.Name);
-        }
-        string FinalEmails = "";
-        int Count = 0;
-        foreach (var Email in Result)
-        {
-            FinalEmails = GUBPNode.MergeSpaceStrings(FinalEmails, Email);
-            Count++;
-        }
-        return FinalEmails;
+		return Result.ToArray();
     }
     public abstract class GUBPFrequencyHacker
     {
