Cache the full direct and indirect dependencies list, rather than recalculating it all the time.

diff --git a/Engine/Source/Programs/AutomationTool/GUBP/GUBP.cs b/Engine/Source/Programs/AutomationTool/GUBP/GUBP.cs
--- a/Engine/Source/Programs/AutomationTool/GUBP/GUBP.cs
+++ b/Engine/Source/Programs/AutomationTool/GUBP/GUBP.cs
@@ -37,6 +37,8 @@ public partial class GUBP : BuildCommand
 		public GUBPNode Node;
 		public NodeInfo[] Dependencies;
 		public NodeInfo[] PseudoDependencies;
+		public NodeInfo[] AllDirectDependencies;
+		public NodeInfo[] AllIndirectDependencies;
 		public NodeInfo[] ControllingTriggers;
 		public int FrequencyShift = -1;
 		public bool IsComplete;
@@ -46,6 +48,11 @@ public partial class GUBP : BuildCommand
 			get { return String.Join(".", ControllingTriggers.Select(x => x.Name)); }
 		}
 
+		public bool DependsOn(NodeInfo Node)
+		{
+			return AllIndirectDependencies.Contains(Node);
+		}
+
 		public override string ToString()
 		{
 			System.Diagnostics.Trace.TraceWarning("Implicit conversion from NodeInfo to string\n{0}", Environment.StackTrace);
@@ -136,52 +143,6 @@ public partial class GUBP : BuildCommand
         GUBPNodes[Node].Node.FullNamesOfPseudosependencies.Clear();
     }
 
-    static List<NodeInfo> GetDependencies(NodeInfo NodeToDo, bool bFlat = false)
-    {
-        List<NodeInfo> Result = new List<NodeInfo>();
-        foreach (NodeInfo Dependency in NodeToDo.Dependencies)
-        {
-            if (!Result.Contains(Dependency))
-            {
-                Result.Add(Dependency);
-            }
-            if (bFlat)
-            {
-                foreach (NodeInfo RNode in GetDependencies(Dependency, bFlat))
-                {
-                    if (!Result.Contains(RNode))
-                    {
-                        Result.Add(RNode);
-                    }
-                }
-            }
-        }
-        foreach (NodeInfo PseudoDependency in NodeToDo.PseudoDependencies)
-        {
-            if (!Result.Contains(PseudoDependency))
-            {
-                Result.Add(PseudoDependency);
-            }
-            if (bFlat)
-            {
-                foreach (NodeInfo RNode in GetDependencies(PseudoDependency, bFlat))
-                {
-                    if (!Result.Contains(RNode))
-                    {
-                        Result.Add(RNode);
-                    }
-                }
-            }
-        }
-        return Result;
-    }
-
-    bool NodeDependsOn(NodeInfo Rootward, NodeInfo Leafward)
-    {
-        List<NodeInfo> Deps = GetDependencies(Leafward, true);
-        return Deps.Contains(Rootward);
-    }
-
 	/// <summary>
 	/// Recursively update the ControllingTriggers array for each of the nodes passed in. 
 	/// </summary>
@@ -276,7 +237,7 @@ public partial class GUBP : BuildCommand
         if (Result < 0)
         {
             Result = NodeToDo.Node.CISFrequencyQuantumShift(this);
-            Result = GetFrequencyForNode(this, NodeToDo, Result);
+            Result = HackFrequency(this, BranchName, NodeToDo, Result);
 
 			int FrequencyOverride;
 			if(FrequencyOverrides.TryGetValue(NodeToDo, out FrequencyOverride) && Result > FrequencyOverride)
@@ -357,17 +318,8 @@ public partial class GUBP : BuildCommand
         WriteAllLines_NoExceptions(ECPerlFile, Args.ToArray());
     }
     string[] GetEMailListForNode(GUBP bp, NodeInfo NodeToDo, string Causers)
-    {        
-        string BranchForEmail = "";
-        if (P4Enabled)
-        {
-            BranchForEmail = P4Env.BuildRootP4;
-        }
-        return HackEmails(Causers, BranchForEmail, NodeToDo);
-    }
-    int GetFrequencyForNode(GUBP bp, NodeInfo NodeToDo, int BaseFrequency)
     {
-        return HackFrequency(bp, BranchName, NodeToDo, BaseFrequency);
+		return HackEmails(Causers, BranchName, NodeToDo);
     }
 
     static List<P4Connection.ChangeRecord> GetChanges(int LastOutputForChanges, int TopCL, int LastGreen)
@@ -1516,7 +1468,7 @@ public partial class GUBP : BuildCommand
 		Dictionary<NodeInfo, List<AggregateInfo>> DependentPromotions = NodesToDo.ToDictionary(x => x, x => new List<AggregateInfo>());
 		foreach (AggregateInfo SeparatePromotion in SeparatePromotions)
 		{
-			NodeInfo[] Dependencies = SeparatePromotion.Dependencies.SelectMany(x => GetDependencies(x)).Distinct().ToArray();
+			NodeInfo[] Dependencies = SeparatePromotion.Dependencies.SelectMany(x => x.AllIndirectDependencies).Distinct().ToArray();
 			foreach (NodeInfo Dependency in Dependencies)
 			{
 				DependentPromotions[Dependency].Add(SeparatePromotion);
@@ -1575,29 +1527,51 @@ public partial class GUBP : BuildCommand
 
 	private static void LinkNode(NodeInfo BuildNode, Dictionary<string, AggregateInfo> GUBPAggregates, Dictionary<string, NodeInfo> GUBPNodes, ref int NumErrors)
 	{
-		// Find all the dependencies
-		HashSet<NodeInfo> Dependencies = new HashSet<NodeInfo>();
-		foreach (string DependencyName in BuildNode.Node.FullNamesOfDependencies)
+		if(BuildNode.Dependencies == null)
 		{
-			if (!FindDependencies(DependencyName, GUBPAggregates, GUBPNodes, Dependencies))
+			// Find all the dependencies
+			HashSet<NodeInfo> Dependencies = new HashSet<NodeInfo>();
+			foreach (string DependencyName in BuildNode.Node.FullNamesOfDependencies)
 			{
-				CommandUtils.LogError("Node {0} is not in the graph. It is a dependency of {1}.", DependencyName, BuildNode.Name);
-				NumErrors++;
+				if (!FindDependencies(DependencyName, GUBPAggregates, GUBPNodes, Dependencies))
+				{
+					CommandUtils.LogError("Node {0} is not in the graph. It is a dependency of {1}.", DependencyName, BuildNode.Name);
+					NumErrors++;
+				}
 			}
-		}
-		BuildNode.Dependencies = Dependencies.ToArray();
+			BuildNode.Dependencies = Dependencies.ToArray();
 
-		// Find all the pseudo-dependencies
-		HashSet<NodeInfo> PseudoDependencies = new HashSet<NodeInfo>();
-		foreach (string PseudoDependencyName in BuildNode.Node.FullNamesOfPseudosependencies)
-		{
-			if (!FindDependencies(PseudoDependencyName, GUBPAggregates, GUBPNodes, PseudoDependencies))
+			// Find all the pseudo-dependencies
+			HashSet<NodeInfo> PseudoDependencies = new HashSet<NodeInfo>();
+			foreach (string PseudoDependencyName in BuildNode.Node.FullNamesOfPseudosependencies)
 			{
-				CommandUtils.LogError("Node {0} is not in the graph. It is a pseudodependency of {1}.", PseudoDependencyName, BuildNode.Name);
+				if (!FindDependencies(PseudoDependencyName, GUBPAggregates, GUBPNodes, PseudoDependencies))
+				{
+					CommandUtils.LogError("Node {0} is not in the graph. It is a pseudodependency of {1}.", PseudoDependencyName, BuildNode.Name);
+					NumErrors++;
+				}
+			}
+			BuildNode.PseudoDependencies = PseudoDependencies.ToArray();
+
+			// Set the direct dependencies list
+			BuildNode.AllDirectDependencies = BuildNode.Dependencies.Union(BuildNode.PseudoDependencies).ToArray();
+
+			// Recursively find the dependencies for all the dependencies
+			HashSet<NodeInfo> IndirectDependenices = new HashSet<NodeInfo>(BuildNode.AllDirectDependencies);
+			foreach(NodeInfo DirectDependency in BuildNode.AllDirectDependencies)
+			{
+				LinkNode(DirectDependency, GUBPAggregates, GUBPNodes, ref NumErrors);
+				IndirectDependenices.UnionWith(DirectDependency.AllIndirectDependencies);
+			}
+			BuildNode.AllIndirectDependencies = IndirectDependenices.ToArray();
+
+			// Check the node doesn't reference itself
+			if(BuildNode.AllIndirectDependencies.Contains(BuildNode))
+			{
+				CommandUtils.LogError("Node {0} has a dependency on itself.", BuildNode.Name);
 				NumErrors++;
 			}
 		}
-		BuildNode.PseudoDependencies = PseudoDependencies.ToArray();
 	}
 
 	private static bool FindDependencies(string Name, Dictionary<string, AggregateInfo> NameToAggregate, Dictionary<string, NodeInfo> NameToNode, HashSet<NodeInfo> Dependencies)
@@ -1842,8 +1816,7 @@ public partial class GUBP : BuildCommand
 		{
 			if (!NodeToDo.Node.IsTest())
 			{
-				List<NodeInfo> ECDependencies = GetDependencies(NodeToDo);
-				foreach (NodeInfo Dep in ECDependencies)
+				foreach (NodeInfo Dep in NodeToDo.AllDirectDependencies)
 				{
 					string CurrentValue;
 					if (!FullNodeDependedOnBy.TryGetValue(Dep.Name, out CurrentValue) || CurrentValue.Length == 0)
@@ -1882,16 +1855,7 @@ public partial class GUBP : BuildCommand
 					Note = "always";
 				}
 
-				List<NodeInfo> Deps = GetDependencies(Node);
-				string All = "";
-				foreach (NodeInfo Dep in Deps)
-				{
-					if (All != "")
-					{
-						All += " ";
-					}
-					All += Dep.Name;
-				}
+				string All = String.Join(" ", Node.AllDirectDependencies.Select(x => x.Name));
 				LogVerbose("  {0}: {1}      {2}", Node.Name, Note, All);
 				FullNodeList.Add(Node.Name, Note);
 				FullNodeDirectDependencies.Add(Node.Name, All);
@@ -2132,7 +2096,7 @@ public partial class GUBP : BuildCommand
 					// if we are triggering, then remove nodes that are not controlled by the trigger or are dependencies of this trigger
 					if (ExplicitTrigger != null)
 					{
-						if (ExplicitTrigger != NodeToDo && !NodeDependsOn(NodeToDo, ExplicitTrigger) && !NodeDependsOn(ExplicitTrigger, NodeToDo))
+						if (ExplicitTrigger != NodeToDo && !ExplicitTrigger.DependsOn(NodeToDo) && !NodeToDo.DependsOn(ExplicitTrigger))
 						{
 							continue; // this wasn't on the chain related to the trigger we are triggering, so it is not relevant
 						}
@@ -2157,7 +2121,7 @@ public partial class GUBP : BuildCommand
 		{
 			foreach (NodeInfo NodeToDo in OrdereredToDo)
 			{
-				if (!NodeToDo.IsComplete && NodeToDo != ExplicitTrigger && !NodeDependsOn(ExplicitTrigger, NodeToDo)) // if something is already finished, we don't put it into EC
+				if (!NodeToDo.IsComplete && NodeToDo != ExplicitTrigger && !NodeToDo.DependsOn(ExplicitTrigger)) // if something is already finished, we don't put it into EC
 				{
 					throw new AutomationException("We are being asked to process node {0}, however, this is an explicit trigger {1}, so everything before it should already be handled. It seems likely that you waited too long to run the trigger. You will have to do a new build from scratch.", NodeToDo.Name, ExplicitTrigger.Name);
 				}
@@ -2276,8 +2240,7 @@ public partial class GUBP : BuildCommand
 
 					List<NodeInfo> UncompletedEcDeps = new List<NodeInfo>();
 					{
-						List<NodeInfo> EcDeps = GetDependencies(NodeToDo);
-						foreach (NodeInfo Dep in EcDeps)
+						foreach (NodeInfo Dep in NodeToDo.AllDirectDependencies)
 						{
 							if (!Dep.IsComplete && OrdereredToDo.Contains(Dep)) // if something is already finished, we don't put it into EC
 							{
@@ -2462,8 +2425,7 @@ public partial class GUBP : BuildCommand
 				// to avoid idle agents (and also EC doesn't actually reserve our agent!), we promote all dependencies to the first one
 				foreach (NodeInfo Chain in MyChain)
 				{
-					List<NodeInfo> EcDeps = GetDependencies(Chain);
-					foreach (NodeInfo Dep in EcDeps)
+					foreach (NodeInfo Dep in Chain.AllDirectDependencies)
 					{
 						if (!Dep.IsComplete && OrdereredToDo.Contains(Dep)) // if something is already finished, we don't put it into EC
 						{
@@ -2493,8 +2455,7 @@ public partial class GUBP : BuildCommand
 			}
 			else
 			{
-				List<NodeInfo> EcDeps = GetDependencies(NodeToDo);
-				foreach (NodeInfo Dep in EcDeps)
+				foreach (NodeInfo Dep in NodeToDo.AllDirectDependencies)
 				{
 					if (!Dep.IsComplete && OrdereredToDo.Contains(Dep)) // if something is already finished, we don't put it into EC
 					{
