Change frequency calculation to be done by the legacy branch setup.

diff --git a/Engine/Source/Programs/AutomationTool/GUBP/GUBP.cs b/Engine/Source/Programs/AutomationTool/GUBP/GUBP.cs
--- a/Engine/Source/Programs/AutomationTool/GUBP/GUBP.cs
+++ b/Engine/Source/Programs/AutomationTool/GUBP/GUBP.cs
@@ -231,36 +231,20 @@ public partial class GUBP : BuildCommand
         }
     }
 
-    int ComputeDependentCISFrequencyQuantumShift(NodeInfo NodeToDo, Dictionary<NodeInfo, int> FrequencyOverrides)
-    {
-        int Result = NodeToDo.FrequencyShift;        
-        if (Result < 0)
-        {
-            Result = NodeToDo.Node.CISFrequencyQuantumShift(this);
-            Result = HackFrequency(this, BranchName, NodeToDo, Result);
-
-			int FrequencyOverride;
-			if(FrequencyOverrides.TryGetValue(NodeToDo, out FrequencyOverride) && Result > FrequencyOverride)
+	/// <summary>
+	/// Update the frequency shift for all build nodes to ensure that each node's frequency is at least that of its dependencies.
+	/// </summary>
+	/// <param name="Nodes">Sequence of nodes to compute frequencies for</param>
+	static void ComputeDependentFrequencies(IEnumerable<NodeInfo> Nodes)
+	{
+		foreach(NodeInfo Node in Nodes)
+		{
+			foreach(NodeInfo IndirectDependency in Node.AllIndirectDependencies)
 			{
-				Result = FrequencyOverride;
+				Node.FrequencyShift = Math.Max(Node.FrequencyShift, IndirectDependency.FrequencyShift);
 			}
-
-            foreach (NodeInfo Dep in NodeToDo.Dependencies)
-            {
-                Result = Math.Max(ComputeDependentCISFrequencyQuantumShift(Dep, FrequencyOverrides), Result);
-            }
-            foreach (NodeInfo Dep in NodeToDo.PseudoDependencies)
-            {
-                Result = Math.Max(ComputeDependentCISFrequencyQuantumShift(Dep, FrequencyOverrides), Result);
-            }
-            if (Result < 0)
-            {
-                throw new AutomationException("Failed to compute shift.");
-            }
-            NodeToDo.FrequencyShift = Result;
-        }
-        return Result;
-    }
+		}
+	}
 
 	static void FindCompletionState(IEnumerable<NodeInfo> NodesToDo, string StoreName, bool LocalOnly)
 	{
@@ -1367,19 +1351,13 @@ public partial class GUBP : BuildCommand
 		LinkGraph(GUBPAggregates, GUBPNodes);
 		FindControllingTriggers(GUBPNodes.Values);
 		FindCompletionState(GUBPNodes.Values, StoreName, LocalOnly);
+		ComputeDependentFrequencies(GUBPNodes.Values);
 
         if (bCleanLocalTempStorage)  // shared temp storage can never be wiped
         {
             TempStorage.DeleteLocalTempStorageManifests(CmdEnv);
         }
 
-		// Compute all the frequencies
-		Dictionary<NodeInfo, int> FrequencyOverrides = ApplyFrequencyBarriers(GUBPNodes, BranchOptions.FrequencyBarriers);
-        foreach (KeyValuePair<string, NodeInfo> NodeToDo in GUBPNodes)
-        {
-            ComputeDependentCISFrequencyQuantumShift(NodeToDo.Value, FrequencyOverrides);
-        }
-
 		HashSet<NodeInfo> NodesToDo = ParseNodesToDo(WithoutLinux, TimeIndex, CommanderSetup);
 
 		NodeInfo ExplicitTrigger = null;
@@ -1593,58 +1571,6 @@ public partial class GUBP : BuildCommand
 		return false;
 	}
 
-	private static Dictionary<NodeInfo, int> ApplyFrequencyBarriers(Dictionary<string, NodeInfo> GUBPNodes, Dictionary<string, sbyte> FrequencyBarriers)
-	{
-		// Make sure that everything that's listed as a frequency barrier is completed with the given interval
-		Dictionary<NodeInfo, int> FrequencyOverrides = new Dictionary<NodeInfo, int>();
-		foreach (KeyValuePair<string, sbyte> Barrier in FrequencyBarriers)
-		{
-			// Find the barrier node
-			NodeInfo BarrierNode;
-			if(!GUBPNodes.TryGetValue(Barrier.Key, out BarrierNode))
-			{
-				throw new AutomationException("Couldn't find node '{0}' for frequency barrier", Barrier.Key);
-			}
-
-			// All the nodes which are dependencies of the barrier node
-			HashSet<NodeInfo> IncludedNodes = new HashSet<NodeInfo> { BarrierNode };
-
-			// Find all the nodes which are indirect dependencies of this node
-			List<NodeInfo> SearchNodes = new List<NodeInfo> { BarrierNode };
-			for (int Idx = 0; Idx < SearchNodes.Count; Idx++)
-			{
-				NodeInfo Node = SearchNodes[Idx];
-				foreach (NodeInfo Dependency in Node.Dependencies.Union(Node.PseudoDependencies))
-				{
-					if (!IncludedNodes.Contains(Dependency))
-					{
-						IncludedNodes.Add(Dependency);
-						SearchNodes.Add(Dependency);
-					}
-				}
-			}
-
-			// Make sure that everything included in this list is before the cap, and everything not in the list is after it
-			foreach (NodeInfo Node in GUBPNodes.Values)
-			{
-				if (IncludedNodes.Contains(Node))
-				{
-					int Frequency;
-					if (FrequencyOverrides.TryGetValue(Node, out Frequency))
-					{
-						Frequency = Math.Min(Frequency, Barrier.Value);
-					}
-					else
-					{
-						Frequency = Barrier.Value;
-					}
-					FrequencyOverrides[Node] = Frequency;
-				}
-			}
-		}
-		return FrequencyOverrides;
-	}
-
 	private HashSet<NodeInfo> ParseNodesToDo(bool WithoutLinux, int TimeIndex, bool CommanderSetup)
 	{
 		List<string> NamesToDo = new List<string>();
diff --git a/Engine/Source/Programs/AutomationTool/GUBP/LegacyBranchSetup.cs b/Engine/Source/Programs/AutomationTool/GUBP/LegacyBranchSetup.cs
--- a/Engine/Source/Programs/AutomationTool/GUBP/LegacyBranchSetup.cs
+++ b/Engine/Source/Programs/AutomationTool/GUBP/LegacyBranchSetup.cs
@@ -1055,6 +1055,20 @@ partial class GUBP
         }
 #endif
 
+		// Calculate the frequencies for each node
+		Dictionary<string, int> FrequencyOverrides = ApplyFrequencyBarriers(GUBPNodes, GUBPAggregates, BranchOptions.FrequencyBarriers);
+		foreach(NodeInfo Node in GUBPNodes.Values)
+		{
+			Node.FrequencyShift = Node.Node.CISFrequencyQuantumShift(this);
+            Node.FrequencyShift = HackFrequency(this, BranchName, Node, Node.FrequencyShift);
+
+			int FrequencyOverride;
+			if(FrequencyOverrides.TryGetValue(Node.Name, out FrequencyOverride) && FrequencyOverride > Node.FrequencyShift)
+			{
+				Node.FrequencyShift = FrequencyOverride;
+			}
+		}
+
 		// Add aggregate nodes for every project in the branch
 		foreach (BranchInfo.BranchUProject GameProj in Branch.AllProjects)
 		{
@@ -1076,4 +1090,61 @@ partial class GUBP
 			}
 		}
 	}
+
+	private static Dictionary<string, int> ApplyFrequencyBarriers(Dictionary<string, NodeInfo> GUBPNodes, Dictionary<string, AggregateInfo> GUBPAggregates, Dictionary<string, sbyte> FrequencyBarriers)
+	{
+		// Make sure that everything that's listed as a frequency barrier is completed with the given interval
+		Dictionary<string, int> FrequencyOverrides = new Dictionary<string,int>();
+		foreach (KeyValuePair<string, sbyte> Barrier in FrequencyBarriers)
+		{
+			// All the nodes which are dependencies of the barrier node
+			HashSet<string> IncludedNodes = new HashSet<string>();
+
+			// Find all the nodes which are indirect dependencies of this node
+			List<string> SearchNodes = new List<string> { Barrier.Key };
+			for (int Idx = 0; Idx < SearchNodes.Count; Idx++)
+			{
+				if(IncludedNodes.Add(SearchNodes[Idx]))
+				{
+					NodeInfo Node;
+					if(GUBPNodes.TryGetValue(SearchNodes[Idx], out Node))
+					{
+						SearchNodes.AddRange(Node.Node.FullNamesOfDependencies);
+						SearchNodes.AddRange(Node.Node.FullNamesOfPseudosependencies);
+					}
+					else
+					{
+						AggregateInfo Aggregate;
+						if(GUBPAggregates.TryGetValue(SearchNodes[Idx], out Aggregate))
+						{
+							SearchNodes.AddRange(Aggregate.Node.Dependencies);
+						}
+						else
+						{
+							throw new AutomationException("Couldn't find referenced node {0} in graph", Aggregate.Name);
+						}
+					}
+				}
+			}
+
+			// Make sure that everything included in this list is before the cap, and everything not in the list is after it
+			foreach (string NodeName in GUBPNodes.Keys)
+			{
+				if (IncludedNodes.Contains(NodeName))
+				{
+					int Frequency;
+					if(FrequencyOverrides.TryGetValue(NodeName, out Frequency))
+					{
+						Frequency = Math.Min(Frequency, Barrier.Value);
+					}
+					else
+					{
+						Frequency = Barrier.Value;
+					}
+					FrequencyOverrides[NodeName] = Frequency;
+				}
+			}
+		}
+		return FrequencyOverrides;
+	}
 }
