UEPLAT-790: One click integeration with CDNs 	#html5
     This allows the build to be uploaded Amazon S3 during packaging. the game http link can directly shared with other people.  This does not depend on AMZ SDK and only uses REST API.

diff --git a/Engine/Source/Developer/HTML5/HTML5PlatformEditor/Classes/HTML5TargetSettings.h b/Engine/Source/Developer/HTML5/HTML5PlatformEditor/Classes/HTML5TargetSettings.h
--- a/Engine/Source/Developer/HTML5/HTML5PlatformEditor/Classes/HTML5TargetSettings.h
+++ b/Engine/Source/Developer/HTML5/HTML5PlatformEditor/Classes/HTML5TargetSettings.h
@@ -63,4 +63,29 @@ public:
 	UPROPERTY(GlobalConfig, EditAnywhere, Category = Packaging, Meta = (DisplayName = "Level transitions for delta paks [experimental,depends on download maps]"))
 	TArray<FHTML5LevelTransitions> LevelTransitions;
 
+	UPROPERTY(GlobalConfig, EditAnywhere, Category = Amazon_S3, Meta = (DisplayName = "Upload builds to Amazon S3 when packaging"))
+	bool UploadToS3;
+
+	/**
+	* Required
+	*/
+	UPROPERTY(GlobalConfig, EditAnywhere, Category = Amazon_S3, Meta = (DisplayName = "Amazon S3 Key ID"))
+	FString S3KeyID;
+	/**
+	* Required
+	*/
+	UPROPERTY(GlobalConfig, EditAnywhere, Category = Amazon_S3, Meta = (DisplayName = "Amazon S3 Secret Access Key"))
+	FString S3SecretAccessKey;
+	/**
+	* Required
+	*/
+	UPROPERTY(GlobalConfig, EditAnywhere, Category = Amazon_S3, Meta = (DisplayName = "Amazon S3 Bucket Name"))
+	FString S3BucketName;
+	/**
+	* Provide another level of nesting beyond the bucket. Can be left empty, defaults to game name. 
+	*/
+	UPROPERTY(GlobalConfig, EditAnywhere, Category = Amazon_S3, Meta = (DisplayName = "Nested Folder Name"))
+	FString S3FolderName;
+
+
 };
diff --git a/Engine/Source/Programs/AutomationTool/HTML5/HTML5Platform.Automation.cs b/Engine/Source/Programs/AutomationTool/HTML5/HTML5Platform.Automation.cs
--- a/Engine/Source/Programs/AutomationTool/HTML5/HTML5Platform.Automation.cs
+++ b/Engine/Source/Programs/AutomationTool/HTML5/HTML5Platform.Automation.cs
@@ -1,16 +1,19 @@
 ﻿// Copyright 1998-2015 Epic Games, Inc. All Rights Reserved.
 using System;
-using System.Collections.Generic;
-using System.Text.RegularExpressions;
+using Ionic.Zip;
+using System.IO;
+using System.Web;
+using System.Net;
 using System.Linq;
 using System.Text;
-using System.IO;
 using AutomationTool;
 using UnrealBuildTool;
-using System.Diagnostics;
 using System.Threading;
+using System.Diagnostics;
 using System.Threading.Tasks;
-using Ionic.Zip;
+using System.Collections.Generic;
+using System.Security.Cryptography;
+using System.Text.RegularExpressions;
 
 
 public class HTML5Platform : Platform
@@ -329,6 +332,8 @@ public class HTML5Platform : Platform
 				SC.ArchivedFiles.Add(PakFile, DestPak);
 			}	
 		}
+
+		UploadToS3(SC);
 	}
 
 	public override ProcessResult RunClient(ERunOptions ClientRunFlags, string ClientApp, string ClientCmdLine, ProjectParams Params)
@@ -451,4 +456,132 @@ public class HTML5Platform : Platform
 		return ExecutableNames;
 	}
 	#endregion
+
+#region AMAZON S3 
+	public void UploadToS3(DeploymentContext SC)
+	{
+		ConfigCacheIni Ini = new ConfigCacheIni(SC.StageTargetPlatform.PlatformType, "Engine", Path.GetDirectoryName(SC.RawProjectPath));
+		bool Upload = false;
+
+		string KeyId = "";
+		string AccessKey = "";
+		string BucketName = "";
+		string FolderName = "";
+
+		if (Ini.GetBool("/Script/HTML5PlatformEditor.HTML5TargetSettings", "UploadToS3", out Upload))
+		{
+			if (!Upload)
+				return;
+		}
+		else
+		{
+			return; 
+		}
+		bool AmazonIdentity = Ini.GetString("/Script/HTML5PlatformEditor.HTML5TargetSettings", "S3KeyID", out KeyId) &&
+								Ini.GetString("/Script/HTML5PlatformEditor.HTML5TargetSettings", "S3SecretAccessKey", out AccessKey) &&
+								Ini.GetString("/Script/HTML5PlatformEditor.HTML5TargetSettings", "S3BucketName", out BucketName) &&
+								Ini.GetString("/Script/HTML5PlatformEditor.HTML5TargetSettings", "S3FolderName", out FolderName);
+
+		if ( !AmazonIdentity )
+		{
+			Log("Amazon S3 Incorrectly configured"); 
+			return; 
+		}
+
+		if ( FolderName == "" )
+		{
+			FolderName = SC.ShortProjectName;
+		}
+
+		List<Task> UploadTasks = new List<Task>();
+		foreach(KeyValuePair<string, string> Entry in SC.ArchivedFiles)
+		{
+			FileInfo Info = new FileInfo(Entry.Key);
+			UploadTasks.Add (Task.Factory.StartNew(() => UploadToS3Worker(Info,KeyId,AccessKey,BucketName,FolderName))); 
+		}
+
+		Task.WaitAll(UploadTasks.ToArray());
+		Log("Upload Tasks finished.");
+	}
+
+	private static IDictionary<string, string> MimeTypeMapping = new Dictionary<string, string>(StringComparer.InvariantCultureIgnoreCase) 
+		{
+			{ ".html", "text/html"},
+			{ ".js.gz", "application/x-javascript" },  // upload compressed javascript. 
+			{ ".data.gz", "appication/octet-stream"}
+		}; 
+
+	public void UploadToS3Worker(FileInfo Info, string KeyId, string AccessKey, string BucketName, string FolderName )
+	{
+			Log(" Uploading " + Info.Name); 
+
+			// force upload files even if the timestamps haven't changed. 
+			string TimeStamp = string.Format("{0:r}", DateTime.UtcNow);
+			string ContentType = ""; 
+			if (MimeTypeMapping.ContainsKey(Info.Extension))
+			{
+				ContentType = MimeTypeMapping[Info.Extension];
+			}
+			else 
+			{
+				// default
+				ContentType = "application/octet-stream";
+			}
+
+		
+			// URL to put things. 
+			string URL = "http://" + BucketName + ".s3.amazonaws.com/" + FolderName + "/" + Info.Name;
+
+			HttpWebRequest Request = (HttpWebRequest)WebRequest.Create(URL); 
+
+			// Upload. 
+			Request.Method = "PUT";
+			Request.Headers["x-amz-date"] = TimeStamp;
+			Request.Headers["x-amz-acl"] = "public-read"; // we probably want to make public read by default. 
+
+			// set correct content encoding for compressed javascript. 
+			if ( Info.Extension == ".gz")
+			{
+				Request.Headers["Content-Encoding"] = "gzip";
+			}
+
+			Request.ContentType = ContentType;
+			Request.ContentLength = Info.Length;
+
+			//http://docs.aws.amazon.com/AmazonS3/latest/dev/RESTAuthentication.html
+			// Find Base64Encoded data. 
+			UTF8Encoding EncodingMethod = new UTF8Encoding();
+			HMACSHA1 Signature = new HMACSHA1 { Key = EncodingMethod.GetBytes(AccessKey) };
+			// don't change this string.
+			string RequestString = "PUT\n\n" + ContentType + "\n\nx-amz-acl:public-read\nx-amz-date:" + TimeStamp + "\n/"+ BucketName + "/" + FolderName + "/" + Info.Name; 
+			Byte[] ComputedHash = Signature.ComputeHash(EncodingMethod.GetBytes(RequestString));
+			var Base64Encoded = Convert.ToBase64String(ComputedHash);
+			
+			// final amz auth header. 
+			Request.Headers["Authorization"] = "AWS " + KeyId + ":" + Base64Encoded;
+
+			try
+			{
+				// may fail for really big stuff. YMMV. May need Multi part approach, we will see. 
+				Byte[] FileData = File.ReadAllBytes(Info.FullName);
+				var requestStream = Request.GetRequestStream();
+				requestStream.Write(FileData, 0, FileData.Length);
+				requestStream.Close();
+
+				using (var response = Request.GetResponse() as HttpWebResponse)
+				{
+					var reader = new StreamReader(response.GetResponseStream());
+					var data = reader.ReadToEnd();
+				}
+			}
+			catch (Exception ex)
+			{
+				Log("Could not connect to S3, incorrect S3 Keys? " + ex.ToString());
+				throw ex;
+			}
+
+			Log(Info.Name + " has been uploaded " ); 
+	}
+
+#endregion 
 }
