Add some function comments, and tidy up a few functions.

diff --git a/Engine/Source/Programs/AutomationTool/GUBP/GUBP.cs b/Engine/Source/Programs/AutomationTool/GUBP/GUBP.cs
--- a/Engine/Source/Programs/AutomationTool/GUBP/GUBP.cs
+++ b/Engine/Source/Programs/AutomationTool/GUBP/GUBP.cs
@@ -40,7 +40,7 @@ public partial class GUBP : BuildCommand
 		public NodeInfo[] AllDirectDependencies;
 		public NodeInfo[] AllIndirectDependencies;
 		public NodeInfo[] ControllingTriggers;
-		public int FrequencyShift = -1;
+		public int FrequencyShift;
 		public bool IsComplete;
 
 		public string ControllingTriggerDotName
@@ -212,22 +212,20 @@ public partial class GUBP : BuildCommand
 		}
 	}
 
-    static string CISFrequencyQuantumShiftString(NodeInfo NodeToDo, int TimeQuantum)
+	/// <summary>
+	/// Get a string describing a time interval, in the format "XhYm".
+	/// </summary>
+	/// <param name="Minutes">The time interval, in minutes</param>
+	/// <returns>String for how frequently the node should be built</returns>
+    static string GetTimeIntervalString(int Minutes)
     {
-        int Quantum = NodeToDo.FrequencyShift;
-		if (Quantum == 0)
-        {
-			return "";
-		}
-
-        int Minutes = TimeQuantum * (1 << Quantum);
         if (Minutes < 60)
         {
-            return String.Format(" ({0}m)", Minutes);
+            return String.Format("{0}m", Minutes);
         }
         else
         {
-            return String.Format(" ({0}h{1}m)", Minutes / 60, Minutes % 60);
+            return String.Format("{0}h{1}m", Minutes / 60, Minutes % 60);
         }
     }
 
@@ -406,79 +404,102 @@ public partial class GUBP : BuildCommand
 
     void PrintNodes(GUBP bp, List<NodeInfo> Nodes, IEnumerable<AggregateInfo> Aggregates, List<NodeInfo> UnfinishedTriggers, int TimeQuantum)
     {
+		AggregateInfo[] MatchingAggregates = Aggregates.Where(x => x.Dependencies.All(y => Nodes.Contains(y))).ToArray();
+		if (MatchingAggregates.Length > 0)
+		{
+			Log("*********** Aggregates");
+			foreach (AggregateInfo Aggregate in MatchingAggregates.OrderBy(x => x.Name))
+			{
+				StringBuilder Note = new StringBuilder("    " + Aggregate.Name);
+				if (Aggregate.Node.IsPromotableAggregate())
+				{
+					Note.Append(" (promotable)");
+				}
+				Log(Note.ToString());
+			}
+		}
+
         bool bShowDependencies = bp.ParseParam("ShowDependencies");
         bool AddEmailProps = bp.ParseParam("ShowEmails");
         bool ECProc = bp.ParseParam("ShowECProc");
-        bool bShowTriggers = true;
         string LastControllingTrigger = "";
         string LastAgentGroup = "";
 
         Log("*********** Desired And Dependent Nodes, in order.");
         foreach (NodeInfo NodeToDo in Nodes)
         {
-            string[] EMails = new string[0];
-            if (AddEmailProps)
+            string MyControllingTrigger = NodeToDo.ControllingTriggerDotName;
+            if (MyControllingTrigger != LastControllingTrigger)
             {
-                EMails = GetEMailListForNode(bp, NodeToDo, "");
-            }
-            if (bShowTriggers)
-            {
-                string MyControllingTrigger = NodeToDo.ControllingTriggerDotName;
-                if (MyControllingTrigger != LastControllingTrigger)
+                LastControllingTrigger = MyControllingTrigger;
+                if (MyControllingTrigger != "")
                 {
-                    LastControllingTrigger = MyControllingTrigger;
-                    if (MyControllingTrigger != "")
+                    string Finished = "";
+                    if (UnfinishedTriggers != null)
                     {
-                        string Finished = "";
-                        if (UnfinishedTriggers != null)
+                        if (NodeToDo.ControllingTriggers.Length > 0 && UnfinishedTriggers.Contains(NodeToDo.ControllingTriggers.Last()))
                         {
-                            if (NodeToDo.ControllingTriggers.Length > 0 && UnfinishedTriggers.Contains(NodeToDo.ControllingTriggers.Last()))
-                            {
-                                Finished = "(not yet triggered)";
-                            }
-                            else
-                            {
-                                Finished = "(already triggered)";
-                            }
+                            Finished = "(not yet triggered)";
+                        }
+                        else
+                        {
+                            Finished = "(already triggered)";
                         }
-                        Log("  Controlling Trigger: {0}    {1}", MyControllingTrigger, Finished);
                     }
+                    Log("  Controlling Trigger: {0}    {1}", MyControllingTrigger, Finished);
                 }
             }
-            if (NodeToDo.Node.AgentSharingGroup != LastAgentGroup && NodeToDo.Node.AgentSharingGroup != "")
+
+			if (NodeToDo.Node.AgentSharingGroup != LastAgentGroup && NodeToDo.Node.AgentSharingGroup != "")
             {
                 Log("    Agent Group: {0}", NodeToDo.Node.AgentSharingGroup);
             }
             LastAgentGroup = NodeToDo.Node.AgentSharingGroup;
 
+			StringBuilder Builder = new StringBuilder("      ");
+			if(LastAgentGroup != "")
+			{
+				Builder.Append("  ");
+			}
+			Builder.AppendFormat("{0} ({1})", NodeToDo.Name, GetTimeIntervalString(TimeQuantum << NodeToDo.FrequencyShift));
+			if(NodeToDo.IsComplete)
+			{
+				Builder.Append(" - (Completed)");
+			}
+			if(NodeToDo.Node.TriggerNode())
+			{
+				Builder.Append(" - (TriggerNode)");
+			}
+			if(NodeToDo.Node.IsSticky())
+			{
+				Builder.Append(" - (Sticky)");
+			}
+
             string Agent = NodeToDo.Node.ECAgentString();
 			if(ParseParamValue("AgentOverride") != "" && !NodeToDo.Node.GetFullName().Contains("Mac"))
 			{
 				Agent = ParseParamValue("AgentOverride");
 			}
-            if (Agent != "")
+            if (!String.IsNullOrEmpty(Agent))
             {
-                Agent = "[" + Agent + "]";
+				Builder.AppendFormat(" [{0}]", Agent);
             }
-            string MemoryReq = "[" + NodeToDo.Node.AgentMemoryRequirement(bp).ToString() + "]";            
-            if(MemoryReq == "[0]")
+
+			int MemoryRequirement = NodeToDo.Node.AgentMemoryRequirement(bp);
+			if(MemoryRequirement != 0)
+			{
+				Builder.AppendFormat(" [{0}gb]", MemoryRequirement);
+			}
+
+            if (AddEmailProps)
             {
-                MemoryReq = "";                
+				Builder.AppendFormat(" {0}", String.Join(" ", GetEMailListForNode(bp, NodeToDo, "")));
             }
-            string FrequencyString = CISFrequencyQuantumShiftString(NodeToDo, TimeQuantum);
-
-            Log("      {0}{1}{2}{3}{4}{5}{6} {7}  {8}",
-                (LastAgentGroup != "" ? "  " : ""),
-                NodeToDo.Name,
-                FrequencyString,
-                NodeToDo.IsComplete ? " - (Completed)" : "",
-                NodeToDo.Node.TriggerNode() ? " - (TriggerNode)" : "",
-                NodeToDo.Node.IsSticky() ? " - (Sticky)" : "",				
-                Agent,
-                MemoryReq,
-                String.Join(" ", EMails),
-                ECProc ? NodeToDo.Node.ECProcedure() : ""
-                );
+			if(ECProc)
+			{
+				Builder.AppendFormat("  {0}", NodeToDo.Node.ECProcedure());
+			}
+			Log(Builder.ToString());
 
             if (bShowDependencies)
             {
@@ -492,21 +513,6 @@ public partial class GUBP : BuildCommand
                 }
             }
         }
-
-		AggregateInfo[] MatchingAggregates = Aggregates.Where(x => x.Dependencies.All(y => Nodes.Contains(y))).ToArray();
-		if(MatchingAggregates.Length > 0)
-		{
-			Log("*********** Aggregates");
-			foreach(AggregateInfo Aggregate in MatchingAggregates.OrderBy(x => x.Name))
-			{
-				StringBuilder Note = new StringBuilder("    " + Aggregate.Name);
-				if (Aggregate.Node.IsPromotableAggregate())
-				{
-					Note.Append(" (promotable)");
-				}
-				Log(Note.ToString());
-			}
-		}
     }
 
     static void SaveGraphVisualization(List<NodeInfo> Nodes)
@@ -1455,16 +1461,21 @@ public partial class GUBP : BuildCommand
 		return DependentPromotions;
 	}
 
-	private static void LinkGraph(Dictionary<string, AggregateInfo> AggregateNodes, Dictionary<string, NodeInfo> BuildNodes)
+	/// <summary>
+	/// Resolves the names of each node and aggregates' dependencies, and links them together into the build graph.
+	/// </summary>
+	/// <param name="AggregateNameToInfo">Map of aggregate names to their info objects</param>
+	/// <param name="NodeNameToInfo">Map of node names to their info objects</param>
+	private static void LinkGraph(Dictionary<string, AggregateInfo> AggregateNameToInfo, Dictionary<string, NodeInfo> NodeNameToInfo)
 	{
 		int NumErrors = 0;
-		foreach (AggregateInfo AggregateNode in AggregateNodes.Values)
+		foreach (AggregateInfo AggregateNode in AggregateNameToInfo.Values)
 		{
-			LinkAggregate(AggregateNode, AggregateNodes, BuildNodes, ref NumErrors);
+			LinkAggregate(AggregateNode, AggregateNameToInfo, NodeNameToInfo, ref NumErrors);
 		}
-		foreach (NodeInfo BuildNode in BuildNodes.Values)
+		foreach (NodeInfo BuildNode in NodeNameToInfo.Values)
 		{
-			LinkNode(BuildNode, AggregateNodes, BuildNodes, ref NumErrors);
+			LinkNode(BuildNode, AggregateNameToInfo, NodeNameToInfo, ref NumErrors);
 		}
 		if(NumErrors > 0)
 		{
@@ -1472,7 +1483,14 @@ public partial class GUBP : BuildCommand
 		}
 	}
 
-	private static void LinkAggregate(AggregateInfo Aggregate, Dictionary<string, AggregateInfo> AggregateNodes, Dictionary<string, NodeInfo> BuildNodes, ref int NumErrors)
+	/// <summary>
+	/// Resolves the dependency names in an aggregate to NodeInfo instances, filling in the AggregateInfo.Dependenices array. Any referenced aggregates will also be linked, recursively.
+	/// </summary>
+	/// <param name="Aggregate">The aggregate to link</param>
+	/// <param name="AggregateNameToInfo">Map of other aggregate names to their corresponding instance.</param>
+	/// <param name="NodeNameToInfo">Map from node names to their corresponding instance.</param>
+	/// <param name="NumErrors">The number of errors output so far. Incremented if resolving this aggregate fails.</param>
+	private static void LinkAggregate(AggregateInfo Aggregate, Dictionary<string, AggregateInfo> AggregateNameToInfo, Dictionary<string, NodeInfo> NodeNameToInfo, ref int NumErrors)
 	{
 		if (Aggregate.Dependencies == null)
 		{
@@ -1482,15 +1500,15 @@ public partial class GUBP : BuildCommand
 			foreach (string DependencyName in Aggregate.Node.Dependencies)
 			{
 				AggregateInfo AggregateDependency;
-				if(AggregateNodes.TryGetValue(DependencyName, out AggregateDependency))
+				if(AggregateNameToInfo.TryGetValue(DependencyName, out AggregateDependency))
 				{
-					LinkAggregate(AggregateDependency, AggregateNodes, BuildNodes, ref NumErrors);
+					LinkAggregate(AggregateDependency, AggregateNameToInfo, NodeNameToInfo, ref NumErrors);
 					Dependencies.UnionWith(AggregateDependency.Dependencies);
 					continue;
 				}
 
 				NodeInfo Dependency;
-				if(BuildNodes.TryGetValue(DependencyName, out Dependency))
+				if(NodeNameToInfo.TryGetValue(DependencyName, out Dependency))
 				{
 					Dependencies.Add(Dependency);
 					continue;
@@ -1503,66 +1521,81 @@ public partial class GUBP : BuildCommand
 		}
 	}
 
-	private static void LinkNode(NodeInfo BuildNode, Dictionary<string, AggregateInfo> GUBPAggregates, Dictionary<string, NodeInfo> GUBPNodes, ref int NumErrors)
+	/// <summary>
+	/// Resolve a node's dependency names to arrays of NodeInfo instances, filling in the appropriate fields in the NodeInfo object. 
+	/// </summary>
+	/// <param name="Node"></param>
+	/// <param name="AggregateNameToInfo">Map of other aggregate names to their corresponding instance.</param>
+	/// <param name="NodeNameToInfo">Map from node names to their corresponding instance.</param>
+	/// <param name="NumErrors">The number of errors output so far. Incremented if resolving this aggregate fails.</param>
+	private static void LinkNode(NodeInfo Node, Dictionary<string, AggregateInfo> AggregateNameToInfo, Dictionary<string, NodeInfo> NodeNameToInfo, ref int NumErrors)
 	{
-		if(BuildNode.Dependencies == null)
+		if(Node.Dependencies == null)
 		{
 			// Find all the dependencies
 			HashSet<NodeInfo> Dependencies = new HashSet<NodeInfo>();
-			foreach (string DependencyName in BuildNode.Node.FullNamesOfDependencies)
+			foreach (string DependencyName in Node.Node.FullNamesOfDependencies)
 			{
-				if (!FindDependencies(DependencyName, GUBPAggregates, GUBPNodes, Dependencies))
+				if (!ResolveDependencies(DependencyName, AggregateNameToInfo, NodeNameToInfo, Dependencies))
 				{
-					CommandUtils.LogError("Node {0} is not in the graph. It is a dependency of {1}.", DependencyName, BuildNode.Name);
+					CommandUtils.LogError("Node {0} is not in the graph. It is a dependency of {1}.", DependencyName, Node.Name);
 					NumErrors++;
 				}
 			}
-			BuildNode.Dependencies = Dependencies.ToArray();
+			Node.Dependencies = Dependencies.ToArray();
 
 			// Find all the pseudo-dependencies
 			HashSet<NodeInfo> PseudoDependencies = new HashSet<NodeInfo>();
-			foreach (string PseudoDependencyName in BuildNode.Node.FullNamesOfPseudosependencies)
+			foreach (string PseudoDependencyName in Node.Node.FullNamesOfPseudosependencies)
 			{
-				if (!FindDependencies(PseudoDependencyName, GUBPAggregates, GUBPNodes, PseudoDependencies))
+				if (!ResolveDependencies(PseudoDependencyName, AggregateNameToInfo, NodeNameToInfo, PseudoDependencies))
 				{
-					CommandUtils.LogError("Node {0} is not in the graph. It is a pseudodependency of {1}.", PseudoDependencyName, BuildNode.Name);
+					CommandUtils.LogError("Node {0} is not in the graph. It is a pseudodependency of {1}.", PseudoDependencyName, Node.Name);
 					NumErrors++;
 				}
 			}
-			BuildNode.PseudoDependencies = PseudoDependencies.ToArray();
+			Node.PseudoDependencies = PseudoDependencies.ToArray();
 
 			// Set the direct dependencies list
-			BuildNode.AllDirectDependencies = BuildNode.Dependencies.Union(BuildNode.PseudoDependencies).ToArray();
+			Node.AllDirectDependencies = Node.Dependencies.Union(Node.PseudoDependencies).ToArray();
 
 			// Recursively find the dependencies for all the dependencies
-			HashSet<NodeInfo> IndirectDependenices = new HashSet<NodeInfo>(BuildNode.AllDirectDependencies);
-			foreach(NodeInfo DirectDependency in BuildNode.AllDirectDependencies)
+			HashSet<NodeInfo> IndirectDependenices = new HashSet<NodeInfo>(Node.AllDirectDependencies);
+			foreach(NodeInfo DirectDependency in Node.AllDirectDependencies)
 			{
-				LinkNode(DirectDependency, GUBPAggregates, GUBPNodes, ref NumErrors);
+				LinkNode(DirectDependency, AggregateNameToInfo, NodeNameToInfo, ref NumErrors);
 				IndirectDependenices.UnionWith(DirectDependency.AllIndirectDependencies);
 			}
-			BuildNode.AllIndirectDependencies = IndirectDependenices.ToArray();
+			Node.AllIndirectDependencies = IndirectDependenices.ToArray();
 
 			// Check the node doesn't reference itself
-			if(BuildNode.AllIndirectDependencies.Contains(BuildNode))
+			if(Node.AllIndirectDependencies.Contains(Node))
 			{
-				CommandUtils.LogError("Node {0} has a dependency on itself.", BuildNode.Name);
+				CommandUtils.LogError("Node {0} has a dependency on itself.", Node.Name);
 				NumErrors++;
 			}
 		}
 	}
 
-	private static bool FindDependencies(string Name, Dictionary<string, AggregateInfo> NameToAggregate, Dictionary<string, NodeInfo> NameToNode, HashSet<NodeInfo> Dependencies)
+	/// <summary>
+	/// Adds all the nodes matching a given name to a hash set, expanding any aggregates to their dependencices.
+	/// </summary>
+	/// <param name="Name">The name to look for</param>
+	/// <param name="AggregateNameToInfo">Map of other aggregate names to their corresponding info instance.</param>
+	/// <param name="NodeNameToInfo">Map from node names to their corresponding info instance.</param>
+	/// <param name="Dependencies">The set of dependencies to add to.</param>
+	/// <returns>True if the name was found (and the dependencies list was updated), false otherwise.</returns>
+	private static bool ResolveDependencies(string Name, Dictionary<string, AggregateInfo> AggregateNameToInfo, Dictionary<string, NodeInfo> NodeNameToInfo, HashSet<NodeInfo> Dependencies)
 	{
 		AggregateInfo AggregateDependency;
-		if (NameToAggregate.TryGetValue(Name, out AggregateDependency))
+		if (AggregateNameToInfo.TryGetValue(Name, out AggregateDependency))
 		{
 			Dependencies.UnionWith(AggregateDependency.Dependencies);
 			return true;
 		}
 
 		NodeInfo NodeDependency;
-		if (NameToNode.TryGetValue(Name, out NodeDependency))
+		if (NodeNameToInfo.TryGetValue(Name, out NodeDependency))
 		{
 			Dependencies.Add(NodeDependency);
 			return true;
@@ -1774,11 +1807,14 @@ public partial class GUBP : BuildCommand
 				string Note = Node.ControllingTriggerDotName;
 				if (Note == "")
 				{
-					Note = CISFrequencyQuantumShiftString(Node, TimeQuantum);
-				}
-				if (Note == "")
-				{
-					Note = "always";
+					if(Node.FrequencyShift == 0)
+					{
+						Note = "always";
+					}
+					else
+					{
+						Note = String.Format(" ({0})", GetTimeIntervalString(TimeQuantum << Node.FrequencyShift));
+					}
 				}
 
 				string All = String.Join(" ", Node.AllDirectDependencies.Select(x => x.Name));
