Fix compile failure on Mac due to unused variable warning being treated as an error.

diff --git a/Engine/Source/Programs/AutomationTool/HTML5/HTML5Platform.Automation.cs b/Engine/Source/Programs/AutomationTool/HTML5/HTML5Platform.Automation.cs
--- a/Engine/Source/Programs/AutomationTool/HTML5/HTML5Platform.Automation.cs
+++ b/Engine/Source/Programs/AutomationTool/HTML5/HTML5Platform.Automation.cs
@@ -571,7 +571,7 @@ public class HTML5Platform : Platform
 				using (var response = Request.GetResponse() as HttpWebResponse)
 				{
 					var reader = new StreamReader(response.GetResponseStream());
-					var data = reader.ReadToEnd();
+					reader.ReadToEnd();
 				}
 			}
 			catch (Exception ex)
