Move email lists into node properties.

diff --git a/Engine/Source/Programs/AutomationTool/GUBP/GUBP.cs b/Engine/Source/Programs/AutomationTool/GUBP/GUBP.cs
--- a/Engine/Source/Programs/AutomationTool/GUBP/GUBP.cs
+++ b/Engine/Source/Programs/AutomationTool/GUBP/GUBP.cs
@@ -42,6 +42,8 @@ public partial class GUBP : BuildCommand
 		public NodeInfo[] ControllingTriggers;
 		public int FrequencyShift;
 		public bool IsComplete;
+		public string[] RecipientsForFailureEmails;
+		public bool AddSubmittersToFailureEmails;
 
 		public string ControllingTriggerDotName
 		{
@@ -299,10 +301,6 @@ public partial class GUBP : BuildCommand
         string ECPerlFile = CommandUtils.CombinePaths(CommandUtils.CmdEnv.LogFolder, "jobsteps.pl");
         WriteAllLines_NoExceptions(ECPerlFile, Args.ToArray());
     }
-    string[] GetEMailListForNode(GUBP bp, NodeInfo NodeToDo, string Causers)
-    {
-		return HackEmails(Causers, BranchName, NodeToDo);
-    }
 
     static List<P4Connection.ChangeRecord> GetChanges(int LastOutputForChanges, int TopCL, int LastGreen)
     {
@@ -493,7 +491,7 @@ public partial class GUBP : BuildCommand
 
             if (AddEmailProps)
             {
-				Builder.AppendFormat(" {0}", String.Join(" ", GetEMailListForNode(bp, NodeToDo, "")));
+				Builder.AppendFormat(" {0}", String.Join(" ", NodeToDo.RecipientsForFailureEmails));
             }
 			if(ECProc)
 			{
@@ -1023,13 +1021,12 @@ public partial class GUBP : BuildCommand
 		RunECTool(String.Format("setProperty \"/myWorkflow/FailCausers/{0}\" \"{1}\"", NodeToDo.Name, FailCauserEMails));
 		RunECTool(String.Format("setProperty \"/myWorkflow/EmailNotes/{0}\" \"{1}\"", NodeToDo.Name, EMailNote));
         {
-			string Causers = "";
-            if (ParseParam("CIS") && !NodeToDo.Node.SendSuccessEmail() && !NodeToDo.Node.TriggerNode())
+			HashSet<string> Emails = new HashSet<string>(NodeToDo.RecipientsForFailureEmails);
+            if (ParseParam("CIS") && !NodeToDo.Node.SendSuccessEmail() && !NodeToDo.Node.TriggerNode() && NodeToDo.AddSubmittersToFailureEmails)
             {
-				Causers = FailCauserEMails;
-           }
-			string[] EMails = GetEMailListForNode(this, NodeToDo, Causers);
-			RunECTool(String.Format("setProperty \"/myWorkflow/FailEmails/{0}\" \"{1}\"", NodeToDo.Name, String.Join(" ", EMails)));            
+				Emails.UnionWith(FailCauserEMails.Split(new char[] { ' ' }, StringSplitOptions.RemoveEmptyEntries));
+			}
+			RunECTool(String.Format("setProperty \"/myWorkflow/FailEmails/{0}\" \"{1}\"", NodeToDo.Name, String.Join(" ", Emails)));            
         }
 		if (NodeToDo.Node.SendSuccessEmail() || SendSuccessForGreenAfterRed)
 		{
@@ -1123,10 +1120,10 @@ public partial class GUBP : BuildCommand
         }
         return Result;
     }
-    List<string> GetECPropsForNode(NodeInfo NodeToDo, string CLString, string[] EMails, bool OnlyLateUpdates = false)
+    List<string> GetECPropsForNode(NodeInfo NodeToDo, string CLString, bool OnlyLateUpdates)
     {
         List<string> ECProps = new List<string>();
-		ECProps.Add("FailEmails/" + NodeToDo.Name + "=" + String.Join(" ", EMails));
+		ECProps.Add("FailEmails/" + NodeToDo.Name + "=" + String.Join(" ", NodeToDo.RecipientsForFailureEmails));
 	
 		if (!OnlyLateUpdates)
 		{
@@ -1149,8 +1146,7 @@ public partial class GUBP : BuildCommand
         try
         {
             Log("Updating node props for node {0}", NodeToDo.Name);
-			string[] EMails = GetEMailListForNode(this, NodeToDo, "");
-            List<string> Props = GetECPropsForNode(NodeToDo, CLString, EMails, true);
+            List<string> Props = GetECPropsForNode(NodeToDo, CLString, true);
             foreach (string Prop in Props)
             {
                 string[] Parts = Prop.Split("=".ToCharArray());
@@ -2167,8 +2163,7 @@ public partial class GUBP : BuildCommand
 			{
 				if (!NodeToDo.IsComplete) // if something is already finished, we don't put it into EC  
 				{
-					string[] EMails = GetEMailListForNode(this, NodeToDo, "");
-					List<string> NodeProps = GetECPropsForNode(NodeToDo, CLString, EMails);
+					List<string> NodeProps = GetECPropsForNode(NodeToDo, CLString, false);
 					ECProps.AddRange(NodeProps);
 
 					bool Sticky = NodeToDo.Node.IsSticky();
@@ -2194,9 +2189,9 @@ public partial class GUBP : BuildCommand
 						Args = Args + ProcedureParams;
 					}
 
-					if ((Procedure == "GUBP_UAT_Trigger" || Procedure == "GUBP_Hardcoded_Trigger") && EMails.Length > 0)
+					if ((Procedure == "GUBP_UAT_Trigger" || Procedure == "GUBP_Hardcoded_Trigger") && NodeToDo.RecipientsForFailureEmails.Length > 0)
 					{
-						Args = Args + ", {actualParameterName => 'EmailsForTrigger', value => \'" + String.Join(" ", EMails) + "\'}";
+						Args = Args + ", {actualParameterName => 'EmailsForTrigger', value => \'" + String.Join(" ", NodeToDo.RecipientsForFailureEmails) + "\'}";
 					}
 					Args = Args + "]";
 
diff --git a/Engine/Source/Programs/AutomationTool/GUBP/LegacyBranchSetup.cs b/Engine/Source/Programs/AutomationTool/GUBP/LegacyBranchSetup.cs
--- a/Engine/Source/Programs/AutomationTool/GUBP/LegacyBranchSetup.cs
+++ b/Engine/Source/Programs/AutomationTool/GUBP/LegacyBranchSetup.cs
@@ -133,19 +133,47 @@ partial class GUBP
 		}
     }
 
-    private static List<GUBPEmailHacker> EmailHackers;
-    private string[] HackEmails(string Causers, string Branch, NodeInfo NodeInfo)
+	private void FindEmailsForNodes(IEnumerable<NodeInfo> Nodes)
+	{
+		List<GUBPEmailHacker> EmailHackers = new List<GUBPEmailHacker>();
+
+		Assembly[] LoadedAssemblies = AppDomain.CurrentDomain.GetAssemblies();
+		foreach (var Dll in LoadedAssemblies)
+		{
+			Type[] AllTypes = GetTypesFromAssembly(Dll);
+			foreach (var PotentialConfigType in AllTypes)
+			{
+				if (PotentialConfigType != typeof(GUBPEmailHacker) && typeof(GUBPEmailHacker).IsAssignableFrom(PotentialConfigType))
+				{
+					GUBPEmailHacker Config = Activator.CreateInstance(PotentialConfigType) as GUBPEmailHacker;
+					if (Config != null)
+					{
+						EmailHackers.Add(Config);
+					}
+				}
+			}
+		}
+
+		foreach (NodeInfo Node in GUBPNodes.Values)
+		{
+			Node.RecipientsForFailureEmails = HackEmails(EmailHackers, BranchName, Node, out Node.AddSubmittersToFailureEmails);
+		}
+	}
+
+    private string[] HackEmails(List<GUBPEmailHacker> EmailHackers, string Branch, NodeInfo NodeInfo, out bool bIncludeCausers)
     {
 		string OnlyEmail = ParseParamValue("OnlyEmail");
         if (!String.IsNullOrEmpty(OnlyEmail))
         {
+			bIncludeCausers = false;
             return new string[]{ OnlyEmail };
         }
 
 		string EmailOnly = ParseParamValue("EmailOnly");
 		if (!String.IsNullOrEmpty(EmailOnly))
 		{
-			return new string[]{ EmailOnly };
+			bIncludeCausers = false;
+			return new string[] { EmailOnly };
 		}
 
         string EmailHint = ParseParamValue("EmailHint");
@@ -153,26 +181,7 @@ partial class GUBP
         {
             EmailHint = "";
         }			
-        if (EmailHackers == null)
-        {
-            EmailHackers = new List<GUBPEmailHacker>();
-            Assembly[] LoadedAssemblies = AppDomain.CurrentDomain.GetAssemblies();
-            foreach (var Dll in LoadedAssemblies)
-            {
-				Type[] AllTypes = GetTypesFromAssembly(Dll);
-                foreach (var PotentialConfigType in AllTypes)
-                {
-                    if (PotentialConfigType != typeof(GUBPEmailHacker) && typeof(GUBPEmailHacker).IsAssignableFrom(PotentialConfigType))
-                    {
-                        GUBPEmailHacker Config = Activator.CreateInstance(PotentialConfigType) as GUBPEmailHacker;
-                        if (Config != null)
-                        {
-                            EmailHackers.Add(Config);
-                        }
-                    }
-                }
-            }
-        }
+
         List<string> Result = new List<string>();
 
         string AddEmails = ParseParamValue("AddEmails");
@@ -181,10 +190,6 @@ partial class GUBP
 			Result.AddRange(AddEmails.Split(new char[]{ ' ' }, StringSplitOptions.RemoveEmptyEntries));
         }
 
-		if(!EmailHackers.Any(x => x.VetoEmailingCausers(this, Branch, NodeInfo.Name)))
-		{
-			Result.AddRange(Causers.Split(' '));
-		}
         foreach (var EmailHacker in EmailHackers)
         {
             Result.AddRange(EmailHacker.AddEmails(this, Branch, NodeInfo.Name, EmailHint));
@@ -198,6 +203,8 @@ partial class GUBP
             }
             Result = NewResult;
         }
+
+		bIncludeCausers = !EmailHackers.Any(x => x.VetoEmailingCausers(this, Branch, NodeInfo.Name));
 		return Result.ToArray();
     }
     public abstract class GUBPFrequencyHacker
@@ -1069,6 +1076,9 @@ partial class GUBP
 			}
 		}
 
+		// Get the email list for each node
+		FindEmailsForNodes(GUBPNodes.Values);
+
 		// Add aggregate nodes for every project in the branch
 		foreach (BranchInfo.BranchUProject GameProj in Branch.AllProjects)
 		{
