Remove string lookup for EC display order, and GetFullNodeListAndDirectDependencies function.

diff --git a/Engine/Source/Programs/AutomationTool/GUBP/GUBP.cs b/Engine/Source/Programs/AutomationTool/GUBP/GUBP.cs
--- a/Engine/Source/Programs/AutomationTool/GUBP/GUBP.cs
+++ b/Engine/Source/Programs/AutomationTool/GUBP/GUBP.cs
@@ -1788,39 +1788,6 @@ public partial class GUBP : BuildCommand
 		return FullNodeDependedOnBy;
 	}
 
-	private static void GetFullNodeListAndDirectDependencies(IEnumerable<NodeInfo> Nodes, int TimeQuantum, out Dictionary<string, string> FullNodeList, out Dictionary<string, string> FullNodeDirectDependencies)
-	{
-		using(TelemetryStopwatch Timer = new TelemetryStopwatch("GetDependencies"))
-		{
-			FullNodeList = new Dictionary<string,string>();
-			FullNodeDirectDependencies = new Dictionary<string,string>();
-
-			List<NodeInfo> SortedNodes = TopologicalSort(new HashSet<NodeInfo>(Nodes), null, SubSort: false, DoNotConsiderCompletion: true);
-			Log("******* {0} GUBP Nodes", SortedNodes.Count);
-
-			foreach (NodeInfo Node in SortedNodes)
-			{
-				string Note = Node.ControllingTriggerDotName;
-				if (Note == "")
-				{
-					if(Node.FrequencyShift == 0)
-					{
-						Note = "always";
-					}
-					else
-					{
-						Note = String.Format(" ({0})", GetTimeIntervalString(TimeQuantum << Node.FrequencyShift));
-					}
-				}
-
-				string All = String.Join(" ", Node.AllDirectDependencies.Select(x => x.Name));
-				LogVerbose("  {0}: {1}      {2}", Node.Name, Note, All);
-				FullNodeList.Add(Node.Name, Note);
-				FullNodeDirectDependencies.Add(Node.Name, All);
-			}
-		}
-	}
-
 	private int UpdateCISCounter()
 	{
 		if (!P4Enabled)
@@ -1994,11 +1961,10 @@ public partial class GUBP : BuildCommand
 		List<AggregateInfo> SeparatePromotables = FindPromotables(GUBPAggregates.Values);
 		Dictionary<NodeInfo, List<AggregateInfo>> DependentPromotions = FindDependentPromotables(GUBPNodes.Values, SeparatePromotables);
 
-		Dictionary<string, string> FullNodeList;
-		Dictionary<string, string> FullNodeDirectDependencies;
-		GetFullNodeListAndDirectDependencies(GUBPNodes.Values, TimeQuantum, out FullNodeList, out FullNodeDirectDependencies);
+		List<NodeInfo> SortedNodes = TopologicalSort(new HashSet<NodeInfo>(GUBPNodes.Values), null, SubSort: false, DoNotConsiderCompletion: true);
+		Log("******* {0} GUBP Nodes", SortedNodes.Count);
 
-		Dictionary<string, int> FullNodeListSortKey = GetDisplayOrder(FullNodeList.Keys.ToList(), FullNodeDirectDependencies, GUBPNodes);
+		Dictionary<NodeInfo, int> FullNodeListSortKey = GetDisplayOrder(SortedNodes);
 
 		if (OrdereredToDo.Count == 0)
 		{
@@ -2006,17 +1972,17 @@ public partial class GUBP : BuildCommand
 		}
 		List<string> ECProps = new List<string>();
 		ECProps.Add(String.Format("TimeIndex={0}", TimeIndex));
-		foreach (KeyValuePair<string, string> NodePair in FullNodeList)
+		foreach (NodeInfo Node in SortedNodes)
 		{
-			ECProps.Add(string.Format("AllNodes/{0}={1}", NodePair.Key, NodePair.Value));
+			ECProps.Add(string.Format("AllNodes/{0}={1}", Node.Name, GetNodeForAllNodesProperty(Node, TimeQuantum)));
 		}
-		foreach (KeyValuePair<string, string> NodePair in FullNodeDirectDependencies)
+		foreach (NodeInfo Node in SortedNodes)
 		{
-			ECProps.Add(string.Format("DirectDependencies/{0}={1}", NodePair.Key, NodePair.Value));
+			ECProps.Add(string.Format("DirectDependencies/{0}={1}", Node.Name, String.Join(" ", Node.AllDirectDependencies.Select(x => x.Name))));
 		}
-		foreach (KeyValuePair<string, int> NodePair in FullNodeListSortKey)
+		foreach (KeyValuePair<NodeInfo, int> NodePair in FullNodeListSortKey)
 		{
-			ECProps.Add(string.Format("SortKey/{0}={1}", NodePair.Key, NodePair.Value));
+			ECProps.Add(string.Format("SortKey/{0}={1}", NodePair.Key.Name, NodePair.Value));
 		}
 		foreach (KeyValuePair<string, string> NodePair in FullNodeDependedOnBy)
 		{
@@ -2447,7 +2413,24 @@ public partial class GUBP : BuildCommand
 		return PreCondition;
 	}
 
-	private string GetRunConditionForNode(List<NodeInfo> UncompletedEcDeps, string PreconditionParentPath)
+	private static string GetNodeForAllNodesProperty(NodeInfo Node, int TimeQuantum)
+	{
+		string Note = Node.ControllingTriggerDotName;
+		if (Note == "")
+		{
+			if (Node.FrequencyShift == 0)
+			{
+				Note = "always";
+			}
+			else
+			{
+				Note = String.Format(" ({0})", GetTimeIntervalString(TimeQuantum << Node.FrequencyShift));
+			}
+		}
+		return Note;
+	}
+
+	private static string GetRunConditionForNode(List<NodeInfo> UncompletedEcDeps, string PreconditionParentPath)
 	{
 		string RunCondition = "";
 		if (UncompletedEcDeps.Count > 0)
@@ -2746,79 +2729,79 @@ public partial class GUBP : BuildCommand
 	/// Sorts a list of nodes to display in EC. The default order is based on execution order and agent groups, whereas this function arranges nodes by
 	/// frequency then execution order, while trying to group nodes on parallel paths (eg. Mac/Windows editor nodes) together.
 	/// </summary>
-	static Dictionary<string, int> GetDisplayOrder(List<string> NodeNames, Dictionary<string, string> InitialNodeDependencyNames, Dictionary<string, NodeInfo> GUBPNodes)
+	static Dictionary<NodeInfo, int> GetDisplayOrder(List<NodeInfo> Nodes)
 	{
 		// Split the nodes into separate lists for each frequency
-		SortedDictionary<int, List<string>> NodesByFrequency = new SortedDictionary<int,List<string>>();
-		foreach(string NodeName in NodeNames)
+		SortedDictionary<int, List<NodeInfo>> NodesByFrequency = new SortedDictionary<int,List<NodeInfo>>();
+		foreach(NodeInfo Node in Nodes)
 		{
-			List<string> NodesByThisFrequency;
-			if(!NodesByFrequency.TryGetValue(GUBPNodes[NodeName].FrequencyShift, out NodesByThisFrequency))
+			List<NodeInfo> NodesByThisFrequency;
+			if(!NodesByFrequency.TryGetValue(Node.FrequencyShift, out NodesByThisFrequency))
 			{
-				NodesByThisFrequency = new List<string>();
-				NodesByFrequency.Add(GUBPNodes[NodeName].FrequencyShift, NodesByThisFrequency);
+				NodesByThisFrequency = new List<NodeInfo>();
+				NodesByFrequency.Add(Node.FrequencyShift, NodesByThisFrequency);
 			}
-			NodesByThisFrequency.Add(NodeName);
+			NodesByThisFrequency.Add(Node);
 		}
 
 		// Build the output list by scanning each frequency in order
-		HashSet<string> VisitedNodes = new HashSet<string>();
-		Dictionary<string, int> SortedNodes = new Dictionary<string,int>();
-		foreach(List<string> NodesByThisFrequency in NodesByFrequency.Values)
+		HashSet<NodeInfo> VisitedNodes = new HashSet<NodeInfo>();
+		Dictionary<NodeInfo, int> SortedNodes = new Dictionary<NodeInfo,int>();
+		foreach(List<NodeInfo> NodesByThisFrequency in NodesByFrequency.Values)
 		{
 			// Find a list of nodes in each display group. If the group name matches the node name, put that node at the front of the list.
-			Dictionary<string, string> DisplayGroups = new Dictionary<string,string>();
-			foreach(string NodeName in NodesByThisFrequency)
+			Dictionary<string, List<NodeInfo>> DisplayGroups = new Dictionary<string,List<NodeInfo>>();
+			foreach(NodeInfo Node in NodesByThisFrequency)
 			{
-				string GroupName = GUBPNodes[NodeName].Node.GetDisplayGroupName();
+				string GroupName = Node.Node.GetDisplayGroupName();
 				if(!DisplayGroups.ContainsKey(GroupName))
 				{
-					DisplayGroups.Add(GroupName, NodeName);
+					DisplayGroups.Add(GroupName, new List<NodeInfo>{ Node });
 				}
-				else if(GroupName == NodeName)
+				else if(GroupName == Node.Name)
 				{
-					DisplayGroups[GroupName] = NodeName + " " + DisplayGroups[GroupName];
+					DisplayGroups[GroupName].Insert(0, Node);
 				}
 				else
 				{
-					DisplayGroups[GroupName] = DisplayGroups[GroupName] + " " + NodeName;
+					DisplayGroups[GroupName].Add(Node);
 				}
 			}
 
 			// Build a list of ordering dependencies, putting all Mac nodes after Windows nodes with the same names.
-			Dictionary<string, string> NodeDependencyNames = new Dictionary<string,string>(InitialNodeDependencyNames);
-			foreach(KeyValuePair<string, string> DisplayGroup in DisplayGroups)
+			Dictionary<NodeInfo, List<NodeInfo>> NodeDependencies = new Dictionary<NodeInfo,List<NodeInfo>>(Nodes.ToDictionary(x => x, x => x.AllDirectDependencies.ToList()));
+			foreach(KeyValuePair<string, List<NodeInfo>> DisplayGroup in DisplayGroups)
 			{
-				string[] GroupNodes = DisplayGroup.Value.Split(' ');
-				for(int Idx = 1; Idx < GroupNodes.Length; Idx++)
+				List<NodeInfo> GroupNodes = DisplayGroup.Value;
+				for (int Idx = 1; Idx < GroupNodes.Count; Idx++)
 				{
-					NodeDependencyNames[GroupNodes[Idx]] += " " + GroupNodes[0];
+					NodeDependencies[GroupNodes[Idx]].Add(GroupNodes[0]);
 				}
 			}
 
 			// Add nodes for each frequency into the master list, trying to match up different groups along the way
-			foreach(string FirstNodeName in NodesByThisFrequency)
+			foreach(NodeInfo FirstNode in NodesByThisFrequency)
 			{
-				string[] GroupNodeNames = DisplayGroups[GUBPNodes[FirstNodeName].Node.GetDisplayGroupName()].Split(' ');
-				foreach(string GroupNodeName in GroupNodeNames)
+				List<NodeInfo> GroupNodes = DisplayGroups[FirstNode.Node.GetDisplayGroupName()];
+				foreach(NodeInfo GroupNode in GroupNodes)
 				{
-					AddNodeAndDependencies(GroupNodeName, NodeDependencyNames, VisitedNodes, SortedNodes);
+					AddNodeAndDependencies(GroupNode, NodeDependencies, VisitedNodes, SortedNodes);
 				}
 			}
 		}
 		return SortedNodes;
 	}
 
-	static void AddNodeAndDependencies(string NodeName, Dictionary<string, string> NodeDependencyNames, HashSet<string> VisitedNodes, Dictionary<string, int> SortedNodes)
+	static void AddNodeAndDependencies(NodeInfo Node, Dictionary<NodeInfo, List<NodeInfo>> NodeDependencies, HashSet<NodeInfo> VisitedNodes, Dictionary<NodeInfo, int> SortedNodes)
 	{
-		if(!VisitedNodes.Contains(NodeName))
+		if(!VisitedNodes.Contains(Node))
 		{
-			VisitedNodes.Add(NodeName);
-			foreach(string NodeDependencyName in NodeDependencyNames[NodeName].Split(new char[]{ ' ' }, StringSplitOptions.RemoveEmptyEntries))
+			VisitedNodes.Add(Node);
+			foreach (NodeInfo NodeDependency in NodeDependencies[Node])
 			{
-				AddNodeAndDependencies(NodeDependencyName, NodeDependencyNames, VisitedNodes, SortedNodes);
+				AddNodeAndDependencies(NodeDependency, NodeDependencies, VisitedNodes, SortedNodes);
 			}
-			SortedNodes.Add(NodeName, SortedNodes.Count);
+			SortedNodes.Add(Node, SortedNodes.Count);
 		}
 	}
 }
