Move NodeInfo into its own file, and rename GUBP.AggregateNode to GUBPAggregateNode

diff --git a/Engine/Source/Programs/AutomationTool/GUBP/BuildNode.cs b/Engine/Source/Programs/AutomationTool/GUBP/BuildNode.cs
new file mode 100644
--- /dev/null
+++ b/Engine/Source/Programs/AutomationTool/GUBP/BuildNode.cs
@@ -0,0 +1,40 @@
+﻿using System;
+using System.Collections.Generic;
+using System.Diagnostics;
+using System.Linq;
+using System.Text;
+
+namespace AutomationTool
+{
+	[DebuggerDisplay("{Name}")]
+	class NodeInfo
+	{
+		public string Name;
+		public GUBP.GUBPNode Node;
+		public NodeInfo[] Dependencies;
+		public NodeInfo[] PseudoDependencies;
+		public NodeInfo[] AllDirectDependencies;
+		public NodeInfo[] AllIndirectDependencies;
+		public NodeInfo[] ControllingTriggers;
+		public int FrequencyShift;
+		public bool IsComplete;
+		public string[] RecipientsForFailureEmails;
+		public bool AddSubmittersToFailureEmails;
+
+		public string ControllingTriggerDotName
+		{
+			get { return String.Join(".", ControllingTriggers.Select(x => x.Name)); }
+		}
+
+		public bool DependsOn(NodeInfo Node)
+		{
+			return AllIndirectDependencies.Contains(Node);
+		}
+
+		public override string ToString()
+		{
+			System.Diagnostics.Trace.TraceWarning("Implicit conversion from NodeInfo to string\n{0}", Environment.StackTrace);
+			return Name;
+		}
+	}
+}
diff --git a/Engine/Source/Programs/AutomationTool/GUBP/GUBP.Automation.csproj b/Engine/Source/Programs/AutomationTool/GUBP/GUBP.Automation.csproj
--- a/Engine/Source/Programs/AutomationTool/GUBP/GUBP.Automation.csproj
+++ b/Engine/Source/Programs/AutomationTool/GUBP/GUBP.Automation.csproj
@@ -58,6 +58,7 @@
     <Compile Include="..\..\DotNETCommon\MetaData.cs">
       <Link>Properties\MetaData.cs</Link>
     </Compile>
+    <Compile Include="BuildNode.cs" />
     <Compile Include="ElectricCommander.cs" />
     <Compile Include="GUBP.cs" />
     <Compile Include="LegacyBranchSetup.cs" />
diff --git a/Engine/Source/Programs/AutomationTool/GUBP/GUBP.cs b/Engine/Source/Programs/AutomationTool/GUBP/GUBP.cs
--- a/Engine/Source/Programs/AutomationTool/GUBP/GUBP.cs
+++ b/Engine/Source/Programs/AutomationTool/GUBP/GUBP.cs
@@ -30,43 +30,11 @@ public partial class GUBP : BuildCommand
 	Dictionary<string, NodeInfo> GUBPNodes;
 	Dictionary<string, AggregateInfo> GUBPAggregates;
 
-	[DebuggerDisplay("{Name}")]
-	class NodeInfo
-	{
-		public string Name;
-		public GUBPNode Node;
-		public NodeInfo[] Dependencies;
-		public NodeInfo[] PseudoDependencies;
-		public NodeInfo[] AllDirectDependencies;
-		public NodeInfo[] AllIndirectDependencies;
-		public NodeInfo[] ControllingTriggers;
-		public int FrequencyShift;
-		public bool IsComplete;
-		public string[] RecipientsForFailureEmails;
-		public bool AddSubmittersToFailureEmails;
-
-		public string ControllingTriggerDotName
-		{
-			get { return String.Join(".", ControllingTriggers.Select(x => x.Name)); }
-		}
-
-		public bool DependsOn(NodeInfo Node)
-		{
-			return AllIndirectDependencies.Contains(Node);
-		}
-
-		public override string ToString()
-		{
-			System.Diagnostics.Trace.TraceWarning("Implicit conversion from NodeInfo to string\n{0}", Environment.StackTrace);
-			return Name;
-		}
-	}
-
 	[DebuggerDisplay("{Name}")]
 	class AggregateInfo
 	{
 		public string Name;
-		public AggregateNode Node;
+		public GUBPAggregateNode Node;
 		public NodeInfo[] Dependencies;
 	}
 
@@ -94,7 +62,7 @@ public partial class GUBP : BuildCommand
         return Name;
     }
 
-	public string AddNode(AggregateNode Node)
+	public string AddNode(GUBPAggregateNode Node)
 	{
 		string Name = Node.GetFullName();
 		if (GUBPNodes.ContainsKey(Name) || GUBPAggregates.ContainsKey(Name))
@@ -115,7 +83,7 @@ public partial class GUBP : BuildCommand
         return GUBPNodes[Node].Node;
     }
 
-    public AggregateNode FindAggregateNode(string Node)
+    public GUBPAggregateNode FindAggregateNode(string Node)
     {
         return GUBPAggregates[Node].Node;
     }
diff --git a/Engine/Source/Programs/AutomationTool/GUBP/LegacyNodes.cs b/Engine/Source/Programs/AutomationTool/GUBP/LegacyNodes.cs
--- a/Engine/Source/Programs/AutomationTool/GUBP/LegacyNodes.cs
+++ b/Engine/Source/Programs/AutomationTool/GUBP/LegacyNodes.cs
@@ -1559,11 +1559,11 @@ partial class GUBP
         }
     }
 
-    public abstract class AggregateNode
+    public abstract class GUBPAggregateNode
     {
 		public List<string> Dependencies = new List<string>();
 
-        public AggregateNode()
+        public GUBPAggregateNode()
         {
         }
 
@@ -1593,7 +1593,7 @@ partial class GUBP
 		}
     }
 
-    public abstract class HostPlatformAggregateNode : AggregateNode
+    public abstract class HostPlatformAggregateNode : GUBPAggregateNode
     {
         protected UnrealTargetPlatform HostPlatform;
         public HostPlatformAggregateNode(UnrealTargetPlatform InHostPlatform)
@@ -1634,7 +1634,7 @@ partial class GUBP
     }
 
 
-    public class AggregatePromotableNode : AggregateNode
+    public class AggregatePromotableNode : GUBPAggregateNode
     {
         protected List<UnrealTargetPlatform> HostPlatforms;
         string PromotionLabelPrefix;		
@@ -1778,7 +1778,7 @@ partial class GUBP
         }
     }
 
-	public class FullGameAggregateNode : AggregateNode
+	public class FullGameAggregateNode : GUBPAggregateNode
 	{
 		string GameName;
 
@@ -2154,7 +2154,7 @@ partial class GUBP
         }
     }
 
-    public class SharedLabelPromotableSuccessNode : AggregateNode
+    public class SharedLabelPromotableSuccessNode : GUBPAggregateNode
     {        
         public SharedLabelPromotableSuccessNode()
         {
@@ -2172,7 +2172,7 @@ partial class GUBP
 		}
     }
 
-    public class WaitForTestShared : AggregateNode
+    public class WaitForTestShared : GUBPAggregateNode
     {
         public WaitForTestShared(GUBP bp)
         {
