Move hacks to remove pseudo-dependencies into legacy branch setup.

diff --git a/Engine/Source/Programs/AutomationTool/GUBP/GUBP.cs b/Engine/Source/Programs/AutomationTool/GUBP/GUBP.cs
--- a/Engine/Source/Programs/AutomationTool/GUBP/GUBP.cs
+++ b/Engine/Source/Programs/AutomationTool/GUBP/GUBP.cs
@@ -104,15 +104,6 @@ public partial class GUBP : BuildCommand
         GUBPNodes[Node].Node.RemovePseudodependency(Dep);
     }
 
-    public void RemoveAllPseudodependenciesFromNode(string Node)
-    {
-        if (!GUBPNodes.ContainsKey(Node))
-        {
-            throw new AutomationException("Node {0} not found", Node);
-        }
-        GUBPNodes[Node].Node.FullNamesOfPseudosependencies.Clear();
-    }
-
 	/// <summary>
 	/// Recursively update the ControllingTriggers array for each of the nodes passed in. 
 	/// </summary>
@@ -1644,20 +1635,6 @@ public partial class GUBP : BuildCommand
 			HashSet<NodeInfo> Fringe = new HashSet<NodeInfo>();
 			foreach (NodeInfo NodeToDo in NodesToDo)
 			{
-				if (BranchOptions.NodesToRemovePseudoDependencies.Count > 0)
-				{
-					List<string> ListOfRetainedNodesToRemovePseudoDependencies = BranchOptions.NodesToRemovePseudoDependencies;
-					foreach (string NodeToRemovePseudoDependencies in BranchOptions.NodesToRemovePseudoDependencies)
-					{
-						if (NodeToDo.Name == NodeToRemovePseudoDependencies)
-						{
-							RemoveAllPseudodependenciesFromNode(NodeToDo.Name);
-							ListOfRetainedNodesToRemovePseudoDependencies.Remove(NodeToDo.Name);
-							break;
-						}
-					}
-					BranchOptions.NodesToRemovePseudoDependencies = ListOfRetainedNodesToRemovePseudoDependencies;
-				}
 				foreach (NodeInfo Dep in NodeToDo.Dependencies)
 				{
 					if (!NodesToDo.Contains(Dep))
diff --git a/Engine/Source/Programs/AutomationTool/GUBP/LegacyBranchSetup.cs b/Engine/Source/Programs/AutomationTool/GUBP/LegacyBranchSetup.cs
--- a/Engine/Source/Programs/AutomationTool/GUBP/LegacyBranchSetup.cs
+++ b/Engine/Source/Programs/AutomationTool/GUBP/LegacyBranchSetup.cs
@@ -1062,6 +1062,16 @@ partial class GUBP
         }
 #endif
 
+		// Remove all the pseudo-dependencies on nodes specified in the branch options
+		foreach(string NodeToRemovePseudoDependencies in BranchOptions.NodesToRemovePseudoDependencies)
+		{
+			GUBPNode Node = TryFindNode(NodeToRemovePseudoDependencies);
+			if(Node != null)
+			{
+				Node.FullNamesOfPseudosependencies.Clear();
+			}
+		}
+
 		// Calculate the frequencies for each node
 		Dictionary<string, int> FrequencyOverrides = ApplyFrequencyBarriers(GUBPNodes, GUBPAggregates, BranchOptions.FrequencyBarriers);
 		foreach(NodeInfo Node in GUBPNodes.Values)
