Only generate obsolete manifests when deploying.  Don't want to attempt to retreive manifests when there are no devices

diff --git a/Engine/Source/Programs/AutomationTool/Scripts/CopyBuildToStagingDirectory.Automation.cs b/Engine/Source/Programs/AutomationTool/Scripts/CopyBuildToStagingDirectory.Automation.cs
--- a/Engine/Source/Programs/AutomationTool/Scripts/CopyBuildToStagingDirectory.Automation.cs
+++ b/Engine/Source/Programs/AutomationTool/Scripts/CopyBuildToStagingDirectory.Automation.cs
@@ -1581,6 +1581,7 @@ public partial class Project : CommandUtils
 					CreateStagingManifest(Params, SC);
 					ApplyStagingManifest(Params, SC);
 
+					if (Params.Deploy)
 					{
 						// get the deployed file data
 						Dictionary<string, string> DeployedUFSFiles = new Dictionary<string, string>();
