Move handling of the GUBPNodes and GUBPAggregateNodes dictionaries out into LegacyBranchSetup.

diff --git a/Engine/Source/Programs/AutomationTool/GUBP/GUBP.cs b/Engine/Source/Programs/AutomationTool/GUBP/GUBP.cs
--- a/Engine/Source/Programs/AutomationTool/GUBP/GUBP.cs
+++ b/Engine/Source/Programs/AutomationTool/GUBP/GUBP.cs
@@ -27,9 +27,6 @@ public partial class GUBP : BuildCommand
     public string PreflightMangleSuffix = "";
     public GUBPBranchHacker.BranchOptions BranchOptions = null;	
 
-	Dictionary<string, NodeInfo> GUBPNodes;
-	Dictionary<string, AggregateInfo> GUBPAggregates;
-
 	[DebuggerDisplay("{Name}")]
 	class AggregateInfo
 	{
@@ -51,59 +48,6 @@ public partial class GUBP : BuildCommand
         public List<int> AllFailed = new List<int>();
     };
 
-    public string AddNode(GUBPNode Node)
-    {
-        string Name = Node.GetFullName();
-		if (GUBPNodes.ContainsKey(Name) || GUBPAggregates.ContainsKey(Name))
-		{
-            throw new AutomationException("Attempt to add a duplicate node {0}", Node.GetFullName());
-        }
-        GUBPNodes.Add(Name, new NodeInfo{ Name = Name, Node = Node });
-        return Name;
-    }
-
-	public string AddNode(GUBPAggregateNode Node)
-	{
-		string Name = Node.GetFullName();
-		if (GUBPNodes.ContainsKey(Name) || GUBPAggregates.ContainsKey(Name))
-		{
-			throw new AutomationException("Attempt to add a duplicate node {0}", Node.GetFullName());
-		}
-		GUBPAggregates.Add(Name, new AggregateInfo { Name = Name, Node = Node });
-		return Name;
-	}
-
-    public bool HasNode(string Node)
-    {
-        return GUBPNodes.ContainsKey(Node) || GUBPAggregates.ContainsKey(Node);
-    }
-
-    public GUBPNode FindNode(string Node)
-    {
-        return GUBPNodes[Node].Node;
-    }
-
-    public GUBPAggregateNode FindAggregateNode(string Node)
-    {
-        return GUBPAggregates[Node].Node;
-    }
-
-	public GUBPNode TryFindNode(string Node)
-	{
-		NodeInfo Result;
-		GUBPNodes.TryGetValue(Node, out Result);
-		return (Result == null)? null : Result.Node;
-	}
-
-    public void RemovePseudodependencyFromNode(string Node, string Dep)
-    {
-        if (!GUBPNodes.ContainsKey(Node))
-        {
-            throw new AutomationException("Node {0} not found", Node);
-        }
-        GUBPNodes[Node].Node.RemovePseudodependency(Dep);
-    }
-
 	/// <summary>
 	/// Recursively update the ControllingTriggers array for each of the nodes passed in. 
 	/// </summary>
@@ -1303,37 +1247,37 @@ public partial class GUBP : BuildCommand
             Log("Setting TimeIndex to {0}", TimeIndex);
         }
 
-        GUBPNodes = new Dictionary<string, NodeInfo>();
-		GUBPAggregates = new Dictionary<string, AggregateInfo>();
         Branch = new BranchInfo(HostPlatforms);        
 
-		AddNodesForBranch(TimeIndex, bNoAutomatedTesting);
+		List<NodeInfo> AllNodes;
+		List<AggregateInfo> AllAggregates;
+		AddNodesForBranch(TimeIndex, bNoAutomatedTesting, out AllNodes, out AllAggregates);
 
-		LinkGraph(GUBPAggregates, GUBPNodes);
-		FindControllingTriggers(GUBPNodes.Values);
-		FindCompletionState(GUBPNodes.Values, StoreName, LocalOnly);
-		ComputeDependentFrequencies(GUBPNodes.Values);
+		LinkGraph(AllAggregates, AllNodes);
+		FindControllingTriggers(AllNodes);
+		FindCompletionState(AllNodes, StoreName, LocalOnly);
+		ComputeDependentFrequencies(AllNodes);
 
         if (bCleanLocalTempStorage)  // shared temp storage can never be wiped
         {
             TempStorage.DeleteLocalTempStorageManifests(CmdEnv);
         }
 
-		HashSet<NodeInfo> NodesToDo = ParseNodesToDo(WithoutLinux, TimeIndex, CommanderSetup);
+		HashSet<NodeInfo> NodesToDo = ParseNodesToDo(AllNodes, AllAggregates, WithoutLinux, TimeIndex, CommanderSetup);
 
 		NodeInfo ExplicitTrigger = null;
 		if (CommanderSetup)
         {
             if (!String.IsNullOrEmpty(ExplicitTriggerName))
             {
-                foreach (KeyValuePair<string, NodeInfo> Node in GUBPNodes)
+                foreach (NodeInfo Node in AllNodes)
                 {
-                    if (Node.Value.Node.GetFullName().Equals(ExplicitTriggerName, StringComparison.InvariantCultureIgnoreCase))
+                    if (Node.Name.Equals(ExplicitTriggerName, StringComparison.InvariantCultureIgnoreCase))
                     {
-                        if (Node.Value.Node.TriggerNode())
+                        if (Node.Node.TriggerNode())
                         {
-                            Node.Value.Node.SetAsExplicitTrigger();
-							ExplicitTrigger = Node.Value;
+                            Node.Node.SetAsExplicitTrigger();
+							ExplicitTrigger = Node;
                             break;
                         }
                     }
@@ -1347,11 +1291,11 @@ public partial class GUBP : BuildCommand
             {
                 if (bSkipTriggers)
                 {
-                    foreach (KeyValuePair<string, NodeInfo> Node in GUBPNodes)
+                    foreach (NodeInfo Node in AllNodes)
                     {
-                        if (Node.Value.Node.TriggerNode())
+                        if (Node.Node.TriggerNode())
                         {
-                            Node.Value.Node.SetAsExplicitTrigger();
+                            Node.Node.SetAsExplicitTrigger();
                         }
                     }
                 }
@@ -1368,7 +1312,7 @@ public partial class GUBP : BuildCommand
 
 		List<NodeInfo> UnfinishedTriggers = FindUnfinishedTriggers(bSkipTriggers, ExplicitTrigger, OrderedToDo);
 
-		PrintNodes(this, OrderedToDo, GUBPAggregates.Values, UnfinishedTriggers, TimeQuantum);
+		PrintNodes(this, OrderedToDo, AllAggregates, UnfinishedTriggers, TimeQuantum);
 
         //check sorting
 		CheckSortOrder(OrderedToDo);
@@ -1376,7 +1320,7 @@ public partial class GUBP : BuildCommand
         string FakeFail = ParseParamValue("FakeFail");
         if(CommanderSetup)
         {
-			DoCommanderSetup(NodesToDo, OrderedToDo, TimeIndex, TimeQuantum, bSkipTriggers, bFake, bFakeEC, CLString, ExplicitTrigger, UnfinishedTriggers, FakeFail);
+			DoCommanderSetup(AllNodes, AllAggregates, NodesToDo, OrderedToDo, TimeIndex, TimeQuantum, bSkipTriggers, bFake, bFakeEC, CLString, ExplicitTrigger, UnfinishedTriggers, FakeFail);
         }
 		else if(ParseParam("SaveGraph"))
 		{
@@ -1421,8 +1365,20 @@ public partial class GUBP : BuildCommand
 	/// </summary>
 	/// <param name="AggregateNameToInfo">Map of aggregate names to their info objects</param>
 	/// <param name="NodeNameToInfo">Map of node names to their info objects</param>
-	private static void LinkGraph(Dictionary<string, AggregateInfo> AggregateNameToInfo, Dictionary<string, NodeInfo> NodeNameToInfo)
+	private static void LinkGraph(IEnumerable<AggregateInfo> Aggregates, IEnumerable<NodeInfo> Nodes)
 	{
+		Dictionary<string, NodeInfo> NodeNameToInfo = new Dictionary<string,NodeInfo>();
+		foreach(NodeInfo Node in Nodes)
+		{
+			NodeNameToInfo.Add(Node.Name, Node);
+		}
+
+		Dictionary<string, AggregateInfo> AggregateNameToInfo = new Dictionary<string, AggregateInfo>();
+		foreach(AggregateInfo Aggregate in Aggregates)
+		{
+			AggregateNameToInfo.Add(Aggregate.Name, Aggregate);
+		}
+
 		int NumErrors = 0;
 		foreach (AggregateInfo AggregateNode in AggregateNameToInfo.Values)
 		{
@@ -1559,7 +1515,7 @@ public partial class GUBP : BuildCommand
 		return false;
 	}
 
-	private HashSet<NodeInfo> ParseNodesToDo(bool WithoutLinux, int TimeIndex, bool CommanderSetup)
+	private HashSet<NodeInfo> ParseNodesToDo(IEnumerable<NodeInfo> PotentialNodes, IEnumerable<AggregateInfo> PotentialAggregates, bool WithoutLinux, int TimeIndex, bool CommanderSetup)
 	{
 		List<string> NamesToDo = new List<string>();
 
@@ -1579,7 +1535,7 @@ public partial class GUBP : BuildCommand
 		foreach (string NameToDo in NamesToDo)
 		{
 			int FoundNames = 0;
-			foreach (NodeInfo PotentialNode in GUBPNodes.Values)
+			foreach (NodeInfo PotentialNode in PotentialNodes)
 			{
 				if (String.Compare(PotentialNode.Name, NameToDo, StringComparison.InvariantCultureIgnoreCase) == 0 || String.Compare(PotentialNode.Node.AgentSharingGroup, NameToDo, StringComparison.InvariantCultureIgnoreCase) == 0)
 				{
@@ -1587,7 +1543,7 @@ public partial class GUBP : BuildCommand
 					FoundNames++;
 				}
 			}
-			foreach (AggregateInfo PotentialAggregate in GUBPAggregates.Values)
+			foreach (AggregateInfo PotentialAggregate in PotentialAggregates)
 			{
 				if (String.Compare(PotentialAggregate.Name, NameToDo, StringComparison.InvariantCultureIgnoreCase) == 0)
 				{
@@ -1604,10 +1560,7 @@ public partial class GUBP : BuildCommand
 		if (NodesToDo.Count == 0)
 		{
 			LogVerbose("No nodes specified, adding all nodes");
-			foreach (KeyValuePair<string, NodeInfo> Node in GUBPNodes)
-			{
-				NodesToDo.Add(Node.Value);
-			}
+			NodesToDo.UnionWith(PotentialNodes);
 		}
 		else if (TimeIndex != 0)
 		{
@@ -1899,14 +1852,14 @@ public partial class GUBP : BuildCommand
 		}
 	}
 
-	private void DoCommanderSetup(HashSet<NodeInfo> NodesToDo, List<NodeInfo> OrdereredToDo, int TimeIndex, int TimeQuantum, bool bSkipTriggers, bool bFake, bool bFakeEC, string CLString, NodeInfo ExplicitTrigger, List<NodeInfo> UnfinishedTriggers, string FakeFail)
+	private void DoCommanderSetup(IEnumerable<NodeInfo> AllNodes, IEnumerable<AggregateInfo> AllAggregates, HashSet<NodeInfo> NodesToDo, List<NodeInfo> OrdereredToDo, int TimeIndex, int TimeQuantum, bool bSkipTriggers, bool bFake, bool bFakeEC, string CLString, NodeInfo ExplicitTrigger, List<NodeInfo> UnfinishedTriggers, string FakeFail)
 	{
 		Dictionary<string, string> FullNodeDependedOnBy = GetFullNodeDependedOnBy(NodesToDo);
 
-		List<AggregateInfo> SeparatePromotables = FindPromotables(GUBPAggregates.Values);
-		Dictionary<NodeInfo, List<AggregateInfo>> DependentPromotions = FindDependentPromotables(GUBPNodes.Values, SeparatePromotables);
+		List<AggregateInfo> SeparatePromotables = FindPromotables(AllAggregates);
+		Dictionary<NodeInfo, List<AggregateInfo>> DependentPromotions = FindDependentPromotables(AllNodes, SeparatePromotables);
 
-		List<NodeInfo> SortedNodes = TopologicalSort(new HashSet<NodeInfo>(GUBPNodes.Values), null, SubSort: false, DoNotConsiderCompletion: true);
+		List<NodeInfo> SortedNodes = TopologicalSort(new HashSet<NodeInfo>(AllNodes), null, SubSort: false, DoNotConsiderCompletion: true);
 		Log("******* {0} GUBP Nodes", SortedNodes.Count);
 
 		Dictionary<NodeInfo, int> FullNodeListSortKey = GetDisplayOrder(SortedNodes);
@@ -1983,7 +1936,7 @@ public partial class GUBP : BuildCommand
 		using(TelemetryStopwatch PrintNodesTimer = new TelemetryStopwatch("SetupCommanderPrint"))
 		{
 			Log("*********** EC Nodes, in order.");
-			PrintNodes(this, OrdereredToDo, GUBPAggregates.Values, UnfinishedTriggers, TimeQuantum);
+			PrintNodes(this, OrdereredToDo, AllAggregates, UnfinishedTriggers, TimeQuantum);
 		}
 		// here we are just making sure everything before the explicit trigger is completed.
 		if (ExplicitTrigger != null)
diff --git a/Engine/Source/Programs/AutomationTool/GUBP/LegacyBranchSetup.cs b/Engine/Source/Programs/AutomationTool/GUBP/LegacyBranchSetup.cs
--- a/Engine/Source/Programs/AutomationTool/GUBP/LegacyBranchSetup.cs
+++ b/Engine/Source/Programs/AutomationTool/GUBP/LegacyBranchSetup.cs
@@ -11,9 +11,68 @@ using System.Linq;
 
 partial class GUBP
 {
+	public class GUBPBranchConfig
+	{
+		public Dictionary<string, GUBPNode> GUBPNodes = new Dictionary<string,GUBPNode>();
+		public Dictionary<string, GUBPAggregateNode> GUBPAggregates = new Dictionary<string,GUBPAggregateNode>();
+
+		public string AddNode(GUBPNode Node)
+		{
+			string Name = Node.GetFullName();
+			if (GUBPNodes.ContainsKey(Name) || GUBPAggregates.ContainsKey(Name))
+			{
+				throw new AutomationException("Attempt to add a duplicate node {0}", Node.GetFullName());
+			}
+			GUBPNodes.Add(Name, Node);
+			return Name;
+		}
+
+		public string AddNode(GUBPAggregateNode Node)
+		{
+			string Name = Node.GetFullName();
+			if (GUBPNodes.ContainsKey(Name) || GUBPAggregates.ContainsKey(Name))
+			{
+				throw new AutomationException("Attempt to add a duplicate node {0}", Node.GetFullName());
+			}
+			GUBPAggregates.Add(Name, Node);
+			return Name;
+		}
+
+		public bool HasNode(string Node)
+		{
+			return GUBPNodes.ContainsKey(Node) || GUBPAggregates.ContainsKey(Node);
+		}
+
+		public GUBPNode FindNode(string Node)
+		{
+			return GUBPNodes[Node];
+		}
+
+		public GUBPAggregateNode FindAggregateNode(string Node)
+		{
+			return GUBPAggregates[Node];
+		}
+
+		public GUBPNode TryFindNode(string Node)
+		{
+			GUBPNode Result;
+			GUBPNodes.TryGetValue(Node, out Result);
+			return Result;
+		}
+
+		public void RemovePseudodependencyFromNode(string Node, string Dep)
+		{
+			if (!GUBPNodes.ContainsKey(Node))
+			{
+				throw new AutomationException("Node {0} not found", Node);
+			}
+			GUBPNodes[Node].RemovePseudodependency(Dep);
+		}
+	}
+
     public abstract class GUBPNodeAdder
     {
-        public abstract void AddNodes(GUBP bp, UnrealTargetPlatform InHostPlatform, List<UnrealTargetPlatform> InActivePlatforms);
+        public abstract void AddNodes(GUBP bp, GUBPBranchConfig BranchConfig, UnrealTargetPlatform InHostPlatform, List<UnrealTargetPlatform> InActivePlatforms);
     }
 
     private static List<GUBPNodeAdder> Adders;
@@ -32,7 +91,7 @@ partial class GUBP
 		return AllTypesWeCanGet;
 	}
 
-    private void AddCustomNodes(UnrealTargetPlatform InHostPlatform, List<UnrealTargetPlatform> InActivePlatforms)
+    private void AddCustomNodes(GUBPBranchConfig BranchConfig, UnrealTargetPlatform InHostPlatform, List<UnrealTargetPlatform> InActivePlatforms)
     {
         if (Adders == null)
         {
@@ -56,7 +115,7 @@ partial class GUBP
         }
         foreach(GUBPNodeAdder Adder in Adders)
         {
-            Adder.AddNodes(this, InHostPlatform, InActivePlatforms);
+            Adder.AddNodes(this, BranchConfig, InHostPlatform, InActivePlatforms);
         }
     }
 
@@ -154,7 +213,7 @@ partial class GUBP
 			}
 		}
 
-		foreach (NodeInfo Node in GUBPNodes.Values)
+		foreach (NodeInfo Node in Nodes)
 		{
 			Node.RecipientsForFailureEmails = HackEmails(EmailHackers, BranchName, Node, out Node.AddSubmittersToFailureEmails);
 		}
@@ -245,7 +304,7 @@ partial class GUBP
         return Frequency;
     }    
 
-	void AddNodesForBranch(int TimeIndex, bool bNoAutomatedTesting)
+	void AddNodesForBranch(int TimeIndex, bool bNoAutomatedTesting, out List<NodeInfo> AllNodes, out List<AggregateInfo> AllAggregates)
 	{
 		List<UnrealTargetPlatform> ActivePlatforms;
         if (IsBuildMachine || ParseParam("AllPlatforms"))
@@ -325,24 +384,25 @@ partial class GUBP
             }
         }
 
+		GUBPBranchConfig BranchConfig = new GUBPBranchConfig();
 
-        AddNode(new VersionFilesNode());
+        BranchConfig.AddNode(new VersionFilesNode());
 
 
         foreach (var HostPlatform in HostPlatforms)
         {
-			AddNode(new ToolsForCompileNode(HostPlatform, TimeIndex));
+			BranchConfig.AddNode(new ToolsForCompileNode(HostPlatform, TimeIndex));
 			
 			if (!BranchOptions.ExcludePlatformsForEditor.Contains(HostPlatform))
 			{
-				AddNode(new RootEditorNode(HostPlatform));			
-				AddNode(new ToolsNode(this, HostPlatform));            
-				AddNode(new InternalToolsNode(this, HostPlatform));
+				BranchConfig.AddNode(new RootEditorNode(HostPlatform));			
+				BranchConfig.AddNode(new ToolsNode(this, HostPlatform));            
+				BranchConfig.AddNode(new InternalToolsNode(this, HostPlatform));
 			if (HostPlatform == UnrealTargetPlatform.Win64 && ActivePlatforms.Contains(UnrealTargetPlatform.Linux))
 				{
 					if (!BranchOptions.ExcludePlatformsForEditor.Contains(UnrealTargetPlatform.Linux))
 					{
-						AddNode(new ToolsCrossCompileNode(HostPlatform));
+						BranchConfig.AddNode(new ToolsCrossCompileNode(HostPlatform));
 					}
 				}
 				foreach (var ProgramTarget in Branch.BaseEngineProject.Properties.Programs)
@@ -357,16 +417,16 @@ partial class GUBP
                         {
                             if (bInternalOnly)
                             {
-                                AddNode(new SingleInternalToolsNode(HostPlatform, ProgramTarget));
+                                BranchConfig.AddNode(new SingleInternalToolsNode(HostPlatform, ProgramTarget));
                             }
                             else
                             {
-                                AddNode(new SingleToolsNode(HostPlatform, ProgramTarget));
+                                BranchConfig.AddNode(new SingleToolsNode(HostPlatform, ProgramTarget));
                             }
                         }
                         if (ProgramTarget.Rules.GUBP_IncludeNonUnityToolTest())
                         {
-                            AddNode(new NonUnityToolNode(HostPlatform, ProgramTarget));
+                            BranchConfig.AddNode(new NonUnityToolNode(HostPlatform, ProgramTarget));
                         }
                     }
 				}
@@ -384,22 +444,22 @@ partial class GUBP
                             {
                                 if (bInternalNodeOnly)
                                 {
-                                    AddNode(new SingleInternalToolsNode(HostPlatform, ProgramTarget));
+                                    BranchConfig.AddNode(new SingleInternalToolsNode(HostPlatform, ProgramTarget));
                                 }
                                 else
                                 {
-                                    AddNode(new SingleToolsNode(HostPlatform, ProgramTarget));
+                                    BranchConfig.AddNode(new SingleToolsNode(HostPlatform, ProgramTarget));
                                 }
                             }
                             if (ProgramTarget.Rules.GUBP_IncludeNonUnityToolTest())
                             {
-                                AddNode(new NonUnityToolNode(HostPlatform, ProgramTarget));
+                                BranchConfig.AddNode(new NonUnityToolNode(HostPlatform, ProgramTarget));
                             }
                         }
 					}
 				}
 
-				AddNode(new EditorAndToolsNode(this, HostPlatform));
+				BranchConfig.AddNode(new EditorAndToolsNode(this, HostPlatform));
 			}
 
             bool DoASharedPromotable = false;
@@ -473,7 +533,7 @@ partial class GUBP
 
 			if (!BranchOptions.ExcludePlatformsForEditor.Contains(HostPlatform))
 			{
-			AddNode(new NonUnityTestNode(HostPlatform));
+			BranchConfig.AddNode(new NonUnityTestNode(HostPlatform));
 			}
 
             if (DoASharedPromotable)
@@ -497,7 +557,7 @@ partial class GUBP
                     {
                         if (!bNoAutomatedTesting)
                         {							
-                            EditorTestNodes.Add(AddNode(new UATTestNode(this, HostPlatform, Branch.BaseEngineProject, Test.Key, Test.Value, AgentSharingGroup)));
+                            EditorTestNodes.Add(BranchConfig.AddNode(new UATTestNode(this, BranchConfig, HostPlatform, Branch.BaseEngineProject, Test.Key, Test.Value, AgentSharingGroup)));
                         
                             foreach (var NonCodeProject in Branch.NonCodeProjects)
                             {
@@ -506,13 +566,13 @@ partial class GUBP
                                     continue;
                                 }
                                 if (HostPlatform == UnrealTargetPlatform.Mac) continue; //temp hack till mac automated testing works
-                                EditorTestNodes.Add(AddNode(new UATTestNode(this, HostPlatform, NonCodeProject, Test.Key, Test.Value, AgentSharingGroup)));
+                                EditorTestNodes.Add(BranchConfig.AddNode(new UATTestNode(this, BranchConfig, HostPlatform, NonCodeProject, Test.Key, Test.Value, AgentSharingGroup)));
                             }
                         }
                     }
                     if (EditorTestNodes.Count > 0)
                     {
-                        AddNode(new GameAggregateNode(this, HostPlatform, Branch.BaseEngineProject, "AllEditorTests", EditorTestNodes));
+                        BranchConfig.AddNode(new GameAggregateNode(this, HostPlatform, Branch.BaseEngineProject, "AllEditorTests", EditorTestNodes));
                     }
                 }
 
@@ -553,19 +613,19 @@ partial class GUBP
                                 {
                                     GamePlatforms.Add(Plat);
                                 }
-                                if (!HasNode(GamePlatformMonolithicsNode.StaticGetFullName(HostPlatform, Branch.BaseEngineProject, Plat)))
+                                if (!BranchConfig.HasNode(GamePlatformMonolithicsNode.StaticGetFullName(HostPlatform, Branch.BaseEngineProject, Plat)))
                                 {
 									if(GamePlatformMonolithicsNode.HasPrecompiledTargets(Branch.BaseEngineProject, HostPlatform, Plat))
 									{
-	                                    AddNode(new GamePlatformMonolithicsNode(this, HostPlatform, ActivePlatforms, Branch.BaseEngineProject, Plat, InPrecompiled: true));
+	                                    BranchConfig.AddNode(new GamePlatformMonolithicsNode(this, BranchConfig, HostPlatform, ActivePlatforms, Branch.BaseEngineProject, Plat, InPrecompiled: true));
 									}
-                                    AddNode(new GamePlatformMonolithicsNode(this, HostPlatform, ActivePlatforms, Branch.BaseEngineProject, Plat));
+                                    BranchConfig.AddNode(new GamePlatformMonolithicsNode(this, BranchConfig, HostPlatform, ActivePlatforms, Branch.BaseEngineProject, Plat));
                                 }
 								if (Plat == UnrealTargetPlatform.Win32 && Target.Rules.GUBP_BuildWindowsXPMonolithics() && Kind == TargetRules.TargetType.Game)
 								{
-									if (!HasNode(GamePlatformMonolithicsNode.StaticGetFullName(HostPlatform, Branch.BaseEngineProject, Plat, true)))
+									if (!BranchConfig.HasNode(GamePlatformMonolithicsNode.StaticGetFullName(HostPlatform, Branch.BaseEngineProject, Plat, true)))
 									{
-										AddNode(new GamePlatformMonolithicsNode(this, HostPlatform, ActivePlatforms, Branch.BaseEngineProject, Plat, true));
+										BranchConfig.AddNode(new GamePlatformMonolithicsNode(this, BranchConfig, HostPlatform, ActivePlatforms, Branch.BaseEngineProject, Plat, true));
 									}
 	                            }
                             }
@@ -601,13 +661,13 @@ partial class GUBP
                             if (ActivePlatforms.Contains(Plat))
                             {
                                 string CookedPlatform = Platform.Platforms[Plat].GetCookPlatform(Kind == TargetRules.TargetType.Server, Kind == TargetRules.TargetType.Client, "");								
-                                if (!HasNode(CookNode.StaticGetFullName(HostPlatform, Branch.BaseEngineProject, CookedPlatform)))
+                                if (!BranchConfig.HasNode(CookNode.StaticGetFullName(HostPlatform, Branch.BaseEngineProject, CookedPlatform)))
                                 {
-                                    GameCookNodes.Add(AddNode(new CookNode(this, HostPlatform, Branch.BaseEngineProject, Plat, CookedPlatform)));
+                                    GameCookNodes.Add(BranchConfig.AddNode(new CookNode(this, BranchConfig, HostPlatform, Branch.BaseEngineProject, Plat, CookedPlatform)));
                                 }
-                                if (!HasNode(GamePlatformCookedAndCompiledNode.StaticGetFullName(HostPlatform, Branch.BaseEngineProject, Plat)))
+                                if (!BranchConfig.HasNode(GamePlatformCookedAndCompiledNode.StaticGetFullName(HostPlatform, Branch.BaseEngineProject, Plat)))
                                 {
-                                    AddNode(new GamePlatformCookedAndCompiledNode(this, HostPlatform, Branch.BaseEngineProject, Plat, false));
+                                    BranchConfig.AddNode(new GamePlatformCookedAndCompiledNode(this, HostPlatform, Branch.BaseEngineProject, Plat, false));
                                 }
                                 var GameTests = Target.Rules.GUBP_GetGameTests_MonolithicOnly(HostPlatform, GetAltHostPlatform(HostPlatform), Plat);
                                 var RequiredPlatforms = new List<UnrealTargetPlatform> { Plat };
@@ -618,11 +678,11 @@ partial class GUBP
                                     foreach (var Test in GameTests)
                                     {
                                         var TestName = Test.Key + "_" + Plat.ToString();
-                                        ThisMonoGameTestNodes.Add(AddNode(new UATTestNode(this, HostPlatform, Branch.BaseEngineProject, TestName, Test.Value, CookedAgentSharingGroup, false, RequiredPlatforms)));
+                                        ThisMonoGameTestNodes.Add(BranchConfig.AddNode(new UATTestNode(this, BranchConfig, HostPlatform, Branch.BaseEngineProject, TestName, Test.Value, CookedAgentSharingGroup, false, RequiredPlatforms)));
                                     }
                                     if (ThisMonoGameTestNodes.Count > 0)
                                     {										
-                                        GameTestNodes.Add(AddNode(new GameAggregateNode(this, HostPlatform, Branch.BaseEngineProject, "CookedTests_" + Plat.ToString() + "_" + Kind.ToString() + HostPlatformNode.StaticGetHostPlatformSuffix(HostPlatform), ThisMonoGameTestNodes)));
+                                        GameTestNodes.Add(BranchConfig.AddNode(new GameAggregateNode(this, HostPlatform, Branch.BaseEngineProject, "CookedTests_" + Plat.ToString() + "_" + Kind.ToString() + HostPlatformNode.StaticGetHostPlatformSuffix(HostPlatform), ThisMonoGameTestNodes)));
                                     }
                                 }
 
@@ -634,13 +694,13 @@ partial class GUBP
                                     }
                                     if (BranchOptions.ProjectsToCook.Contains(NonCodeProject.GameName) || BranchOptions.ProjectsToCook.Count == 0)
                                     {
-                                    if (!HasNode(CookNode.StaticGetFullName(HostPlatform, NonCodeProject, CookedPlatform)))
+                                    if (!BranchConfig.HasNode(CookNode.StaticGetFullName(HostPlatform, NonCodeProject, CookedPlatform)))
                                     {
-                                        GameCookNodes.Add(AddNode(new CookNode(this, HostPlatform, NonCodeProject, Plat, CookedPlatform)));
+                                        GameCookNodes.Add(BranchConfig.AddNode(new CookNode(this, BranchConfig, HostPlatform, NonCodeProject, Plat, CookedPlatform)));
                                     }
-                                    if (!HasNode(GamePlatformCookedAndCompiledNode.StaticGetFullName(HostPlatform, NonCodeProject, Plat)))
+                                    if (!BranchConfig.HasNode(GamePlatformCookedAndCompiledNode.StaticGetFullName(HostPlatform, NonCodeProject, Plat)))
                                     {
-                                        AddNode(new GamePlatformCookedAndCompiledNode(this, HostPlatform, NonCodeProject, Plat, false));
+                                        BranchConfig.AddNode(new GamePlatformCookedAndCompiledNode(this, HostPlatform, NonCodeProject, Plat, false));
 
                                         if (NonCodeFormalBuilds.ContainsKey(NonCodeProject.GameName))
                                         {
@@ -649,10 +709,10 @@ partial class GUBP
                                             {
                                                 if (PlatPair.TargetPlatform == Plat)
                                                 {
-														var NodeName = AddNode(new FormalBuildNode(this, NonCodeProject, HostPlatform, new List<UnrealTargetPlatform>() { Plat }, new List<UnrealTargetConfiguration>() { PlatPair.TargetConfig }));                                                    
+														var NodeName = BranchConfig.AddNode(new FormalBuildNode(this, NonCodeProject, HostPlatform, new List<UnrealTargetPlatform>() { Plat }, new List<UnrealTargetConfiguration>() { PlatPair.TargetConfig }));                                                    
 													// we don't want this delayed
 													// this would normally wait for the testing phase, we just want to build it right away
-													RemovePseudodependencyFromNode(
+													BranchConfig.RemovePseudodependencyFromNode(
 														CookNode.StaticGetFullName(HostPlatform, NonCodeProject, CookedPlatform),
 														WaitForTestShared.StaticGetFullName());
 													string BuildAgentSharingGroup = "";
@@ -663,20 +723,20 @@ partial class GUBP
 														{
 															BuildAgentSharingGroup = "";
 														}
-														FindNode(CookNode.StaticGetFullName(HostPlatform, NonCodeProject, CookedPlatform)).AgentSharingGroup = BuildAgentSharingGroup;
-														FindNode(NodeName).AgentSharingGroup = BuildAgentSharingGroup;
+														BranchConfig.FindNode(CookNode.StaticGetFullName(HostPlatform, NonCodeProject, CookedPlatform)).AgentSharingGroup = BuildAgentSharingGroup;
+														BranchConfig.FindNode(NodeName).AgentSharingGroup = BuildAgentSharingGroup;
 													}
 													else
 													{
 														//GUBPNodes[NodeName].AgentSharingGroup = FormalAgentSharingGroup;
 															if (Plat == UnrealTargetPlatform.XboxOne)
 														{
-															FindNode(NodeName).AgentSharingGroup = "";
+															BranchConfig.FindNode(NodeName).AgentSharingGroup = "";
 														}
 													}
 													if (PlatPair.bTest)
 													{
-														AddNode(new FormalBuildTestNode(this, NonCodeProject, HostPlatform, Plat, PlatPair.TargetConfig));
+														BranchConfig.AddNode(new FormalBuildTestNode(this, NonCodeProject, HostPlatform, Plat, PlatPair.TargetConfig));
 												}
 											}
 										}
@@ -690,11 +750,11 @@ partial class GUBP
 										foreach (var Test in GameTests)
 										{
 											var TestName = Test.Key + "_" + Plat.ToString();
-											ThisMonoGameTestNodes.Add(AddNode(new UATTestNode(this, HostPlatform, NonCodeProject, TestName, Test.Value, CookedAgentSharingGroup, false, RequiredPlatforms)));
+											ThisMonoGameTestNodes.Add(BranchConfig.AddNode(new UATTestNode(this, BranchConfig, HostPlatform, NonCodeProject, TestName, Test.Value, CookedAgentSharingGroup, false, RequiredPlatforms)));
 										}
 										if (ThisMonoGameTestNodes.Count > 0)
                                             {
-											GameTestNodes.Add(AddNode(new GameAggregateNode(this, HostPlatform, NonCodeProject, "CookedTests_" + Plat.ToString() + "_" + Kind.ToString() + HostPlatformNode.StaticGetHostPlatformSuffix(HostPlatform), ThisMonoGameTestNodes)));
+											GameTestNodes.Add(BranchConfig.AddNode(new GameAggregateNode(this, HostPlatform, NonCodeProject, "CookedTests_" + Plat.ToString() + "_" + Kind.ToString() + HostPlatformNode.StaticGetHostPlatformSuffix(HostPlatform), ThisMonoGameTestNodes)));
 										}
 									}
 								}
@@ -736,13 +796,13 @@ partial class GUBP
 #endif
                 if (GameTestNodes.Count > 0)
                 {
-                    AddNode(new GameAggregateNode(this, HostPlatform, Branch.BaseEngineProject, "AllCookedTests", GameTestNodes));
+                    BranchConfig.AddNode(new GameAggregateNode(this, HostPlatform, Branch.BaseEngineProject, "AllCookedTests", GameTestNodes));
                 }
             }
 			
 			if(HostPlatform == MakeFeaturePacksNode.GetDefaultBuildPlatform(this) && !BranchOptions.bNoInstalledEngine)
 			{
-				AddNode(new MakeFeaturePacksNode(HostPlatform, Branch.AllProjects.Where(x => MakeFeaturePacksNode.IsFeaturePack(x))));
+				BranchConfig.AddNode(new MakeFeaturePacksNode(HostPlatform, Branch.AllProjects.Where(x => MakeFeaturePacksNode.IsFeaturePack(x))));
 			}
 
             foreach (var CodeProj in Branch.CodeProjects)
@@ -761,10 +821,10 @@ partial class GUBP
 
 				if (!BranchOptions.ExcludePlatformsForEditor.Contains(HostPlatform) && !Options.bIsNonCode)
 				{
-					EditorGameNode Node = (EditorGameNode)TryFindNode(EditorGameNode.StaticGetFullName(HostPlatform, CodeProj));
+					EditorGameNode Node = (EditorGameNode)BranchConfig.TryFindNode(EditorGameNode.StaticGetFullName(HostPlatform, CodeProj));
 					if(Node == null)
 					{
-					AddNode(new EditorGameNode(this, HostPlatform, CodeProj));
+					BranchConfig.AddNode(new EditorGameNode(this, HostPlatform, CodeProj));
 				}
 					else
 					{
@@ -785,15 +845,15 @@ partial class GUBP
 						}
 						foreach (var Test in EditorTests)
 						{
-							EditorTestNodes.Add(AddNode(new UATTestNode(this, HostPlatform, CodeProj, Test.Key, Test.Value, AgentSharingGroup)));
-							if (!Options.bTestWithShared || !HasNode(WaitForTestShared.StaticGetFullName()))
+							EditorTestNodes.Add(BranchConfig.AddNode(new UATTestNode(this, BranchConfig, HostPlatform, CodeProj, Test.Key, Test.Value, AgentSharingGroup)));
+							if (!Options.bTestWithShared || !BranchConfig.HasNode(WaitForTestShared.StaticGetFullName()))
 							{
-								RemovePseudodependencyFromNode((UATTestNode.StaticGetFullName(HostPlatform, CodeProj, Test.Key)), WaitForTestShared.StaticGetFullName());
+								BranchConfig.RemovePseudodependencyFromNode((UATTestNode.StaticGetFullName(HostPlatform, CodeProj, Test.Key)), WaitForTestShared.StaticGetFullName());
 							}
 						}
 						if (EditorTestNodes.Count > 0)
 						{
-							AddNode(new GameAggregateNode(this, HostPlatform, CodeProj, "AllEditorTests", EditorTestNodes));
+							BranchConfig.AddNode(new GameAggregateNode(this, HostPlatform, CodeProj, "AllEditorTests", EditorTestNodes));
 						}
 					}
                 }
@@ -823,9 +883,9 @@ partial class GUBP
                             }
 							if(Plat == UnrealTargetPlatform.Win32 && Target.Rules.GUBP_BuildWindowsXPMonolithics())
 							{
-								if(!HasNode(GamePlatformMonolithicsNode.StaticGetFullName(HostPlatform, CodeProj, Plat, true)))
+								if(!BranchConfig.HasNode(GamePlatformMonolithicsNode.StaticGetFullName(HostPlatform, CodeProj, Plat, true)))
 								{
-									AddNode(new GamePlatformMonolithicsNode(this, HostPlatform, ActivePlatforms, CodeProj, Plat, true));
+									BranchConfig.AddNode(new GamePlatformMonolithicsNode(this, BranchConfig, HostPlatform, ActivePlatforms, CodeProj, Plat, true));
 								}
 							}
                             if (ActivePlatforms.Contains(Plat))
@@ -838,13 +898,13 @@ partial class GUBP
                                 {
                                     GamePlatforms.Add(Plat);
                                 }
-                                if (!HasNode(GamePlatformMonolithicsNode.StaticGetFullName(HostPlatform, CodeProj, Plat)))
+                                if (!BranchConfig.HasNode(GamePlatformMonolithicsNode.StaticGetFullName(HostPlatform, CodeProj, Plat)))
                                 {
 									if(GamePlatformMonolithicsNode.HasPrecompiledTargets(CodeProj, HostPlatform, Plat))
 									{
-	                                    AddNode(new GamePlatformMonolithicsNode(this, HostPlatform, ActivePlatforms, CodeProj, Plat, InPrecompiled: true));
+	                                    BranchConfig.AddNode(new GamePlatformMonolithicsNode(this, BranchConfig, HostPlatform, ActivePlatforms, CodeProj, Plat, InPrecompiled: true));
 									}
-                                    AddNode(new GamePlatformMonolithicsNode(this, HostPlatform, ActivePlatforms, CodeProj, Plat));
+                                    BranchConfig.AddNode(new GamePlatformMonolithicsNode(this, BranchConfig, HostPlatform, ActivePlatforms, CodeProj, Plat));
                                 }
                                 var FormalBuildConfigs = Target.Rules.GUBP_GetConfigsForFormalBuilds_MonolithicOnly(HostPlatform);
 								if (!AdditionalPlatforms.Contains(Plat) && (BranchOptions.ProjectsToCook.Contains(CodeProj.GameName) || BranchOptions.ProjectsToCook.Count == 0))
@@ -854,13 +914,13 @@ partial class GUBP
 									{
 										CookedPlatform = Target.Rules.GUBP_AlternateCookPlatform(HostPlatform, CookedPlatform);
 									}
-									if (!HasNode(CookNode.StaticGetFullName(HostPlatform, CodeProj, CookedPlatform)))
+									if (!BranchConfig.HasNode(CookNode.StaticGetFullName(HostPlatform, CodeProj, CookedPlatform)))
 									{
-										AddNode(new CookNode(this, HostPlatform, CodeProj, Plat, CookedPlatform));
+										BranchConfig.AddNode(new CookNode(this, BranchConfig, HostPlatform, CodeProj, Plat, CookedPlatform));
 									}
-									if (!HasNode(GamePlatformCookedAndCompiledNode.StaticGetFullName(HostPlatform, CodeProj, Plat)))
+									if (!BranchConfig.HasNode(GamePlatformCookedAndCompiledNode.StaticGetFullName(HostPlatform, CodeProj, Plat)))
 									{
-										AddNode(new GamePlatformCookedAndCompiledNode(this, HostPlatform, CodeProj, Plat, true));
+										BranchConfig.AddNode(new GamePlatformCookedAndCompiledNode(this, HostPlatform, CodeProj, Plat, true));
 									}
 									
 
@@ -871,28 +931,28 @@ partial class GUBP
 										{
 											if (Plat == Config.TargetPlatform)
 											{
-												FormalNodeName = AddNode(new FormalBuildNode(this, CodeProj, HostPlatform, InClientTargetPlatforms: new List<UnrealTargetPlatform>() { Config.TargetPlatform }, InClientConfigs: new List<UnrealTargetConfiguration>() { Config.TargetConfig }, InClientNotGame: true));
+												FormalNodeName = BranchConfig.AddNode(new FormalBuildNode(this, CodeProj, HostPlatform, InClientTargetPlatforms: new List<UnrealTargetPlatform>() { Config.TargetPlatform }, InClientConfigs: new List<UnrealTargetConfiguration>() { Config.TargetConfig }, InClientNotGame: true));
 											}
 										}
 										else if (Kind == TargetRules.TargetType.Server)
 										{
 											if (Plat == Config.TargetPlatform)
 											{
-												FormalNodeName = AddNode(new FormalBuildNode(this, CodeProj, HostPlatform, InServerTargetPlatforms: new List<UnrealTargetPlatform>() { Config.TargetPlatform }, InServerConfigs: new List<UnrealTargetConfiguration>() { Config.TargetConfig }));
+												FormalNodeName = BranchConfig.AddNode(new FormalBuildNode(this, CodeProj, HostPlatform, InServerTargetPlatforms: new List<UnrealTargetPlatform>() { Config.TargetPlatform }, InServerConfigs: new List<UnrealTargetConfiguration>() { Config.TargetConfig }));
 											}
 										}
 										else if (Kind == TargetRules.TargetType.Game)
 										{
 											if (Plat == Config.TargetPlatform)
 											{
-												FormalNodeName = AddNode(new FormalBuildNode(this, CodeProj, HostPlatform, InClientTargetPlatforms: new List<UnrealTargetPlatform>() { Config.TargetPlatform }, InClientConfigs: new List<UnrealTargetConfiguration>() { Config.TargetConfig }));
+												FormalNodeName = BranchConfig.AddNode(new FormalBuildNode(this, CodeProj, HostPlatform, InClientTargetPlatforms: new List<UnrealTargetPlatform>() { Config.TargetPlatform }, InClientConfigs: new List<UnrealTargetConfiguration>() { Config.TargetConfig }));
 											}
 										}
 										if (FormalNodeName != null)
 										{
 											// we don't want this delayed
 											// this would normally wait for the testing phase, we just want to build it right away
-											RemovePseudodependencyFromNode(
+											BranchConfig.RemovePseudodependencyFromNode(
 												CookNode.StaticGetFullName(HostPlatform, CodeProj, CookedPlatform),
 												WaitForTestShared.StaticGetFullName());
 											string BuildAgentSharingGroup = "";
@@ -904,20 +964,20 @@ partial class GUBP
 												{
 													BuildAgentSharingGroup = "";
 												}
-												FindNode(CookNode.StaticGetFullName(HostPlatform, CodeProj, CookedPlatform)).AgentSharingGroup = BuildAgentSharingGroup;
-												FindNode(FormalNodeName).AgentSharingGroup = BuildAgentSharingGroup;
+												BranchConfig.FindNode(CookNode.StaticGetFullName(HostPlatform, CodeProj, CookedPlatform)).AgentSharingGroup = BuildAgentSharingGroup;
+												BranchConfig.FindNode(FormalNodeName).AgentSharingGroup = BuildAgentSharingGroup;
 											}
 											else
 											{
 												//GUBPNodes[FormalNodeName].AgentSharingGroup = FormalAgentSharingGroup;
 												if (Plat == UnrealTargetPlatform.XboxOne)
 												{
-													FindNode(FormalNodeName).AgentSharingGroup = "";
+													BranchConfig.FindNode(FormalNodeName).AgentSharingGroup = "";
 												}
 											}
 											if (Config.bTest)
 											{
-												AddNode(new FormalBuildTestNode(this, CodeProj, HostPlatform, Plat, Config.TargetConfig));
+												BranchConfig.AddNode(new FormalBuildTestNode(this, CodeProj, HostPlatform, Plat, Config.TargetConfig));
 											}																				
 										}
 									}
@@ -931,15 +991,15 @@ partial class GUBP
 										foreach (var Test in GameTests)
 										{
 											var TestNodeName = Test.Key + "_" + Plat.ToString();
-											ThisMonoGameTestNodes.Add(AddNode(new UATTestNode(this, HostPlatform, CodeProj, TestNodeName, Test.Value, CookedAgentSharingGroup, false, RequiredPlatforms)));
-											if (!Options.bTestWithShared || !HasNode(WaitForTestShared.StaticGetFullName()))
+											ThisMonoGameTestNodes.Add(BranchConfig.AddNode(new UATTestNode(this, BranchConfig, HostPlatform, CodeProj, TestNodeName, Test.Value, CookedAgentSharingGroup, false, RequiredPlatforms)));
+											if (!Options.bTestWithShared || !BranchConfig.HasNode(WaitForTestShared.StaticGetFullName()))
 											{
-												RemovePseudodependencyFromNode((UATTestNode.StaticGetFullName(HostPlatform, CodeProj, TestNodeName)), WaitForTestShared.StaticGetFullName());
+												BranchConfig.RemovePseudodependencyFromNode((UATTestNode.StaticGetFullName(HostPlatform, CodeProj, TestNodeName)), WaitForTestShared.StaticGetFullName());
 											}
 										}
 										if (ThisMonoGameTestNodes.Count > 0)
 										{
-											GameTestNodes.Add(AddNode(new GameAggregateNode(this, HostPlatform, CodeProj, "CookedTests_" + Plat.ToString() + "_" + Kind.ToString() + HostPlatformNode.StaticGetHostPlatformSuffix(HostPlatform), ThisMonoGameTestNodes)));
+											GameTestNodes.Add(BranchConfig.AddNode(new GameAggregateNode(this, HostPlatform, CodeProj, "CookedTests_" + Plat.ToString() + "_" + Kind.ToString() + HostPlatformNode.StaticGetHostPlatformSuffix(HostPlatform), ThisMonoGameTestNodes)));
 										}
 									}
 								}
@@ -964,13 +1024,13 @@ partial class GUBP
 							foreach (var Test in ClientServerTests)
 							{
 								var TestNodeName = Test.Key + "_" + GamePlatform.ToString() + "_" + ServerPlatform.ToString();
-								GameTestNodes.Add(AddNode(new UATTestNode(this, HostPlatform, CodeProj, TestNodeName, Test.Value, CookedAgentSharingGroup, false, RequiredPlatforms)));
+								GameTestNodes.Add(BranchConfig.AddNode(new UATTestNode(this, BranchConfig, HostPlatform, CodeProj, TestNodeName, Test.Value, CookedAgentSharingGroup, false, RequiredPlatforms)));
 							}
 						}
 					}
 					if (GameTestNodes.Count > 0)
 					{
-						AddNode(new GameAggregateNode(this, HostPlatform, CodeProj, "AllCookedTests", GameTestNodes));
+						BranchConfig.AddNode(new GameAggregateNode(this, HostPlatform, CodeProj, "AllCookedTests", GameTestNodes));
 					}
 				}
 			}
@@ -1002,48 +1062,48 @@ partial class GUBP
                 }
                 if (PromotedHosts.Count > 0)
                 {
-                    AddNode(new GameAggregatePromotableNode(this, PromotedHosts, ActivePlatforms, CodeProj, true, bNoIOSOnPC));
+                    BranchConfig.AddNode(new GameAggregatePromotableNode(this, PromotedHosts, ActivePlatforms, CodeProj, true, bNoIOSOnPC));
                     if (AnySeparate)
                     {
-                        AddNode(new WaitForGamePromotionUserInput(this, CodeProj, false));
-                        AddNode(new GameLabelPromotableNode(this, CodeProj, false));
-                        AddNode(new WaitForGamePromotionUserInput(this, CodeProj, true));
-                        AddNode(new GameLabelPromotableNode(this, CodeProj, true));
+                        BranchConfig.AddNode(new WaitForGamePromotionUserInput(this, CodeProj, false));
+                        BranchConfig.AddNode(new GameLabelPromotableNode(this, CodeProj, false));
+                        BranchConfig.AddNode(new WaitForGamePromotionUserInput(this, CodeProj, true));
+                        BranchConfig.AddNode(new GameLabelPromotableNode(this, CodeProj, true));
                     }
                 }
             }
         }
         if (NumSharedAllHosts > 0)
         {
-            AddNode(new GameAggregatePromotableNode(this, HostPlatforms, ActivePlatforms, Branch.BaseEngineProject, false, bNoIOSOnPC));
-
-            AddNode(new SharedAggregatePromotableNode(this, HostPlatforms, ActivePlatforms));
-            AddNode(new WaitForSharedPromotionUserInput(this, false));
-            AddNode(new SharedLabelPromotableNode(this, false));
-            AddNode(new SharedLabelPromotableSuccessNode());
-            AddNode(new WaitForTestShared(this));
-            AddNode(new WaitForSharedPromotionUserInput(this, true));
-            AddNode(new SharedLabelPromotableNode(this, true));			
+            BranchConfig.AddNode(new GameAggregatePromotableNode(this, HostPlatforms, ActivePlatforms, Branch.BaseEngineProject, false, bNoIOSOnPC));
+
+            BranchConfig.AddNode(new SharedAggregatePromotableNode(this, HostPlatforms, ActivePlatforms));
+            BranchConfig.AddNode(new WaitForSharedPromotionUserInput(this, false));
+            BranchConfig.AddNode(new SharedLabelPromotableNode(this, false));
+            BranchConfig.AddNode(new SharedLabelPromotableSuccessNode());
+            BranchConfig.AddNode(new WaitForTestShared(this));
+            BranchConfig.AddNode(new WaitForSharedPromotionUserInput(this, true));
+            BranchConfig.AddNode(new SharedLabelPromotableNode(this, true));			
         }
         foreach (var HostPlatform in HostPlatforms)
         {
-            AddCustomNodes(HostPlatform, ActivePlatforms);
+            AddCustomNodes(BranchConfig, HostPlatform, ActivePlatforms);
         }
         
-        if (HasNode(ToolsForCompileNode.StaticGetFullName(UnrealTargetPlatform.Win64)))
+        if (BranchConfig.HasNode(ToolsForCompileNode.StaticGetFullName(UnrealTargetPlatform.Win64)))
         {
-			if (HasNode(GamePlatformMonolithicsNode.StaticGetFullName(UnrealTargetPlatform.Mac, Branch.BaseEngineProject, UnrealTargetPlatform.IOS)) && HasNode(ToolsNode.StaticGetFullName(UnrealTargetPlatform.Win64)))
+			if (BranchConfig.HasNode(GamePlatformMonolithicsNode.StaticGetFullName(UnrealTargetPlatform.Mac, Branch.BaseEngineProject, UnrealTargetPlatform.IOS)) && BranchConfig.HasNode(ToolsNode.StaticGetFullName(UnrealTargetPlatform.Win64)))
 			{
 				//AddNode(new IOSOnPCTestNode(this)); - Disable IOSOnPCTest until a1011 crash is fixed
 			}
 			//AddNode(new VSExpressTestNode(this));
 			if (ActivePlatforms.Contains(UnrealTargetPlatform.Linux) && !BranchOptions.ExcludePlatformsForEditor.Contains(UnrealTargetPlatform.Linux))
 			{
-				AddNode(new RootEditorCrossCompileLinuxNode(UnrealTargetPlatform.Win64));
+				BranchConfig.AddNode(new RootEditorCrossCompileLinuxNode(UnrealTargetPlatform.Win64));
 			}
             if (!bPreflightBuild)
             {
-                AddNode(new CleanSharedTempStorageNode(this));
+                BranchConfig.AddNode(new CleanSharedTempStorageNode(this, BranchConfig));
             }
         }
 #if false
@@ -1065,16 +1125,25 @@ partial class GUBP
 		// Remove all the pseudo-dependencies on nodes specified in the branch options
 		foreach(string NodeToRemovePseudoDependencies in BranchOptions.NodesToRemovePseudoDependencies)
 		{
-			GUBPNode Node = TryFindNode(NodeToRemovePseudoDependencies);
+			GUBPNode Node = BranchConfig.TryFindNode(NodeToRemovePseudoDependencies);
 			if(Node != null)
 			{
 				Node.FullNamesOfPseudosependencies.Clear();
 			}
 		}
 
+		// Calculate the frequency overrides
+		Dictionary<string, int> FrequencyOverrides = ApplyFrequencyBarriers(BranchConfig.GUBPNodes, BranchConfig.GUBPAggregates, BranchOptions.FrequencyBarriers);
+
+		// Convert the GUBPNodes and GUBPAggregates maps into lists
+		AllNodes = new List<NodeInfo>();
+		AllNodes.AddRange(BranchConfig.GUBPNodes.Values.Select(x => new NodeInfo{ Name = x.GetFullName(), Node = x }));
+
+		AllAggregates = new List<AggregateInfo>();
+		AllAggregates.AddRange(BranchConfig.GUBPAggregates.Values.Select(x => new AggregateInfo{ Name = x.GetFullName(), Node = x }));
+
 		// Calculate the frequencies for each node
-		Dictionary<string, int> FrequencyOverrides = ApplyFrequencyBarriers(GUBPNodes, GUBPAggregates, BranchOptions.FrequencyBarriers);
-		foreach(NodeInfo Node in GUBPNodes.Values)
+		foreach(NodeInfo Node in AllNodes)
 		{
 			Node.FrequencyShift = Node.Node.CISFrequencyQuantumShift(this);
             Node.FrequencyShift = HackFrequency(this, BranchName, Node, Node.FrequencyShift);
@@ -1087,7 +1156,7 @@ partial class GUBP
 		}
 
 		// Get the email list for each node
-		FindEmailsForNodes(GUBPNodes.Values);
+		FindEmailsForNodes(AllNodes);
 
 		// Add aggregate nodes for every project in the branch
 		foreach (BranchInfo.BranchUProject GameProj in Branch.AllProjects)
@@ -1097,7 +1166,7 @@ partial class GUBP
 			{
 				NodeNames.Add(GameAggregatePromotableNode.StaticGetFullName(GameProj));
 			}
-			foreach (NodeInfo Node in GUBPNodes.Values)
+			foreach (NodeInfo Node in AllNodes)
 			{
 				if (Node.Node.GameNameIfAnyForTempStorage() == GameProj.GameName)
 				{
@@ -1106,12 +1175,12 @@ partial class GUBP
 			}
 			if (NodeNames.Count > 0)
 			{
-				AddNode(new FullGameAggregateNode(GameProj.GameName, NodeNames));
+				BranchConfig.AddNode(new FullGameAggregateNode(GameProj.GameName, NodeNames));
 			}
 		}
 	}
 
-	private static Dictionary<string, int> ApplyFrequencyBarriers(Dictionary<string, NodeInfo> GUBPNodes, Dictionary<string, AggregateInfo> GUBPAggregates, Dictionary<string, sbyte> FrequencyBarriers)
+	private static Dictionary<string, int> ApplyFrequencyBarriers(Dictionary<string, GUBPNode> GUBPNodes, Dictionary<string, GUBPAggregateNode> GUBPAggregates, Dictionary<string, sbyte> FrequencyBarriers)
 	{
 		// Make sure that everything that's listed as a frequency barrier is completed with the given interval
 		Dictionary<string, int> FrequencyOverrides = new Dictionary<string,int>();
@@ -1126,22 +1195,22 @@ partial class GUBP
 			{
 				if(IncludedNodes.Add(SearchNodes[Idx]))
 				{
-					NodeInfo Node;
+					GUBPNode Node;
 					if(GUBPNodes.TryGetValue(SearchNodes[Idx], out Node))
 					{
-						SearchNodes.AddRange(Node.Node.FullNamesOfDependencies);
-						SearchNodes.AddRange(Node.Node.FullNamesOfPseudosependencies);
+						SearchNodes.AddRange(Node.FullNamesOfDependencies);
+						SearchNodes.AddRange(Node.FullNamesOfPseudosependencies);
 					}
 					else
 					{
-						AggregateInfo Aggregate;
+						GUBPAggregateNode Aggregate;
 						if(GUBPAggregates.TryGetValue(SearchNodes[Idx], out Aggregate))
 						{
-							SearchNodes.AddRange(Aggregate.Node.Dependencies);
+							SearchNodes.AddRange(Aggregate.Dependencies);
 						}
 						else
 						{
-							throw new AutomationException("Couldn't find referenced node {0} in graph", Aggregate.Name);
+							throw new AutomationException("Couldn't find referenced node {0} in graph", Aggregate.GetFullName());
 						}
 					}
 				}
diff --git a/Engine/Source/Programs/AutomationTool/GUBP/LegacyNodes.cs b/Engine/Source/Programs/AutomationTool/GUBP/LegacyNodes.cs
--- a/Engine/Source/Programs/AutomationTool/GUBP/LegacyNodes.cs
+++ b/Engine/Source/Programs/AutomationTool/GUBP/LegacyNodes.cs
@@ -1011,11 +1011,11 @@ partial class GUBP
         {
 			SetupSingleInternalToolsNode(InProgramTarget);
         }
-		public SingleInternalToolsNode(GUBP bp, UnrealTargetPlatform InHostPlatform, SingleTargetProperties InProgramTarget)
+		public SingleInternalToolsNode(GUBP bp, GUBPBranchConfig BranchConfig, UnrealTargetPlatform InHostPlatform, SingleTargetProperties InProgramTarget)
             : base(InHostPlatform)
         {
 			// Don't add rooteditor dependency if it isn't in the graph
-			var bRootEditorNodeDoesExit = bp.HasNode(RootEditorNode.StaticGetFullName(HostPlatform));
+			var bRootEditorNodeDoesExit = BranchConfig.HasNode(RootEditorNode.StaticGetFullName(HostPlatform));
 			SetupSingleInternalToolsNode(InProgramTarget, !bRootEditorNodeDoesExit && bp.BranchOptions.bNoEditorDependenciesForTools);
         }
 		private void SetupSingleInternalToolsNode(SingleTargetProperties InProgramTarget, bool bSkipRootEditorPsuedoDependency = true)
@@ -1253,7 +1253,7 @@ partial class GUBP
 		bool WithXp;
 		bool Precompiled; // If true, just builds targets which generate static libraries for the -UsePrecompiled option to UBT. If false, just build those that don't.
 
-        public GamePlatformMonolithicsNode(GUBP bp, UnrealTargetPlatform InHostPlatform, List<UnrealTargetPlatform> InActivePlatforms, BranchInfo.BranchUProject InGameProj, UnrealTargetPlatform InTargetPlatform, bool InWithXp = false, bool InPrecompiled = false)
+        public GamePlatformMonolithicsNode(GUBP bp, GUBPBranchConfig BranchConfig, UnrealTargetPlatform InHostPlatform, List<UnrealTargetPlatform> InActivePlatforms, BranchInfo.BranchUProject InGameProj, UnrealTargetPlatform InTargetPlatform, bool InWithXp = false, bool InPrecompiled = false)
             : base(InHostPlatform)
         {
             GameProj = InGameProj;
@@ -1274,14 +1274,14 @@ partial class GUBP
 				{
 					AddPseudodependency(EditorGameNode.StaticGetFullName(InHostPlatform, GameProj));
 				}
-                if (bp.HasNode(GamePlatformMonolithicsNode.StaticGetFullName(HostPlatform, bp.Branch.BaseEngineProject, TargetPlatform)))
+                if (BranchConfig.HasNode(GamePlatformMonolithicsNode.StaticGetFullName(HostPlatform, bp.Branch.BaseEngineProject, TargetPlatform)))
                 {
                     AddPseudodependency(GamePlatformMonolithicsNode.StaticGetFullName(HostPlatform, bp.Branch.BaseEngineProject, TargetPlatform));
                 }
             }
             else
             {
-                if (TargetPlatform != InHostPlatform && bp.HasNode(GamePlatformMonolithicsNode.StaticGetFullName(InHostPlatform, bp.Branch.BaseEngineProject, InHostPlatform, Precompiled: Precompiled)))
+                if (TargetPlatform != InHostPlatform && BranchConfig.HasNode(GamePlatformMonolithicsNode.StaticGetFullName(InHostPlatform, bp.Branch.BaseEngineProject, InHostPlatform, Precompiled: Precompiled)))
                 {
                     AddPseudodependency(GamePlatformMonolithicsNode.StaticGetFullName(InHostPlatform, bp.Branch.BaseEngineProject, InHostPlatform, Precompiled: Precompiled));
                 }
@@ -2195,7 +2195,7 @@ partial class GUBP
         bool bIsMassive;
 		
 
-        public CookNode(GUBP bp, UnrealTargetPlatform InHostPlatform, BranchInfo.BranchUProject InGameProj, UnrealTargetPlatform InTargetPlatform, string InCookPlatform)
+        public CookNode(GUBP bp, GUBPBranchConfig BranchConfig, UnrealTargetPlatform InHostPlatform, BranchInfo.BranchUProject InGameProj, UnrealTargetPlatform InTargetPlatform, string InCookPlatform)
             : base(InHostPlatform)
         {
             GameProj = InGameProj;
@@ -2239,7 +2239,7 @@ partial class GUBP
                 AgentSharingGroup = "SharedCooks" + StaticGetHostPlatformSuffix(HostPlatform);
 
                 // If the cook fails for the base engine, don't bother trying
-                if (InGameProj.GameName != bp.Branch.BaseEngineProject.GameName && bp.HasNode(CookNode.StaticGetFullName(HostPlatform, bp.Branch.BaseEngineProject, CookPlatform)))
+                if (InGameProj.GameName != bp.Branch.BaseEngineProject.GameName && BranchConfig.HasNode(CookNode.StaticGetFullName(HostPlatform, bp.Branch.BaseEngineProject, CookPlatform)))
                 {
                     AddPseudodependency(CookNode.StaticGetFullName(HostPlatform, bp.Branch.BaseEngineProject, CookPlatform));
                 }
@@ -2247,7 +2247,7 @@ partial class GUBP
                 // If the base cook platform fails, don't bother trying other ones
                 string BaseCookedPlatform = Platform.Platforms[HostPlatform].GetCookPlatform(false, false, "");
                 if (InGameProj.GameName == bp.Branch.BaseEngineProject.GameName && CookPlatform != BaseCookedPlatform &&
-                    bp.HasNode(CookNode.StaticGetFullName(HostPlatform, bp.Branch.BaseEngineProject, BaseCookedPlatform)))
+                    BranchConfig.HasNode(CookNode.StaticGetFullName(HostPlatform, bp.Branch.BaseEngineProject, BaseCookedPlatform)))
                 {
                     AddPseudodependency(CookNode.StaticGetFullName(HostPlatform, bp.Branch.BaseEngineProject, BaseCookedPlatform));
                 }
@@ -3067,7 +3067,7 @@ partial class GUBP
         List<UnrealTargetPlatform> DependsOnCooked;
         float ECPriority;
 
-        public UATTestNode(GUBP bp, UnrealTargetPlatform InHostPlatform, BranchInfo.BranchUProject InGameProj, string InTestName, string InUATCommandLine, string InAgentSharingGroup, bool InDependsOnEditor = true, List<UnrealTargetPlatform> InDependsOnCooked = null, float InECPriority = 0.0f)
+        public UATTestNode(GUBP bp, GUBPBranchConfig BranchConfig, UnrealTargetPlatform InHostPlatform, BranchInfo.BranchUProject InGameProj, string InTestName, string InUATCommandLine, string InAgentSharingGroup, bool InDependsOnEditor = true, List<UnrealTargetPlatform> InDependsOnCooked = null, float InECPriority = 0.0f)
             : base(InHostPlatform)
         {
             AgentSharingGroup = InAgentSharingGroup;
@@ -3105,7 +3105,7 @@ partial class GUBP
             // If the same test fails for the base engine, don't bother trying
             if (InGameProj.GameName != bp.Branch.BaseEngineProject.GameName)
             {
-                if (bp.HasNode(UATTestNode.StaticGetFullName(HostPlatform, bp.Branch.BaseEngineProject, TestName)))
+                if (BranchConfig.HasNode(UATTestNode.StaticGetFullName(HostPlatform, bp.Branch.BaseEngineProject, TestName)))
                 {
                     AddPseudodependency(UATTestNode.StaticGetFullName(HostPlatform, bp.Branch.BaseEngineProject, InTestName));
                 }
@@ -3115,14 +3115,14 @@ partial class GUBP
                     foreach (var Plat in DependsOnCooked)
                     {
                         var PlatTestName = "CookedGameTest_"  + Plat.ToString();
-                        if (bp.HasNode(UATTestNode.StaticGetFullName(HostPlatform, bp.Branch.BaseEngineProject, PlatTestName)))
+                        if (BranchConfig.HasNode(UATTestNode.StaticGetFullName(HostPlatform, bp.Branch.BaseEngineProject, PlatTestName)))
                         {
                             AddPseudodependency(UATTestNode.StaticGetFullName(HostPlatform, bp.Branch.BaseEngineProject, PlatTestName));
                             bFoundACook = true;
                         }
                     }
 
-                    if (!bFoundACook && bp.HasNode(UATTestNode.StaticGetFullName(HostPlatform, bp.Branch.BaseEngineProject, "EditorTest")))
+                    if (!bFoundACook && BranchConfig.HasNode(UATTestNode.StaticGetFullName(HostPlatform, bp.Branch.BaseEngineProject, "EditorTest")))
                     {
                         AddPseudodependency(UATTestNode.StaticGetFullName(HostPlatform, bp.Branch.BaseEngineProject, "EditorTest"));
                     }
@@ -3238,9 +3238,12 @@ partial class GUBP
 
     public class CleanSharedTempStorageNode : GUBPNode
     {
-        public CleanSharedTempStorageNode(GUBP bp)
+		HashSet<string> GameNames;
+
+        public CleanSharedTempStorageNode(GUBP bp, GUBPBranchConfig BranchConfig)
         {
-            var ToolsNode = bp.FindNode(ToolsForCompileNode.StaticGetFullName(UnrealTargetPlatform.Win64));
+            var ToolsNode = BranchConfig.FindNode(ToolsForCompileNode.StaticGetFullName(UnrealTargetPlatform.Win64));
+			GameNames = new HashSet<string>(BranchConfig.GUBPNodes.Values.Select(x => x.GameNameIfAnyForTempStorage()));
             AgentSharingGroup = ToolsNode.AgentSharingGroup;
         }
         public override float Priority()
@@ -3260,9 +3263,9 @@ partial class GUBP
         {
             {
                 var StartTime = DateTime.UtcNow;
-                foreach (var NodeToDo in bp.GUBPNodes)
+                foreach (string GameName in GameNames)
                 {
-                    TempStorage.CleanSharedTempStorageDirectory(NodeToDo.Value.Node.GameNameIfAnyForTempStorage());
+                    TempStorage.CleanSharedTempStorageDirectory(GameName);
                 }
                 var BuildDuration = (DateTime.UtcNow - StartTime).TotalMilliseconds;
                 Log("Took {0}s to clear temp storage of old files.", BuildDuration / 1000);
diff --git a/Engine/Source/Programs/AutomationTool/Scripts/BuildDocumentation.Automation.cs b/Engine/Source/Programs/AutomationTool/Scripts/BuildDocumentation.Automation.cs
--- a/Engine/Source/Programs/AutomationTool/Scripts/BuildDocumentation.Automation.cs
+++ b/Engine/Source/Programs/AutomationTool/Scripts/BuildDocumentation.Automation.cs
@@ -136,13 +136,13 @@ class BlueprintDocumentationNode : DocumentationNode
 
 public class DocumentationNodeAdder : GUBP.GUBPNodeAdder
 {
-	public override void AddNodes(GUBP bp, UnrealTargetPlatform InHostPlatform, List<UnrealTargetPlatform> InActivePlatforms)
+	public override void AddNodes(GUBP bp, GUBP.GUBPBranchConfig BranchConfig, UnrealTargetPlatform InHostPlatform, List<UnrealTargetPlatform> InActivePlatforms)
 	{
 		if(InHostPlatform == UnrealTargetPlatform.Win64 && !bp.BranchOptions.bNoDocumentation)
 		{
-			bp.AddNode(new ToolsForDocumentationNode(InHostPlatform));
-			bp.AddNode(new CodeDocumentationNode(InHostPlatform));
-			bp.AddNode(new BlueprintDocumentationNode(InHostPlatform));
+			BranchConfig.AddNode(new ToolsForDocumentationNode(InHostPlatform));
+			BranchConfig.AddNode(new CodeDocumentationNode(InHostPlatform));
+			BranchConfig.AddNode(new BlueprintDocumentationNode(InHostPlatform));
 		}
 	}
 }
diff --git a/Engine/Source/Programs/AutomationTool/Scripts/RocketBuild.Automation.cs b/Engine/Source/Programs/AutomationTool/Scripts/RocketBuild.Automation.cs
--- a/Engine/Source/Programs/AutomationTool/Scripts/RocketBuild.Automation.cs
+++ b/Engine/Source/Programs/AutomationTool/Scripts/RocketBuild.Automation.cs
@@ -74,7 +74,7 @@ namespace Rocket
 		{
 		}
 
-		public override void AddNodes(GUBP bp, UnrealTargetPlatform HostPlatform, List<UnrealTargetPlatform> ActivePlatforms)
+		public override void AddNodes(GUBP bp, GUBP.GUBPBranchConfig BranchConfig, UnrealTargetPlatform HostPlatform, List<UnrealTargetPlatform> ActivePlatforms)
 		{
 			if(!bp.BranchOptions.bNoInstalledEngine)
 			{
@@ -90,8 +90,8 @@ namespace Rocket
 				// Strip the host platform
 				if (StripRocketNode.IsRequiredForPlatform(HostPlatform))
 				{
-					bp.AddNode(new StripRocketToolsNode(HostPlatform, StrippedDir));
-					bp.AddNode(new StripRocketEditorNode(HostPlatform, StrippedDir));
+					BranchConfig.AddNode(new StripRocketToolsNode(BranchConfig, HostPlatform, StrippedDir));
+					BranchConfig.AddNode(new StripRocketEditorNode(BranchConfig, HostPlatform, StrippedDir));
 				}
 
 				// Strip all the target platforms that are built on this host
@@ -99,22 +99,22 @@ namespace Rocket
 				{
 					if (GetSourceHostPlatform(bp, HostPlatform, TargetPlatform) == HostPlatform && StripRocketNode.IsRequiredForPlatform(TargetPlatform))
 					{
-						bp.AddNode(new StripRocketMonolithicsNode(bp, HostPlatform, TargetPlatform, StrippedDir));
+						BranchConfig.AddNode(new StripRocketMonolithicsNode(bp, BranchConfig, HostPlatform, TargetPlatform, StrippedDir));
 					}
 				}
 
 				// Build the DDC
-				bp.AddNode(new BuildDerivedDataCacheNode(HostPlatform, GetCookPlatforms(HostPlatform, TargetPlatforms), CurrentFeaturePacks));
+				BranchConfig.AddNode(new BuildDerivedDataCacheNode(HostPlatform, GetCookPlatforms(HostPlatform, TargetPlatforms), CurrentFeaturePacks));
 
 				// Generate a list of files that needs to be copied for each target platform
-				bp.AddNode(new FilterRocketNode(bp, HostPlatform, TargetPlatforms, CurrentFeaturePacks, CurrentTemplates));
+				BranchConfig.AddNode(new FilterRocketNode(bp, BranchConfig, HostPlatform, TargetPlatforms, CurrentFeaturePacks, CurrentTemplates));
 
 				// Copy the install to the output directory
 				string LocalOutputDir = CommandUtils.CombinePaths(CommandUtils.CmdEnv.LocalRoot, "LocalBuilds", "Rocket", CommandUtils.GetGenericPlatformName(HostPlatform));
-				bp.AddNode(new GatherRocketNode(HostPlatform, TargetPlatforms, LocalOutputDir));
+				BranchConfig.AddNode(new GatherRocketNode(BranchConfig, HostPlatform, TargetPlatforms, LocalOutputDir));
 
 				// Add the aggregate node for the entire install
-				GUBP.SharedAggregatePromotableNode PromotableNode = (GUBP.SharedAggregatePromotableNode)bp.FindAggregateNode(GUBP.SharedAggregatePromotableNode.StaticGetFullName());
+				GUBP.SharedAggregatePromotableNode PromotableNode = (GUBP.SharedAggregatePromotableNode)BranchConfig.FindAggregateNode(GUBP.SharedAggregatePromotableNode.StaticGetFullName());
 				PromotableNode.AddDependency(FilterRocketNode.StaticGetFullName(HostPlatform));
 				PromotableNode.AddDependency(BuildDerivedDataCacheNode.StaticGetFullName(HostPlatform));
 
@@ -124,7 +124,7 @@ namespace Rocket
 					string GitConfigRelativePath = "Engine/Build/Git/UnrealBot.ini";
 					if(CommandUtils.FileExists(CommandUtils.CombinePaths(CommandUtils.CmdEnv.LocalRoot, GitConfigRelativePath)))
 					{
-						bp.AddNode(new BuildGitPromotable(bp, HostPlatform, GitConfigRelativePath));
+						BranchConfig.AddNode(new BuildGitPromotable(bp, HostPlatform, GitConfigRelativePath));
 						PromotableNode.AddDependency(BuildGitPromotable.StaticGetFullName(HostPlatform));
 					}
 				}
@@ -141,16 +141,16 @@ namespace Rocket
 				}
 
 				// Publish the install to the network
-				bp.AddNode(new PublishRocketNode(HostPlatform, LocalOutputDir, PublishedEngineDir));
-				bp.AddNode(new PublishRocketSymbolsNode(bp, HostPlatform, TargetPlatforms, PublishedEngineDir + "Symbols"));
+				BranchConfig.AddNode(new PublishRocketNode(HostPlatform, LocalOutputDir, PublishedEngineDir));
+				BranchConfig.AddNode(new PublishRocketSymbolsNode(bp, HostPlatform, TargetPlatforms, PublishedEngineDir + "Symbols"));
 
 				// Add a dependency on this being published as part of the shared promotable being labeled
-				GUBP.SharedLabelPromotableSuccessNode LabelPromotableNode = (GUBP.SharedLabelPromotableSuccessNode)bp.FindAggregateNode(GUBP.SharedLabelPromotableSuccessNode.StaticGetFullName());
+				GUBP.SharedLabelPromotableSuccessNode LabelPromotableNode = (GUBP.SharedLabelPromotableSuccessNode)BranchConfig.FindAggregateNode(GUBP.SharedLabelPromotableSuccessNode.StaticGetFullName());
 				LabelPromotableNode.AddDependency(PublishRocketNode.StaticGetFullName(HostPlatform));
 				LabelPromotableNode.AddDependency(PublishRocketSymbolsNode.StaticGetFullName(HostPlatform));
 
 				// Add dependencies on a promotable to do these steps too
-				GUBP.WaitForSharedPromotionUserInput WaitForPromotionNode = (GUBP.WaitForSharedPromotionUserInput)bp.FindNode(GUBP.WaitForSharedPromotionUserInput.StaticGetFullName(true));
+				GUBP.WaitForSharedPromotionUserInput WaitForPromotionNode = (GUBP.WaitForSharedPromotionUserInput)BranchConfig.FindNode(GUBP.WaitForSharedPromotionUserInput.StaticGetFullName(true));
 				WaitForPromotionNode.AddDependency(PublishRocketNode.StaticGetFullName(HostPlatform));
 				WaitForPromotionNode.AddDependency(PublishRocketSymbolsNode.StaticGetFullName(HostPlatform));
 
@@ -159,10 +159,10 @@ namespace Rocket
 				{
 					string WaitForTrigger = GUBP.WaitForSharedPromotionUserInput.StaticGetFullName(false);
 
-					GatherRocketNode GatherRocket = (GatherRocketNode)bp.FindNode(GatherRocketNode.StaticGetFullName(HostPlatform));
+					GatherRocketNode GatherRocket = (GatherRocketNode)BranchConfig.FindNode(GatherRocketNode.StaticGetFullName(HostPlatform));
 					GatherRocket.AddDependency(WaitForTrigger);
 
-					PublishRocketSymbolsNode PublishRocketSymbols = (PublishRocketSymbolsNode)bp.FindNode(PublishRocketSymbolsNode.StaticGetFullName(HostPlatform));
+					PublishRocketSymbolsNode PublishRocketSymbols = (PublishRocketSymbolsNode)BranchConfig.FindNode(PublishRocketSymbolsNode.StaticGetFullName(HostPlatform));
 					PublishRocketSymbols.AddDependency(WaitForTrigger);
 				}
 			}
@@ -327,12 +327,14 @@ namespace Rocket
 
 	public abstract class StripRocketNode : GUBP.HostPlatformNode
 	{
+		public GUBP.GUBPBranchConfig BranchConfig;
 		public UnrealTargetPlatform TargetPlatform;
 		public string StrippedDir;
 		public List<string> NodesToStrip = new List<string>();
 
-		public StripRocketNode(UnrealTargetPlatform InHostPlatform, UnrealTargetPlatform InTargetPlatform, string InStrippedDir) : base(InHostPlatform)
+		public StripRocketNode(GUBP.GUBPBranchConfig InBranchConfig, UnrealTargetPlatform InHostPlatform, UnrealTargetPlatform InTargetPlatform, string InStrippedDir) : base(InHostPlatform)
 		{
+			BranchConfig = InBranchConfig;
 			TargetPlatform = InTargetPlatform;
 			StrippedDir = InStrippedDir;
 		}
@@ -366,7 +368,7 @@ namespace Rocket
 			List<string> TargetPaths = new List<string>();
 			foreach(string NodeToStrip in NodesToStrip)
 			{
-				GUBP.GUBPNode Node = bp.FindNode(NodeToStrip);
+				GUBP.GUBPNode Node = BranchConfig.FindNode(NodeToStrip);
 				foreach(string DependencyBuildProduct in Node.BuildProducts)
 				{
 					string RelativePath = CommandUtils.StripBaseDirectory(Path.GetFullPath(DependencyBuildProduct), InputDir);
@@ -400,8 +402,8 @@ namespace Rocket
 
 	public class StripRocketToolsNode : StripRocketNode
 	{
-		public StripRocketToolsNode(UnrealTargetPlatform InHostPlatform, string InStrippedDir)
-			: base(InHostPlatform, InHostPlatform, InStrippedDir)
+		public StripRocketToolsNode(GUBP.GUBPBranchConfig InBranchConfig, UnrealTargetPlatform InHostPlatform, string InStrippedDir)
+			: base(InBranchConfig, InHostPlatform, InHostPlatform, InStrippedDir)
 		{
 			AddNodeToStrip(GUBP.ToolsForCompileNode.StaticGetFullName(HostPlatform));
 			AddNodeToStrip(GUBP.ToolsNode.StaticGetFullName(HostPlatform));
@@ -421,8 +423,8 @@ namespace Rocket
 
 	public class StripRocketEditorNode : StripRocketNode
 	{
-		public StripRocketEditorNode(UnrealTargetPlatform InHostPlatform, string InStrippedDir)
-			: base(InHostPlatform, InHostPlatform, InStrippedDir)
+		public StripRocketEditorNode(GUBP.GUBPBranchConfig InBranchConfig, UnrealTargetPlatform InHostPlatform, string InStrippedDir)
+			: base(InBranchConfig, InHostPlatform, InHostPlatform, InStrippedDir)
 		{
 			AddNodeToStrip(GUBP.RootEditorNode.StaticGetFullName(HostPlatform));
 			AgentSharingGroup = "Editor" + StaticGetHostPlatformSuffix(HostPlatform);
@@ -449,12 +451,12 @@ namespace Rocket
 		BranchInfo.BranchUProject Project;
 		bool bIsCodeTargetPlatform;
 
-		public StripRocketMonolithicsNode(GUBP bp, UnrealTargetPlatform InHostPlatform, UnrealTargetPlatform InTargetPlatform, string InStrippedDir) : base(InHostPlatform, InTargetPlatform, InStrippedDir)
+		public StripRocketMonolithicsNode(GUBP bp, GUBP.GUBPBranchConfig InBranchConfig, UnrealTargetPlatform InHostPlatform, UnrealTargetPlatform InTargetPlatform, string InStrippedDir) : base(InBranchConfig, InHostPlatform, InTargetPlatform, InStrippedDir)
 		{
 			Project = bp.Branch.BaseEngineProject;
 			bIsCodeTargetPlatform = RocketBuild.IsCodeTargetPlatform(InHostPlatform, InTargetPlatform);
 
-			GUBP.GUBPNode Node = bp.FindNode(GUBP.GamePlatformMonolithicsNode.StaticGetFullName(HostPlatform, Project, InTargetPlatform, Precompiled: bIsCodeTargetPlatform));
+			GUBP.GUBPNode Node = InBranchConfig.FindNode(GUBP.GamePlatformMonolithicsNode.StaticGetFullName(HostPlatform, Project, InTargetPlatform, Precompiled: bIsCodeTargetPlatform));
 			if(String.IsNullOrEmpty(Node.AgentSharingGroup))
 			{
 				Node.AgentSharingGroup = bp.Branch.BaseEngineProject.GameName + "_MonolithicsGroup_" + InTargetPlatform + StaticGetHostPlatformSuffix(InHostPlatform);
@@ -487,6 +489,7 @@ namespace Rocket
 
 	public class FilterRocketNode : GUBP.HostPlatformNode
 	{
+		GUBP.GUBPBranchConfig BranchConfig;
 		List<UnrealTargetPlatform> SourceHostPlatforms;
 		List<UnrealTargetPlatform> TargetPlatforms;
 		string[] CurrentFeaturePacks;
@@ -494,9 +497,10 @@ namespace Rocket
 		public readonly string DepotManifestPath;
 		public Dictionary<string, string> StrippedNodeManifestPaths = new Dictionary<string, string>();
 
-		public FilterRocketNode(GUBP bp, UnrealTargetPlatform InHostPlatform, List<UnrealTargetPlatform> InTargetPlatforms, string[] InCurrentFeaturePacks, string[] InCurrentTemplates)
+		public FilterRocketNode(GUBP bp, GUBP.GUBPBranchConfig InBranchConfig, UnrealTargetPlatform InHostPlatform, List<UnrealTargetPlatform> InTargetPlatforms, string[] InCurrentFeaturePacks, string[] InCurrentTemplates)
 			: base(InHostPlatform)
 		{
+			BranchConfig = InBranchConfig;
 			TargetPlatforms = new List<UnrealTargetPlatform>(InTargetPlatforms);
 			CurrentFeaturePacks = InCurrentFeaturePacks;
 			CurrentTemplates = InCurrentTemplates;
@@ -584,15 +588,15 @@ namespace Rocket
 			FileFilter Filter = new FileFilter();
 
 			// Include all the editor products
-			AddRuleForBuildProducts(Filter, bp, GUBP.ToolsForCompileNode.StaticGetFullName(HostPlatform), FileFilterType.Include);
-			AddRuleForBuildProducts(Filter, bp, GUBP.RootEditorNode.StaticGetFullName(HostPlatform), FileFilterType.Include);
-			AddRuleForBuildProducts(Filter, bp, GUBP.ToolsNode.StaticGetFullName(HostPlatform), FileFilterType.Include);
+			AddRuleForBuildProducts(Filter, BranchConfig, GUBP.ToolsForCompileNode.StaticGetFullName(HostPlatform), FileFilterType.Include);
+			AddRuleForBuildProducts(Filter, BranchConfig, GUBP.RootEditorNode.StaticGetFullName(HostPlatform), FileFilterType.Include);
+			AddRuleForBuildProducts(Filter, BranchConfig, GUBP.ToolsNode.StaticGetFullName(HostPlatform), FileFilterType.Include);
 
 			// Include win64 tools on Mac, to get the win64 build of UBT, UAT and IPP
 			if (HostPlatform == UnrealTargetPlatform.Mac && bp.HostPlatforms.Contains(UnrealTargetPlatform.Win64))
 			{
-				AddRuleForBuildProducts(Filter, bp, GUBP.ToolsNode.StaticGetFullName(UnrealTargetPlatform.Win64), FileFilterType.Include);
-				AddRuleForBuildProducts(Filter, bp, GUBP.ToolsForCompileNode.StaticGetFullName(UnrealTargetPlatform.Win64), FileFilterType.Include);
+				AddRuleForBuildProducts(Filter, BranchConfig, GUBP.ToolsNode.StaticGetFullName(UnrealTargetPlatform.Win64), FileFilterType.Include);
+				AddRuleForBuildProducts(Filter, BranchConfig, GUBP.ToolsForCompileNode.StaticGetFullName(UnrealTargetPlatform.Win64), FileFilterType.Include);
 			}
 
 			// Include the editor headers
@@ -615,7 +619,7 @@ namespace Rocket
 			{
 				UnrealTargetPlatform SourceHostPlatform = RocketBuild.GetSourceHostPlatform(bp, HostPlatform, TargetPlatform);
 				bool bIsCodeTargetPlatform = RocketBuild.IsCodeTargetPlatform(SourceHostPlatform, TargetPlatform);
-				AddRuleForBuildProducts(Filter, bp, GUBP.GamePlatformMonolithicsNode.StaticGetFullName(SourceHostPlatform, bp.Branch.BaseEngineProject, TargetPlatform, Precompiled: bIsCodeTargetPlatform), FileFilterType.Include);
+				AddRuleForBuildProducts(Filter, BranchConfig, GUBP.GamePlatformMonolithicsNode.StaticGetFullName(SourceHostPlatform, bp.Branch.BaseEngineProject, TargetPlatform, Precompiled: bIsCodeTargetPlatform), FileFilterType.Include);
 			}
 
 			// Include the feature packs
@@ -654,7 +658,7 @@ namespace Rocket
 			{
 				List<string> StrippedFiles = new List<string>();
 
-				StripRocketNode StripNode = (StripRocketNode)bp.FindNode(StrippedNodeManifestPath.Key);
+				StripRocketNode StripNode = (StripRocketNode)BranchConfig.FindNode(StrippedNodeManifestPath.Key);
 				foreach(string BuildProduct in StripNode.BuildProducts)
 				{
 					if(Utils.IsFileUnderDirectory(BuildProduct, StripNode.StrippedDir))
@@ -697,9 +701,9 @@ namespace Rocket
 			}
 		}
 
-		static void AddRuleForBuildProducts(FileFilter Filter, GUBP bp, string NodeName, FileFilterType Type)
+		static void AddRuleForBuildProducts(FileFilter Filter, GUBP.GUBPBranchConfig BranchConfig, string NodeName, FileFilterType Type)
 		{
-			GUBP.GUBPNode Node = bp.FindNode(NodeName);
+			GUBP.GUBPNode Node = BranchConfig.FindNode(NodeName);
 			if(Node == null)
 			{
 				throw new AutomationException("Couldn't find node '{0}'", NodeName);
@@ -722,11 +726,13 @@ namespace Rocket
 
 	public class GatherRocketNode : GUBP.HostPlatformNode
 	{
+		GUBP.GUBPBranchConfig BranchConfig;
 		public readonly string OutputDir;
 		public List<UnrealTargetPlatform> CodeTargetPlatforms;
 
-		public GatherRocketNode(UnrealTargetPlatform HostPlatform, List<UnrealTargetPlatform> InCodeTargetPlatforms, string InOutputDir) : base(HostPlatform)
+		public GatherRocketNode(GUBP.GUBPBranchConfig InBranchConfig, UnrealTargetPlatform HostPlatform, List<UnrealTargetPlatform> InCodeTargetPlatforms, string InOutputDir) : base(HostPlatform)
 		{
+			BranchConfig = InBranchConfig;
 			OutputDir = InOutputDir;
 			CodeTargetPlatforms = new List<UnrealTargetPlatform>(InCodeTargetPlatforms);
 
@@ -762,18 +768,18 @@ namespace Rocket
 			}
 
 			// Copy the depot files to the output directory
-			FilterRocketNode FilterNode = (FilterRocketNode)bp.FindNode(FilterRocketNode.StaticGetFullName(HostPlatform));
+			FilterRocketNode FilterNode = (FilterRocketNode)BranchConfig.FindNode(FilterRocketNode.StaticGetFullName(HostPlatform));
 			CopyManifestFilesToOutput(FilterNode.DepotManifestPath, CommandUtils.CmdEnv.LocalRoot, OutputDir);
 
 			// Copy the stripped files to the output directory
 			foreach(KeyValuePair<string, string> StrippedManifestPath in FilterNode.StrippedNodeManifestPaths)
 			{
-				StripRocketNode StripNode = (StripRocketNode)bp.FindNode(StrippedManifestPath.Key);
+				StripRocketNode StripNode = (StripRocketNode)BranchConfig.FindNode(StrippedManifestPath.Key);
 				CopyManifestFilesToOutput(StrippedManifestPath.Value, StripNode.StrippedDir, OutputDir);
 			}
 
 			// Copy the DDC to the output directory
-			BuildDerivedDataCacheNode DerivedDataCacheNode = (BuildDerivedDataCacheNode)bp.FindNode(BuildDerivedDataCacheNode.StaticGetFullName(HostPlatform));
+			BuildDerivedDataCacheNode DerivedDataCacheNode = (BuildDerivedDataCacheNode)BranchConfig.FindNode(BuildDerivedDataCacheNode.StaticGetFullName(HostPlatform));
 			CopyManifestFilesToOutput(DerivedDataCacheNode.SavedManifestPath, DerivedDataCacheNode.SavedDir, OutputDir);
 
 			// Write the Rocket.txt file with the 
