Remove statics used to cache frequency and email hackers.

diff --git a/Engine/Source/Programs/AutomationTool/GUBP/GUBP.cs b/Engine/Source/Programs/AutomationTool/GUBP/GUBP.cs
--- a/Engine/Source/Programs/AutomationTool/GUBP/GUBP.cs
+++ b/Engine/Source/Programs/AutomationTool/GUBP/GUBP.cs
@@ -487,18 +487,6 @@ public partial class GUBP : BuildCommand
         Log("Wrote graph to {0}", Filename);
      }
 
-
-    // when the host is win64, this is win32 because those are also "host platforms"
-    static public UnrealTargetPlatform GetAltHostPlatform(UnrealTargetPlatform HostPlatform)
-    {
-        UnrealTargetPlatform AltHostPlatform = UnrealTargetPlatform.Unknown; // when the host is win64, this is win32 because those are also "host platforms"
-        if (HostPlatform == UnrealTargetPlatform.Win64)
-        {
-            AltHostPlatform = UnrealTargetPlatform.Win32;
-        }
-        return AltHostPlatform;
-    }
-
     static List<int> ConvertCLToIntList(List<string> Strings)
     {
         List<int> Result = new List<int>();
diff --git a/Engine/Source/Programs/AutomationTool/GUBP/LegacyBranchSetup.cs b/Engine/Source/Programs/AutomationTool/GUBP/LegacyBranchSetup.cs
--- a/Engine/Source/Programs/AutomationTool/GUBP/LegacyBranchSetup.cs
+++ b/Engine/Source/Programs/AutomationTool/GUBP/LegacyBranchSetup.cs
@@ -75,8 +75,6 @@ partial class GUBP
         public abstract void AddNodes(GUBP bp, GUBPBranchConfig BranchConfig, UnrealTargetPlatform InHostPlatform, List<UnrealTargetPlatform> InActivePlatforms);
     }
 
-    private static List<GUBPNodeAdder> Adders;
-
 	Type[] GetTypesFromAssembly(Assembly Asm)
 	{
 		Type[] AllTypesWeCanGet;
@@ -91,31 +89,33 @@ partial class GUBP
 		return AllTypesWeCanGet;
 	}
 
-    private void AddCustomNodes(GUBPBranchConfig BranchConfig, UnrealTargetPlatform InHostPlatform, List<UnrealTargetPlatform> InActivePlatforms)
+    private void AddCustomNodes(GUBPBranchConfig BranchConfig, List<UnrealTargetPlatform> ActivePlatforms)
     {
-        if (Adders == null)
-        {
-            Adders = new List<GUBPNodeAdder>();
-            Assembly[] LoadedAssemblies = AppDomain.CurrentDomain.GetAssemblies();
-            foreach (Assembly Dll in LoadedAssemblies)
-            {
-				Type[] AllTypes = GetTypesFromAssembly(Dll);
-                foreach (Type PotentialConfigType in AllTypes)
-                {
-                    if (PotentialConfigType != typeof(GUBPNodeAdder) && typeof(GUBPNodeAdder).IsAssignableFrom(PotentialConfigType))
-                    {
-                        GUBPNodeAdder Config = Activator.CreateInstance(PotentialConfigType) as GUBPNodeAdder;
-                        if (Config != null)
-                        {
-                            Adders.Add(Config);
-                        }
-                    }
-                }
-            }
-        }
-        foreach(GUBPNodeAdder Adder in Adders)
+		List<GUBPNodeAdder> Adders = new List<GUBPNodeAdder>();
+
+		Assembly[] LoadedAssemblies = AppDomain.CurrentDomain.GetAssemblies();
+		foreach (Assembly Dll in LoadedAssemblies)
+		{
+			Type[] AllTypes = GetTypesFromAssembly(Dll);
+			foreach (Type PotentialConfigType in AllTypes)
+			{
+				if (PotentialConfigType != typeof(GUBPNodeAdder) && typeof(GUBPNodeAdder).IsAssignableFrom(PotentialConfigType))
+				{
+					GUBPNodeAdder Config = Activator.CreateInstance(PotentialConfigType) as GUBPNodeAdder;
+					if (Config != null)
+					{
+						Adders.Add(Config);
+					}
+				}
+			}
+		}
+
+        foreach (UnrealTargetPlatform HostPlatform in HostPlatforms)
         {
-            Adder.AddNodes(this, BranchConfig, InHostPlatform, InActivePlatforms);
+			foreach(GUBPNodeAdder Adder in Adders)
+			{
+				Adder.AddNodes(this, BranchConfig, HostPlatform, ActivePlatforms);
+			}
         }
     }
 
@@ -273,36 +273,55 @@ partial class GUBP
             return new int();
         }
     }
-    private static List<GUBPFrequencyHacker> FrequencyHackers;
-    private int HackFrequency(GUBP bp, string Branch, NodeInfo NodeInfo, int BaseFrequency)
-    {
-        int Frequency = BaseFrequency;
-        if (FrequencyHackers == null)
+
+	void FindFrequenciesForNodes(IEnumerable<NodeInfo> Nodes, Dictionary<string, int> FrequencyOverrides)
+	{
+        List<GUBPFrequencyHacker> FrequencyHackers = new List<GUBPFrequencyHacker>();
+
+        Assembly[] LoadedAssemblies = AppDomain.CurrentDomain.GetAssemblies();
+        foreach (var Dll in LoadedAssemblies)
         {
-            FrequencyHackers = new List<GUBPFrequencyHacker>();
-            Assembly[] LoadedAssemblies = AppDomain.CurrentDomain.GetAssemblies();
-            foreach (var Dll in LoadedAssemblies)
+			Type[] AllTypes = GetTypesFromAssembly(Dll);
+            foreach (var PotentialConfigType in AllTypes)
             {
-				Type[] AllTypes = GetTypesFromAssembly(Dll);
-                foreach (var PotentialConfigType in AllTypes)
+                if (PotentialConfigType != typeof(GUBPFrequencyHacker) && typeof(GUBPFrequencyHacker).IsAssignableFrom(PotentialConfigType))
                 {
-                    if (PotentialConfigType != typeof(GUBPFrequencyHacker) && typeof(GUBPFrequencyHacker).IsAssignableFrom(PotentialConfigType))
+                    GUBPFrequencyHacker Config = Activator.CreateInstance(PotentialConfigType) as GUBPFrequencyHacker;
+                    if (Config != null)
                     {
-                        GUBPFrequencyHacker Config = Activator.CreateInstance(PotentialConfigType) as GUBPFrequencyHacker;
-                        if (Config != null)
-                        {
-                            FrequencyHackers.Add(Config);
-                        }
+                        FrequencyHackers.Add(Config);
                     }
                 }
             }
         }
-        foreach(var FrequencyHacker in FrequencyHackers)
+
+		foreach(NodeInfo Node in Nodes)
+		{
+			Node.FrequencyShift = Node.Node.CISFrequencyQuantumShift(this);
+
+			foreach(GUBPFrequencyHacker FrequencyHacker in FrequencyHackers)
+			{
+				Node.FrequencyShift = FrequencyHacker.GetNodeFrequency(this, BranchName, Node.Name, Node.FrequencyShift);            
+			}
+
+			int FrequencyOverride;
+			if (FrequencyOverrides.TryGetValue(Node.Name, out FrequencyOverride) && Node.FrequencyShift > FrequencyOverride)
+			{
+				Node.FrequencyShift = FrequencyOverride;
+			}
+		}
+	}
+
+    // when the host is win64, this is win32 because those are also "host platforms"
+    static public UnrealTargetPlatform GetAltHostPlatform(UnrealTargetPlatform HostPlatform)
+    {
+        UnrealTargetPlatform AltHostPlatform = UnrealTargetPlatform.Unknown; // when the host is win64, this is win32 because those are also "host platforms"
+        if (HostPlatform == UnrealTargetPlatform.Win64)
         {
-           Frequency = FrequencyHacker.GetNodeFrequency(bp, Branch, NodeInfo.Name, BaseFrequency);            
+            AltHostPlatform = UnrealTargetPlatform.Win32;
         }
-        return Frequency;
-    }    
+        return AltHostPlatform;
+    }
 
 	void AddNodesForBranch(int TimeIndex, bool bNoAutomatedTesting, out List<NodeInfo> AllNodes, out List<AggregateInfo> AllAggregates)
 	{
@@ -1085,10 +1104,8 @@ partial class GUBP
             BranchConfig.AddNode(new WaitForSharedPromotionUserInput(this, true));
             BranchConfig.AddNode(new SharedLabelPromotableNode(this, true));			
         }
-        foreach (var HostPlatform in HostPlatforms)
-        {
-            AddCustomNodes(BranchConfig, HostPlatform, ActivePlatforms);
-        }
+
+		AddCustomNodes(BranchConfig, ActivePlatforms);
         
         if (BranchConfig.HasNode(ToolsForCompileNode.StaticGetFullName(UnrealTargetPlatform.Win64)))
         {
@@ -1143,17 +1160,7 @@ partial class GUBP
 		AllAggregates.AddRange(BranchConfig.GUBPAggregates.Values.Select(x => new AggregateInfo{ Name = x.GetFullName(), Node = x }));
 
 		// Calculate the frequencies for each node
-		foreach(NodeInfo Node in AllNodes)
-		{
-			Node.FrequencyShift = Node.Node.CISFrequencyQuantumShift(this);
-            Node.FrequencyShift = HackFrequency(this, BranchName, Node, Node.FrequencyShift);
-
-			int FrequencyOverride;
-			if (FrequencyOverrides.TryGetValue(Node.Name, out FrequencyOverride) && Node.FrequencyShift > FrequencyOverride)
-			{
-				Node.FrequencyShift = FrequencyOverride;
-			}
-		}
+		FindFrequenciesForNodes(AllNodes, FrequencyOverrides);
 
 		// Get the email list for each node
 		FindEmailsForNodes(AllNodes);
