Move all GUBP legacy state variables into GUBPBranchConfig.

diff --git a/Engine/Source/Programs/AutomationTool/GUBP/BuildNode.cs b/Engine/Source/Programs/AutomationTool/GUBP/BuildNode.cs
--- a/Engine/Source/Programs/AutomationTool/GUBP/BuildNode.cs
+++ b/Engine/Source/Programs/AutomationTool/GUBP/BuildNode.cs
@@ -20,6 +20,7 @@ namespace AutomationTool
 		public bool IsComplete;
 		public string[] RecipientsForFailureEmails;
 		public bool AddSubmittersToFailureEmails;
+		public int AgentMemoryRequirement;
 
 		public string ControllingTriggerDotName
 		{
diff --git a/Engine/Source/Programs/AutomationTool/GUBP/GUBP.cs b/Engine/Source/Programs/AutomationTool/GUBP/GUBP.cs
--- a/Engine/Source/Programs/AutomationTool/GUBP/GUBP.cs
+++ b/Engine/Source/Programs/AutomationTool/GUBP/GUBP.cs
@@ -17,16 +17,6 @@ public partial class GUBP : BuildCommand
 	const string FailedTempStorageSuffix = "_Failed";
 	const string SucceededTempStorageSuffix = "_Succeeded";
 
-	public string BranchName;
-    public int CL = 0;
-    public bool bSignBuildProducts = false;
-    public BranchInfo Branch = null;
-    public List<UnrealTargetPlatform> HostPlatforms;
-    public bool bForceIncrementalCompile = false;
-    public bool bPreflightBuild = false;
-    public string PreflightMangleSuffix = "";
-    public GUBPBranchHacker.BranchOptions BranchOptions = null;	
-
 	[DebuggerDisplay("{Name}")]
 	class AggregateInfo
 	{
@@ -385,13 +375,10 @@ public partial class GUBP : BuildCommand
             {
 				Builder.AppendFormat(" [{0}]", Agent);
             }
-
-			int MemoryRequirement = NodeToDo.Node.AgentMemoryRequirement(bp);
-			if(MemoryRequirement != 0)
+			if(NodeToDo.AgentMemoryRequirement != 0)
 			{
-				Builder.AppendFormat(" [{0}gb]", MemoryRequirement);
+				Builder.AppendFormat(" [{0}gb]", NodeToDo.AgentMemoryRequirement);
 			}
-
             if (AddEmailProps)
             {
 				Builder.AppendFormat(" {0}", String.Join(" ", NodeToDo.RecipientsForFailureEmails));
@@ -1024,7 +1011,7 @@ public partial class GUBP : BuildCommand
 				AgentReq = ParseParamValue("AgentOverride");
 			}
 			ECProps.Add(string.Format("AgentRequirementString/{0}={1}", NodeToDo.Name, AgentReq));
-			ECProps.Add(string.Format("RequiredMemory/{0}={1}", NodeToDo.Name, NodeToDo.Node.AgentMemoryRequirement(this)));
+			ECProps.Add(string.Format("RequiredMemory/{0}={1}", NodeToDo.Name, NodeToDo.AgentMemoryRequirement));
 			ECProps.Add(string.Format("Timeouts/{0}={1}", NodeToDo.Name, NodeToDo.Node.TimeoutInMinutes()));
 			ECProps.Add(string.Format("JobStepPath/{0}={1}", NodeToDo.Name, GetJobStepPath(NodeToDo)));
 		}
@@ -1105,6 +1092,7 @@ public partial class GUBP : BuildCommand
     {
         Log("************************* GUBP");
 
+		bool bPreflightBuild = false;
 		int PreflightShelveCL = 0;
         string PreflightShelveCLString = GetEnvVar("uebp_PreflightShelveCL");
         if ((!String.IsNullOrEmpty(PreflightShelveCLString) && IsBuildMachine) || ParseParam("PreflightTest"))
@@ -1121,63 +1109,35 @@ public partial class GUBP : BuildCommand
             bPreflightBuild = true;
         }
         
-        HostPlatforms = new List<UnrealTargetPlatform>();
+		List<UnrealTargetPlatform> HostPlatforms = new List<UnrealTargetPlatform>();
         if (!ParseParam("NoPC"))
         {
             HostPlatforms.Add(UnrealTargetPlatform.Win64);
         }
-		if (P4Enabled)
-		{
-			BranchName = P4Env.BuildRootP4;
-		}
-		else
-		{ 
-			BranchName = ParseParamValue("BranchName", "");
-		}
-		BranchOptions = GetBranchOptions(BranchName);
-        bool WithMac = !BranchOptions.PlatformsToRemove.Contains(UnrealTargetPlatform.Mac);
-        if (ParseParam("NoMac"))
+        if (!ParseParam("NoMac"))
         {
-            WithMac = false;
-        }
-        if (WithMac)
-		{
 			HostPlatforms.Add(UnrealTargetPlatform.Mac);
-		}
-
-		bool WithLinux = !BranchOptions.PlatformsToRemove.Contains(UnrealTargetPlatform.Linux);
-		bool WithoutLinux = ParseParam("NoLinux");
-		// @TODO: exclude temporarily unless running on a Linux machine to prevent spurious GUBP failures
-		if (UnrealBuildTool.BuildHostPlatform.Current.Platform != UnrealTargetPlatform.Linux || ParseParam("NoLinux"))
-		{
-			WithLinux = false;
-		}
-		if (WithLinux)
+        }
+		if(!ParseParam("NoLinux") && UnrealBuildTool.BuildHostPlatform.Current.Platform == UnrealTargetPlatform.Linux)
 		{
 			HostPlatforms.Add(UnrealTargetPlatform.Linux);
 		}
 
-        bForceIncrementalCompile = ParseParam("ForceIncrementalCompile");
-        bool bNoAutomatedTesting = ParseParam("NoAutomatedTesting") || BranchOptions.bNoAutomatedTesting;		
         string StoreName = ParseParamValue("Store");
         string StoreSuffix = ParseParamValue("StoreSuffix", "");
 
+		string PreflightMangleSuffix = "";
         if (bPreflightBuild)
         {
             int PreflightUID = ParseParamInt("PreflightUID", 0);
             PreflightMangleSuffix = String.Format("-PF-{0}-{1}", PreflightShelveCL, PreflightUID);
             StoreSuffix = StoreSuffix + PreflightMangleSuffix;
         }
-        CL = ParseParamInt("CL", 0);
+        int CL = ParseParamInt("CL", 0);
         bool bCleanLocalTempStorage = ParseParam("CleanLocal");
         bool bSkipTriggers = ParseParam("SkipTriggers");
         bool bFake = ParseParam("fake");
         bool bFakeEC = ParseParam("FakeEC");
-        int TimeIndex = ParseParamInt("TimeIndex", 0);
-        if (TimeIndex == 0)
-        {
-            TimeIndex = ParseParamInt("UserTimeIndex", 0);
-        }
 
         bool bSaveSharedTempStorage = false;
 
@@ -1213,7 +1173,6 @@ public partial class GUBP : BuildCommand
             {
                 throw new AutomationException("Request to save to temp storage, but {0} is unavailable.", TempStorage.UE4TempStorageDirectory());
             }
-            bSignBuildProducts = true;
         }
         else if (!LocalOnly && !TempStorage.HaveSharedTempStorage(false))
         {
@@ -1229,17 +1188,10 @@ public partial class GUBP : BuildCommand
 			ExplicitTriggerName = ParseParamValue("TriggerNode", "");
 		}
 
-        if (ParseParam("CIS") && ExplicitTriggerName == "" && CommanderSetup) // explicit triggers will already have a time index assigned
-        {
-			TimeIndex = UpdateCISCounter();
-            Log("Setting TimeIndex to {0}", TimeIndex);
-        }
-
-        Branch = new BranchInfo(HostPlatforms);        
-
 		List<NodeInfo> AllNodes;
 		List<AggregateInfo> AllAggregates;
-		AddNodesForBranch(TimeIndex, bNoAutomatedTesting, out AllNodes, out AllAggregates);
+		int TimeQuantum = 20;
+		AddNodesForBranch(CL, HostPlatforms, bPreflightBuild, PreflightMangleSuffix, out AllNodes, out AllAggregates, ref TimeQuantum);
 
 		LinkGraph(AllAggregates, AllNodes);
 		FindControllingTriggers(AllNodes);
@@ -1251,7 +1203,19 @@ public partial class GUBP : BuildCommand
             TempStorage.DeleteLocalTempStorageManifests(CmdEnv);
         }
 
-		HashSet<NodeInfo> NodesToDo = ParseNodesToDo(AllNodes, AllAggregates, WithoutLinux, TimeIndex, CommanderSetup);
+		int TimeIndex = ParseParamInt("TimeIndex", 0);
+		if (TimeIndex == 0)
+		{
+			TimeIndex = ParseParamInt("UserTimeIndex", 0);
+		}
+		if (ParseParam("CIS") && ExplicitTriggerName == "" && CommanderSetup) // explicit triggers will already have a time index assigned
+		{
+			TimeIndex = UpdateCISCounter(TimeQuantum);
+			Log("Setting TimeIndex to {0}", TimeIndex);
+		}
+
+		bool WithoutLinux = ParseParam("NoLinux");
+		HashSet<NodeInfo> NodesToDo = ParseNodesToDo(AllNodes, AllAggregates, WithoutLinux, TimeIndex, bPreflightBuild, CommanderSetup);
 
 		NodeInfo ExplicitTrigger = null;
 		if (CommanderSetup)
@@ -1292,12 +1256,6 @@ public partial class GUBP : BuildCommand
 
         List<NodeInfo> OrderedToDo = TopologicalSort(NodesToDo, ExplicitTrigger, false, false);
 
-		int TimeQuantum = 20;
-		if(BranchOptions.QuantumOverride != 0)
-		{
-			TimeQuantum = BranchOptions.QuantumOverride;
-		}
-
 		List<NodeInfo> UnfinishedTriggers = FindUnfinishedTriggers(bSkipTriggers, ExplicitTrigger, OrderedToDo);
 
 		PrintNodes(this, OrderedToDo, AllAggregates, UnfinishedTriggers, TimeQuantum);
@@ -1308,7 +1266,7 @@ public partial class GUBP : BuildCommand
         string FakeFail = ParseParamValue("FakeFail");
         if(CommanderSetup)
         {
-			DoCommanderSetup(AllNodes, AllAggregates, NodesToDo, OrderedToDo, TimeIndex, TimeQuantum, bSkipTriggers, bFake, bFakeEC, CLString, ExplicitTrigger, UnfinishedTriggers, FakeFail);
+			DoCommanderSetup(AllNodes, AllAggregates, NodesToDo, OrderedToDo, TimeIndex, TimeQuantum, bSkipTriggers, bFake, bFakeEC, CLString, ExplicitTrigger, UnfinishedTriggers, FakeFail, bPreflightBuild);
         }
 		else if(ParseParam("SaveGraph"))
 		{
@@ -1503,7 +1461,7 @@ public partial class GUBP : BuildCommand
 		return false;
 	}
 
-	private HashSet<NodeInfo> ParseNodesToDo(IEnumerable<NodeInfo> PotentialNodes, IEnumerable<AggregateInfo> PotentialAggregates, bool WithoutLinux, int TimeIndex, bool CommanderSetup)
+	private HashSet<NodeInfo> ParseNodesToDo(IEnumerable<NodeInfo> PotentialNodes, IEnumerable<AggregateInfo> PotentialAggregates, bool WithoutLinux, int TimeIndex, bool bPreflightBuild, bool CommanderSetup)
 	{
 		List<string> NamesToDo = new List<string>();
 
@@ -1674,7 +1632,7 @@ public partial class GUBP : BuildCommand
 		return FullNodeDependedOnBy;
 	}
 
-	private int UpdateCISCounter()
+	private int UpdateCISCounter(int TimeQuantum)
 	{
 		if (!P4Enabled)
 		{
@@ -1739,11 +1697,6 @@ public partial class GUBP : BuildCommand
 
 			int DeltaMinutes = NowMinutes - Minutes;
 
-			int TimeQuantum = 20;
-			if (BranchOptions.QuantumOverride != 0)
-			{
-				TimeQuantum = BranchOptions.QuantumOverride;
-			}
 			int NewIndex = Index + 1;
 
 			if (DeltaMinutes > TimeQuantum * 2)
@@ -1840,7 +1793,7 @@ public partial class GUBP : BuildCommand
 		}
 	}
 
-	private void DoCommanderSetup(IEnumerable<NodeInfo> AllNodes, IEnumerable<AggregateInfo> AllAggregates, HashSet<NodeInfo> NodesToDo, List<NodeInfo> OrdereredToDo, int TimeIndex, int TimeQuantum, bool bSkipTriggers, bool bFake, bool bFakeEC, string CLString, NodeInfo ExplicitTrigger, List<NodeInfo> UnfinishedTriggers, string FakeFail)
+	private void DoCommanderSetup(IEnumerable<NodeInfo> AllNodes, IEnumerable<AggregateInfo> AllAggregates, HashSet<NodeInfo> NodesToDo, List<NodeInfo> OrdereredToDo, int TimeIndex, int TimeQuantum, bool bSkipTriggers, bool bFake, bool bFakeEC, string CLString, NodeInfo ExplicitTrigger, List<NodeInfo> UnfinishedTriggers, string FakeFail, bool bPreflightBuild)
 	{
 		Dictionary<string, string> FullNodeDependedOnBy = GetFullNodeDependedOnBy(NodesToDo);
 
diff --git a/Engine/Source/Programs/AutomationTool/GUBP/LegacyBranchSetup.cs b/Engine/Source/Programs/AutomationTool/GUBP/LegacyBranchSetup.cs
--- a/Engine/Source/Programs/AutomationTool/GUBP/LegacyBranchSetup.cs
+++ b/Engine/Source/Programs/AutomationTool/GUBP/LegacyBranchSetup.cs
@@ -13,8 +13,29 @@ partial class GUBP
 {
 	public class GUBPBranchConfig
 	{
-		public Dictionary<string, GUBPNode> GUBPNodes = new Dictionary<string,GUBPNode>();
-		public Dictionary<string, GUBPAggregateNode> GUBPAggregates = new Dictionary<string,GUBPAggregateNode>();
+		public int CL = 0;
+		public List<UnrealTargetPlatform> HostPlatforms;
+		public string BranchName;
+		public BranchInfo Branch;
+		public GUBPBranchHacker.BranchOptions BranchOptions;
+		public bool bForceIncrementalCompile;
+		public bool bPreflightBuild;
+		public string PreflightMangleSuffix;
+
+		public Dictionary<string, GUBPNode> GUBPNodes = new Dictionary<string, GUBPNode>();
+		public Dictionary<string, GUBPAggregateNode> GUBPAggregates = new Dictionary<string, GUBPAggregateNode>();
+
+		public GUBPBranchConfig(int InCL, IEnumerable<UnrealTargetPlatform> InHostPlatforms, string InBranchName, BranchInfo InBranch, GUBPBranchHacker.BranchOptions InBranchOptions, bool bInForceIncrementalCompile, bool bInPreflightBuild, string InPreflightMangleSuffix)
+		{
+			CL = InCL;
+			HostPlatforms = new List<UnrealTargetPlatform>(InHostPlatforms);
+			BranchName = InBranchName;
+			Branch = InBranch;
+			BranchOptions = InBranchOptions;
+			bForceIncrementalCompile = bInForceIncrementalCompile;
+			bPreflightBuild = bInPreflightBuild;
+			PreflightMangleSuffix = InPreflightMangleSuffix;
+		}
 
 		public string AddNode(GUBPNode Node)
 		{
@@ -89,7 +110,7 @@ partial class GUBP
 		return AllTypesWeCanGet;
 	}
 
-    private void AddCustomNodes(GUBPBranchConfig BranchConfig, List<UnrealTargetPlatform> ActivePlatforms)
+    private void AddCustomNodes(GUBPBranchConfig BranchConfig, List<UnrealTargetPlatform> HostPlatforms, List<UnrealTargetPlatform> ActivePlatforms)
     {
 		List<GUBPNodeAdder> Adders = new List<GUBPNodeAdder>();
 
@@ -192,7 +213,7 @@ partial class GUBP
 		}
     }
 
-	private void FindEmailsForNodes(IEnumerable<NodeInfo> Nodes)
+	private void FindEmailsForNodes(string BranchName, IEnumerable<NodeInfo> Nodes)
 	{
 		List<GUBPEmailHacker> EmailHackers = new List<GUBPEmailHacker>();
 
@@ -274,7 +295,7 @@ partial class GUBP
         }
     }
 
-	void FindFrequenciesForNodes(IEnumerable<NodeInfo> Nodes, Dictionary<string, int> FrequencyOverrides)
+	void FindFrequenciesForNodes(GUBP.GUBPBranchConfig BranchConfig, IEnumerable<NodeInfo> Nodes, Dictionary<string, int> FrequencyOverrides)
 	{
         List<GUBPFrequencyHacker> FrequencyHackers = new List<GUBPFrequencyHacker>();
 
@@ -297,11 +318,11 @@ partial class GUBP
 
 		foreach(NodeInfo Node in Nodes)
 		{
-			Node.FrequencyShift = Node.Node.CISFrequencyQuantumShift(this);
+			Node.FrequencyShift = Node.Node.CISFrequencyQuantumShift(BranchConfig);
 
 			foreach(GUBPFrequencyHacker FrequencyHacker in FrequencyHackers)
 			{
-				Node.FrequencyShift = FrequencyHacker.GetNodeFrequency(this, BranchName, Node.Name, Node.FrequencyShift);            
+				Node.FrequencyShift = FrequencyHacker.GetNodeFrequency(this, BranchConfig.BranchName, Node.Name, Node.FrequencyShift);            
 			}
 
 			int FrequencyOverride;
@@ -323,8 +344,36 @@ partial class GUBP
         return AltHostPlatform;
     }
 
-	void AddNodesForBranch(int TimeIndex, bool bNoAutomatedTesting, out List<NodeInfo> AllNodes, out List<AggregateInfo> AllAggregates)
+	void AddNodesForBranch(int CL, List<UnrealTargetPlatform> InitialHostPlatforms, bool bPreflightBuild, string PreflightMangleSuffix, out List<NodeInfo> AllNodes, out List<AggregateInfo> AllAggregates, ref int TimeQuantum)
 	{
+		string BranchName;
+		if (P4Enabled)
+		{
+			BranchName = P4Env.BuildRootP4;
+		}
+		else
+		{ 
+			BranchName = ParseParamValue("BranchName", "");
+		}
+
+		GUBPBranchHacker.BranchOptions BranchOptions = GetBranchOptions(BranchName);
+
+		if (BranchOptions.QuantumOverride != 0)
+		{
+			TimeQuantum = BranchOptions.QuantumOverride;
+		}
+
+		bool bForceIncrementalCompile = ParseParam("ForceIncrementalCompile");
+		bool bNoAutomatedTesting = ParseParam("NoAutomatedTesting") || BranchOptions.bNoAutomatedTesting;		
+
+		List<UnrealTargetPlatform> HostPlatforms = new List<UnrealTargetPlatform>(InitialHostPlatforms);
+		foreach(UnrealTargetPlatform PlatformToRemove in BranchOptions.PlatformsToRemove)
+		{
+			HostPlatforms.Remove(PlatformToRemove);
+		}
+
+        BranchInfo Branch = new BranchInfo(HostPlatforms);        
+
 		List<UnrealTargetPlatform> ActivePlatforms;
         if (IsBuildMachine || ParseParam("AllPlatforms"))
         {
@@ -403,25 +452,25 @@ partial class GUBP
             }
         }
 
-		GUBPBranchConfig BranchConfig = new GUBPBranchConfig();
+		GUBPBranchConfig BranchConfig = new GUBPBranchConfig(CL, HostPlatforms, BranchName, Branch, BranchOptions, bForceIncrementalCompile, bPreflightBuild, PreflightMangleSuffix);
 
         BranchConfig.AddNode(new VersionFilesNode());
 
 
         foreach (var HostPlatform in HostPlatforms)
         {
-			BranchConfig.AddNode(new ToolsForCompileNode(HostPlatform, TimeIndex));
+			BranchConfig.AddNode(new ToolsForCompileNode(BranchConfig, HostPlatform, ParseParam("Launcher")));
 			
 			if (!BranchOptions.ExcludePlatformsForEditor.Contains(HostPlatform))
 			{
-				BranchConfig.AddNode(new RootEditorNode(HostPlatform));			
-				BranchConfig.AddNode(new ToolsNode(this, HostPlatform));            
-				BranchConfig.AddNode(new InternalToolsNode(this, HostPlatform));
-			if (HostPlatform == UnrealTargetPlatform.Win64 && ActivePlatforms.Contains(UnrealTargetPlatform.Linux))
+				BranchConfig.AddNode(new RootEditorNode(BranchConfig, HostPlatform));			
+				BranchConfig.AddNode(new ToolsNode(BranchConfig, HostPlatform));            
+				BranchConfig.AddNode(new InternalToolsNode(BranchConfig, HostPlatform));
+				if (HostPlatform == UnrealTargetPlatform.Win64 && ActivePlatforms.Contains(UnrealTargetPlatform.Linux))
 				{
 					if (!BranchOptions.ExcludePlatformsForEditor.Contains(UnrealTargetPlatform.Linux))
 					{
-						BranchConfig.AddNode(new ToolsCrossCompileNode(HostPlatform));
+						BranchConfig.AddNode(new ToolsCrossCompileNode(BranchConfig, HostPlatform));
 					}
 				}
 				foreach (var ProgramTarget in Branch.BaseEngineProject.Properties.Programs)
@@ -436,11 +485,11 @@ partial class GUBP
                         {
                             if (bInternalOnly)
                             {
-                                BranchConfig.AddNode(new SingleInternalToolsNode(HostPlatform, ProgramTarget));
+                                BranchConfig.AddNode(new SingleInternalToolsNode(BranchConfig, HostPlatform, ProgramTarget));
                             }
                             else
                             {
-                                BranchConfig.AddNode(new SingleToolsNode(HostPlatform, ProgramTarget));
+                                BranchConfig.AddNode(new SingleToolsNode(BranchConfig, HostPlatform, ProgramTarget));
                             }
                         }
                         if (ProgramTarget.Rules.GUBP_IncludeNonUnityToolTest())
@@ -463,11 +512,11 @@ partial class GUBP
                             {
                                 if (bInternalNodeOnly)
                                 {
-                                    BranchConfig.AddNode(new SingleInternalToolsNode(HostPlatform, ProgramTarget));
+                                    BranchConfig.AddNode(new SingleInternalToolsNode(BranchConfig, HostPlatform, ProgramTarget));
                                 }
                                 else
                                 {
-                                    BranchConfig.AddNode(new SingleToolsNode(HostPlatform, ProgramTarget));
+                                    BranchConfig.AddNode(new SingleToolsNode(BranchConfig, HostPlatform, ProgramTarget));
                                 }
                             }
                             if (ProgramTarget.Rules.GUBP_IncludeNonUnityToolTest())
@@ -552,7 +601,7 @@ partial class GUBP
 
 			if (!BranchOptions.ExcludePlatformsForEditor.Contains(HostPlatform))
 			{
-			BranchConfig.AddNode(new NonUnityTestNode(HostPlatform));
+				BranchConfig.AddNode(new NonUnityTestNode(BranchConfig, HostPlatform));
 			}
 
             if (DoASharedPromotable)
@@ -576,7 +625,7 @@ partial class GUBP
                     {
                         if (!bNoAutomatedTesting)
                         {							
-                            EditorTestNodes.Add(BranchConfig.AddNode(new UATTestNode(this, BranchConfig, HostPlatform, Branch.BaseEngineProject, Test.Key, Test.Value, AgentSharingGroup)));
+                            EditorTestNodes.Add(BranchConfig.AddNode(new UATTestNode(BranchConfig, HostPlatform, Branch.BaseEngineProject, Test.Key, Test.Value, AgentSharingGroup)));
                         
                             foreach (var NonCodeProject in Branch.NonCodeProjects)
                             {
@@ -585,13 +634,13 @@ partial class GUBP
                                     continue;
                                 }
                                 if (HostPlatform == UnrealTargetPlatform.Mac) continue; //temp hack till mac automated testing works
-                                EditorTestNodes.Add(BranchConfig.AddNode(new UATTestNode(this, BranchConfig, HostPlatform, NonCodeProject, Test.Key, Test.Value, AgentSharingGroup)));
+                                EditorTestNodes.Add(BranchConfig.AddNode(new UATTestNode(BranchConfig, HostPlatform, NonCodeProject, Test.Key, Test.Value, AgentSharingGroup)));
                             }
                         }
                     }
                     if (EditorTestNodes.Count > 0)
                     {
-                        BranchConfig.AddNode(new GameAggregateNode(this, HostPlatform, Branch.BaseEngineProject, "AllEditorTests", EditorTestNodes));
+                        BranchConfig.AddNode(new GameAggregateNode(HostPlatform, Branch.BaseEngineProject, "AllEditorTests", EditorTestNodes));
                     }
                 }
 
@@ -636,15 +685,15 @@ partial class GUBP
                                 {
 									if(GamePlatformMonolithicsNode.HasPrecompiledTargets(Branch.BaseEngineProject, HostPlatform, Plat))
 									{
-	                                    BranchConfig.AddNode(new GamePlatformMonolithicsNode(this, BranchConfig, HostPlatform, ActivePlatforms, Branch.BaseEngineProject, Plat, InPrecompiled: true));
+	                                    BranchConfig.AddNode(new GamePlatformMonolithicsNode(BranchConfig, HostPlatform, ActivePlatforms, Branch.BaseEngineProject, Plat, InPrecompiled: true));
 									}
-                                    BranchConfig.AddNode(new GamePlatformMonolithicsNode(this, BranchConfig, HostPlatform, ActivePlatforms, Branch.BaseEngineProject, Plat));
+                                    BranchConfig.AddNode(new GamePlatformMonolithicsNode(BranchConfig, HostPlatform, ActivePlatforms, Branch.BaseEngineProject, Plat));
                                 }
 								if (Plat == UnrealTargetPlatform.Win32 && Target.Rules.GUBP_BuildWindowsXPMonolithics() && Kind == TargetRules.TargetType.Game)
 								{
 									if (!BranchConfig.HasNode(GamePlatformMonolithicsNode.StaticGetFullName(HostPlatform, Branch.BaseEngineProject, Plat, true)))
 									{
-										BranchConfig.AddNode(new GamePlatformMonolithicsNode(this, BranchConfig, HostPlatform, ActivePlatforms, Branch.BaseEngineProject, Plat, true));
+										BranchConfig.AddNode(new GamePlatformMonolithicsNode(BranchConfig, HostPlatform, ActivePlatforms, Branch.BaseEngineProject, Plat, true));
 									}
 	                            }
                             }
@@ -682,11 +731,11 @@ partial class GUBP
                                 string CookedPlatform = Platform.Platforms[Plat].GetCookPlatform(Kind == TargetRules.TargetType.Server, Kind == TargetRules.TargetType.Client, "");								
                                 if (!BranchConfig.HasNode(CookNode.StaticGetFullName(HostPlatform, Branch.BaseEngineProject, CookedPlatform)))
                                 {
-                                    GameCookNodes.Add(BranchConfig.AddNode(new CookNode(this, BranchConfig, HostPlatform, Branch.BaseEngineProject, Plat, CookedPlatform)));
+									GameCookNodes.Add(BranchConfig.AddNode(new CookNode(BranchConfig, HostPlatform, Branch.BaseEngineProject, Plat, CookedPlatform)));
                                 }
                                 if (!BranchConfig.HasNode(GamePlatformCookedAndCompiledNode.StaticGetFullName(HostPlatform, Branch.BaseEngineProject, Plat)))
                                 {
-                                    BranchConfig.AddNode(new GamePlatformCookedAndCompiledNode(this, HostPlatform, Branch.BaseEngineProject, Plat, false));
+									BranchConfig.AddNode(new GamePlatformCookedAndCompiledNode(BranchConfig, HostPlatform, Branch.BaseEngineProject, Plat, false));
                                 }
                                 var GameTests = Target.Rules.GUBP_GetGameTests_MonolithicOnly(HostPlatform, GetAltHostPlatform(HostPlatform), Plat);
                                 var RequiredPlatforms = new List<UnrealTargetPlatform> { Plat };
@@ -697,11 +746,11 @@ partial class GUBP
                                     foreach (var Test in GameTests)
                                     {
                                         var TestName = Test.Key + "_" + Plat.ToString();
-                                        ThisMonoGameTestNodes.Add(BranchConfig.AddNode(new UATTestNode(this, BranchConfig, HostPlatform, Branch.BaseEngineProject, TestName, Test.Value, CookedAgentSharingGroup, false, RequiredPlatforms)));
+                                        ThisMonoGameTestNodes.Add(BranchConfig.AddNode(new UATTestNode(BranchConfig, HostPlatform, Branch.BaseEngineProject, TestName, Test.Value, CookedAgentSharingGroup, false, RequiredPlatforms)));
                                     }
                                     if (ThisMonoGameTestNodes.Count > 0)
                                     {										
-                                        GameTestNodes.Add(BranchConfig.AddNode(new GameAggregateNode(this, HostPlatform, Branch.BaseEngineProject, "CookedTests_" + Plat.ToString() + "_" + Kind.ToString() + HostPlatformNode.StaticGetHostPlatformSuffix(HostPlatform), ThisMonoGameTestNodes)));
+                                        GameTestNodes.Add(BranchConfig.AddNode(new GameAggregateNode(HostPlatform, Branch.BaseEngineProject, "CookedTests_" + Plat.ToString() + "_" + Kind.ToString() + HostPlatformNode.StaticGetHostPlatformSuffix(HostPlatform), ThisMonoGameTestNodes)));
                                     }
                                 }
 
@@ -715,11 +764,11 @@ partial class GUBP
                                     {
                                     if (!BranchConfig.HasNode(CookNode.StaticGetFullName(HostPlatform, NonCodeProject, CookedPlatform)))
                                     {
-                                        GameCookNodes.Add(BranchConfig.AddNode(new CookNode(this, BranchConfig, HostPlatform, NonCodeProject, Plat, CookedPlatform)));
+                                        GameCookNodes.Add(BranchConfig.AddNode(new CookNode(BranchConfig, HostPlatform, NonCodeProject, Plat, CookedPlatform)));
                                     }
                                     if (!BranchConfig.HasNode(GamePlatformCookedAndCompiledNode.StaticGetFullName(HostPlatform, NonCodeProject, Plat)))
                                     {
-                                        BranchConfig.AddNode(new GamePlatformCookedAndCompiledNode(this, HostPlatform, NonCodeProject, Plat, false));
+										BranchConfig.AddNode(new GamePlatformCookedAndCompiledNode(BranchConfig, HostPlatform, NonCodeProject, Plat, false));
 
                                         if (NonCodeFormalBuilds.ContainsKey(NonCodeProject.GameName))
                                         {
@@ -728,7 +777,7 @@ partial class GUBP
                                             {
                                                 if (PlatPair.TargetPlatform == Plat)
                                                 {
-														var NodeName = BranchConfig.AddNode(new FormalBuildNode(this, NonCodeProject, HostPlatform, new List<UnrealTargetPlatform>() { Plat }, new List<UnrealTargetConfiguration>() { PlatPair.TargetConfig }));                                                    
+													var NodeName = BranchConfig.AddNode(new FormalBuildNode(BranchConfig, NonCodeProject, HostPlatform, new List<UnrealTargetPlatform>() { Plat }, new List<UnrealTargetConfiguration>() { PlatPair.TargetConfig }));                                                    
 													// we don't want this delayed
 													// this would normally wait for the testing phase, we just want to build it right away
 													BranchConfig.RemovePseudodependencyFromNode(
@@ -755,25 +804,26 @@ partial class GUBP
 													}
 													if (PlatPair.bTest)
 													{
-														BranchConfig.AddNode(new FormalBuildTestNode(this, NonCodeProject, HostPlatform, Plat, PlatPair.TargetConfig));
+														BranchConfig.AddNode(new FormalBuildTestNode(BranchConfig, NonCodeProject, HostPlatform, Plat, PlatPair.TargetConfig));
+													}
 												}
 											}
 										}
-									}
-                                        }
+                                    }
 
-                                    if (!bNoAutomatedTesting)
+										if (!bNoAutomatedTesting)
                                         {
-										if (HostPlatform == UnrealTargetPlatform.Mac || HostPlatform == UnrealTargetPlatform.Linux) continue; //temp hack till Linux and Mac automated testing works
-										var ThisMonoGameTestNodes = new List<string>();
-										foreach (var Test in GameTests)
-										{
-											var TestName = Test.Key + "_" + Plat.ToString();
-											ThisMonoGameTestNodes.Add(BranchConfig.AddNode(new UATTestNode(this, BranchConfig, HostPlatform, NonCodeProject, TestName, Test.Value, CookedAgentSharingGroup, false, RequiredPlatforms)));
-										}
-										if (ThisMonoGameTestNodes.Count > 0)
-                                            {
-											GameTestNodes.Add(BranchConfig.AddNode(new GameAggregateNode(this, HostPlatform, NonCodeProject, "CookedTests_" + Plat.ToString() + "_" + Kind.ToString() + HostPlatformNode.StaticGetHostPlatformSuffix(HostPlatform), ThisMonoGameTestNodes)));
+											if (HostPlatform == UnrealTargetPlatform.Mac || HostPlatform == UnrealTargetPlatform.Linux) continue; //temp hack till Linux and Mac automated testing works
+											var ThisMonoGameTestNodes = new List<string>();
+											foreach (var Test in GameTests)
+											{
+												var TestName = Test.Key + "_" + Plat.ToString();
+												ThisMonoGameTestNodes.Add(BranchConfig.AddNode(new UATTestNode(BranchConfig, HostPlatform, NonCodeProject, TestName, Test.Value, CookedAgentSharingGroup, false, RequiredPlatforms)));
+											}
+											if (ThisMonoGameTestNodes.Count > 0)
+											{
+												GameTestNodes.Add(BranchConfig.AddNode(new GameAggregateNode(HostPlatform, NonCodeProject, "CookedTests_" + Plat.ToString() + "_" + Kind.ToString() + HostPlatformNode.StaticGetHostPlatformSuffix(HostPlatform), ThisMonoGameTestNodes)));
+											}
 										}
 									}
 								}
@@ -781,7 +831,6 @@ partial class GUBP
 						}
 					}
                 }
-                }
 #if false
                 //for now, non-code projects don't do client or server.
                 foreach (var ServerPlatform in ServerPlatforms)
@@ -815,11 +864,11 @@ partial class GUBP
 #endif
                 if (GameTestNodes.Count > 0)
                 {
-                    BranchConfig.AddNode(new GameAggregateNode(this, HostPlatform, Branch.BaseEngineProject, "AllCookedTests", GameTestNodes));
+                    BranchConfig.AddNode(new GameAggregateNode(HostPlatform, Branch.BaseEngineProject, "AllCookedTests", GameTestNodes));
                 }
             }
 			
-			if(HostPlatform == MakeFeaturePacksNode.GetDefaultBuildPlatform(this) && !BranchOptions.bNoInstalledEngine)
+			if(HostPlatform == MakeFeaturePacksNode.GetDefaultBuildPlatform(HostPlatforms) && !BranchOptions.bNoInstalledEngine)
 			{
 				BranchConfig.AddNode(new MakeFeaturePacksNode(HostPlatform, Branch.AllProjects.Where(x => MakeFeaturePacksNode.IsFeaturePack(x))));
 			}
@@ -843,8 +892,8 @@ partial class GUBP
 					EditorGameNode Node = (EditorGameNode)BranchConfig.TryFindNode(EditorGameNode.StaticGetFullName(HostPlatform, CodeProj));
 					if(Node == null)
 					{
-					BranchConfig.AddNode(new EditorGameNode(this, HostPlatform, CodeProj));
-				}
+						BranchConfig.AddNode(new EditorGameNode(BranchConfig, HostPlatform, CodeProj));
+					}
 					else
 					{
 						Node.AddProject(CodeProj);
@@ -864,7 +913,7 @@ partial class GUBP
 						}
 						foreach (var Test in EditorTests)
 						{
-							EditorTestNodes.Add(BranchConfig.AddNode(new UATTestNode(this, BranchConfig, HostPlatform, CodeProj, Test.Key, Test.Value, AgentSharingGroup)));
+							EditorTestNodes.Add(BranchConfig.AddNode(new UATTestNode(BranchConfig, HostPlatform, CodeProj, Test.Key, Test.Value, AgentSharingGroup)));
 							if (!Options.bTestWithShared || !BranchConfig.HasNode(WaitForTestShared.StaticGetFullName()))
 							{
 								BranchConfig.RemovePseudodependencyFromNode((UATTestNode.StaticGetFullName(HostPlatform, CodeProj, Test.Key)), WaitForTestShared.StaticGetFullName());
@@ -872,7 +921,7 @@ partial class GUBP
 						}
 						if (EditorTestNodes.Count > 0)
 						{
-							BranchConfig.AddNode(new GameAggregateNode(this, HostPlatform, CodeProj, "AllEditorTests", EditorTestNodes));
+							BranchConfig.AddNode(new GameAggregateNode(HostPlatform, CodeProj, "AllEditorTests", EditorTestNodes));
 						}
 					}
                 }
@@ -904,7 +953,7 @@ partial class GUBP
 							{
 								if(!BranchConfig.HasNode(GamePlatformMonolithicsNode.StaticGetFullName(HostPlatform, CodeProj, Plat, true)))
 								{
-									BranchConfig.AddNode(new GamePlatformMonolithicsNode(this, BranchConfig, HostPlatform, ActivePlatforms, CodeProj, Plat, true));
+									BranchConfig.AddNode(new GamePlatformMonolithicsNode(BranchConfig, HostPlatform, ActivePlatforms, CodeProj, Plat, true));
 								}
 							}
                             if (ActivePlatforms.Contains(Plat))
@@ -921,9 +970,9 @@ partial class GUBP
                                 {
 									if(GamePlatformMonolithicsNode.HasPrecompiledTargets(CodeProj, HostPlatform, Plat))
 									{
-	                                    BranchConfig.AddNode(new GamePlatformMonolithicsNode(this, BranchConfig, HostPlatform, ActivePlatforms, CodeProj, Plat, InPrecompiled: true));
+	                                    BranchConfig.AddNode(new GamePlatformMonolithicsNode(BranchConfig, HostPlatform, ActivePlatforms, CodeProj, Plat, InPrecompiled: true));
 									}
-                                    BranchConfig.AddNode(new GamePlatformMonolithicsNode(this, BranchConfig, HostPlatform, ActivePlatforms, CodeProj, Plat));
+                                    BranchConfig.AddNode(new GamePlatformMonolithicsNode(BranchConfig, HostPlatform, ActivePlatforms, CodeProj, Plat));
                                 }
                                 var FormalBuildConfigs = Target.Rules.GUBP_GetConfigsForFormalBuilds_MonolithicOnly(HostPlatform);
 								if (!AdditionalPlatforms.Contains(Plat) && (BranchOptions.ProjectsToCook.Contains(CodeProj.GameName) || BranchOptions.ProjectsToCook.Count == 0))
@@ -935,11 +984,11 @@ partial class GUBP
 									}
 									if (!BranchConfig.HasNode(CookNode.StaticGetFullName(HostPlatform, CodeProj, CookedPlatform)))
 									{
-										BranchConfig.AddNode(new CookNode(this, BranchConfig, HostPlatform, CodeProj, Plat, CookedPlatform));
+										BranchConfig.AddNode(new CookNode(BranchConfig, HostPlatform, CodeProj, Plat, CookedPlatform));
 									}
 									if (!BranchConfig.HasNode(GamePlatformCookedAndCompiledNode.StaticGetFullName(HostPlatform, CodeProj, Plat)))
 									{
-										BranchConfig.AddNode(new GamePlatformCookedAndCompiledNode(this, HostPlatform, CodeProj, Plat, true));
+										BranchConfig.AddNode(new GamePlatformCookedAndCompiledNode(BranchConfig, HostPlatform, CodeProj, Plat, true));
 									}
 									
 
@@ -950,21 +999,21 @@ partial class GUBP
 										{
 											if (Plat == Config.TargetPlatform)
 											{
-												FormalNodeName = BranchConfig.AddNode(new FormalBuildNode(this, CodeProj, HostPlatform, InClientTargetPlatforms: new List<UnrealTargetPlatform>() { Config.TargetPlatform }, InClientConfigs: new List<UnrealTargetConfiguration>() { Config.TargetConfig }, InClientNotGame: true));
+												FormalNodeName = BranchConfig.AddNode(new FormalBuildNode(BranchConfig, CodeProj, HostPlatform, InClientTargetPlatforms: new List<UnrealTargetPlatform>() { Config.TargetPlatform }, InClientConfigs: new List<UnrealTargetConfiguration>() { Config.TargetConfig }, InClientNotGame: true));
 											}
 										}
 										else if (Kind == TargetRules.TargetType.Server)
 										{
 											if (Plat == Config.TargetPlatform)
 											{
-												FormalNodeName = BranchConfig.AddNode(new FormalBuildNode(this, CodeProj, HostPlatform, InServerTargetPlatforms: new List<UnrealTargetPlatform>() { Config.TargetPlatform }, InServerConfigs: new List<UnrealTargetConfiguration>() { Config.TargetConfig }));
+												FormalNodeName = BranchConfig.AddNode(new FormalBuildNode(BranchConfig, CodeProj, HostPlatform, InServerTargetPlatforms: new List<UnrealTargetPlatform>() { Config.TargetPlatform }, InServerConfigs: new List<UnrealTargetConfiguration>() { Config.TargetConfig }));
 											}
 										}
 										else if (Kind == TargetRules.TargetType.Game)
 										{
 											if (Plat == Config.TargetPlatform)
 											{
-												FormalNodeName = BranchConfig.AddNode(new FormalBuildNode(this, CodeProj, HostPlatform, InClientTargetPlatforms: new List<UnrealTargetPlatform>() { Config.TargetPlatform }, InClientConfigs: new List<UnrealTargetConfiguration>() { Config.TargetConfig }));
+												FormalNodeName = BranchConfig.AddNode(new FormalBuildNode(BranchConfig, CodeProj, HostPlatform, InClientTargetPlatforms: new List<UnrealTargetPlatform>() { Config.TargetPlatform }, InClientConfigs: new List<UnrealTargetConfiguration>() { Config.TargetConfig }));
 											}
 										}
 										if (FormalNodeName != null)
@@ -996,7 +1045,7 @@ partial class GUBP
 											}
 											if (Config.bTest)
 											{
-												BranchConfig.AddNode(new FormalBuildTestNode(this, CodeProj, HostPlatform, Plat, Config.TargetConfig));
+												BranchConfig.AddNode(new FormalBuildTestNode(BranchConfig, CodeProj, HostPlatform, Plat, Config.TargetConfig));
 											}																				
 										}
 									}
@@ -1010,7 +1059,7 @@ partial class GUBP
 										foreach (var Test in GameTests)
 										{
 											var TestNodeName = Test.Key + "_" + Plat.ToString();
-											ThisMonoGameTestNodes.Add(BranchConfig.AddNode(new UATTestNode(this, BranchConfig, HostPlatform, CodeProj, TestNodeName, Test.Value, CookedAgentSharingGroup, false, RequiredPlatforms)));
+											ThisMonoGameTestNodes.Add(BranchConfig.AddNode(new UATTestNode(BranchConfig, HostPlatform, CodeProj, TestNodeName, Test.Value, CookedAgentSharingGroup, false, RequiredPlatforms)));
 											if (!Options.bTestWithShared || !BranchConfig.HasNode(WaitForTestShared.StaticGetFullName()))
 											{
 												BranchConfig.RemovePseudodependencyFromNode((UATTestNode.StaticGetFullName(HostPlatform, CodeProj, TestNodeName)), WaitForTestShared.StaticGetFullName());
@@ -1018,7 +1067,7 @@ partial class GUBP
 										}
 										if (ThisMonoGameTestNodes.Count > 0)
 										{
-											GameTestNodes.Add(BranchConfig.AddNode(new GameAggregateNode(this, HostPlatform, CodeProj, "CookedTests_" + Plat.ToString() + "_" + Kind.ToString() + HostPlatformNode.StaticGetHostPlatformSuffix(HostPlatform), ThisMonoGameTestNodes)));
+											GameTestNodes.Add(BranchConfig.AddNode(new GameAggregateNode(HostPlatform, CodeProj, "CookedTests_" + Plat.ToString() + "_" + Kind.ToString() + HostPlatformNode.StaticGetHostPlatformSuffix(HostPlatform), ThisMonoGameTestNodes)));
 										}
 									}
 								}
@@ -1043,13 +1092,13 @@ partial class GUBP
 							foreach (var Test in ClientServerTests)
 							{
 								var TestNodeName = Test.Key + "_" + GamePlatform.ToString() + "_" + ServerPlatform.ToString();
-								GameTestNodes.Add(BranchConfig.AddNode(new UATTestNode(this, BranchConfig, HostPlatform, CodeProj, TestNodeName, Test.Value, CookedAgentSharingGroup, false, RequiredPlatforms)));
+								GameTestNodes.Add(BranchConfig.AddNode(new UATTestNode(BranchConfig, HostPlatform, CodeProj, TestNodeName, Test.Value, CookedAgentSharingGroup, false, RequiredPlatforms)));
 							}
 						}
 					}
 					if (GameTestNodes.Count > 0)
 					{
-						BranchConfig.AddNode(new GameAggregateNode(this, HostPlatform, CodeProj, "AllCookedTests", GameTestNodes));
+						BranchConfig.AddNode(new GameAggregateNode(HostPlatform, CodeProj, "AllCookedTests", GameTestNodes));
 					}
 				}
 			}
@@ -1081,31 +1130,31 @@ partial class GUBP
                 }
                 if (PromotedHosts.Count > 0)
                 {
-                    BranchConfig.AddNode(new GameAggregatePromotableNode(this, PromotedHosts, ActivePlatforms, CodeProj, true, bNoIOSOnPC));
+					BranchConfig.AddNode(new GameAggregatePromotableNode(BranchConfig, PromotedHosts, ActivePlatforms, CodeProj, true, bNoIOSOnPC));
                     if (AnySeparate)
                     {
                         BranchConfig.AddNode(new WaitForGamePromotionUserInput(this, CodeProj, false));
-                        BranchConfig.AddNode(new GameLabelPromotableNode(this, CodeProj, false));
+						BranchConfig.AddNode(new GameLabelPromotableNode(BranchConfig, CodeProj, false));
                         BranchConfig.AddNode(new WaitForGamePromotionUserInput(this, CodeProj, true));
-                        BranchConfig.AddNode(new GameLabelPromotableNode(this, CodeProj, true));
+						BranchConfig.AddNode(new GameLabelPromotableNode(BranchConfig, CodeProj, true));
                     }
                 }
             }
         }
         if (NumSharedAllHosts > 0)
         {
-            BranchConfig.AddNode(new GameAggregatePromotableNode(this, HostPlatforms, ActivePlatforms, Branch.BaseEngineProject, false, bNoIOSOnPC));
+			BranchConfig.AddNode(new GameAggregatePromotableNode(BranchConfig, HostPlatforms, ActivePlatforms, Branch.BaseEngineProject, false, bNoIOSOnPC));
 
-            BranchConfig.AddNode(new SharedAggregatePromotableNode(this, HostPlatforms, ActivePlatforms));
+            BranchConfig.AddNode(new SharedAggregatePromotableNode(BranchConfig, ActivePlatforms));
             BranchConfig.AddNode(new WaitForSharedPromotionUserInput(this, false));
-            BranchConfig.AddNode(new SharedLabelPromotableNode(this, false));
+            BranchConfig.AddNode(new SharedLabelPromotableNode(BranchConfig, false));
             BranchConfig.AddNode(new SharedLabelPromotableSuccessNode());
             BranchConfig.AddNode(new WaitForTestShared(this));
             BranchConfig.AddNode(new WaitForSharedPromotionUserInput(this, true));
-            BranchConfig.AddNode(new SharedLabelPromotableNode(this, true));			
+			BranchConfig.AddNode(new SharedLabelPromotableNode(BranchConfig, true));			
         }
 
-		AddCustomNodes(BranchConfig, ActivePlatforms);
+		AddCustomNodes(BranchConfig, HostPlatforms, ActivePlatforms);
         
         if (BranchConfig.HasNode(ToolsForCompileNode.StaticGetFullName(UnrealTargetPlatform.Win64)))
         {
@@ -1116,7 +1165,7 @@ partial class GUBP
 			//AddNode(new VSExpressTestNode(this));
 			if (ActivePlatforms.Contains(UnrealTargetPlatform.Linux) && !BranchOptions.ExcludePlatformsForEditor.Contains(UnrealTargetPlatform.Linux))
 			{
-				BranchConfig.AddNode(new RootEditorCrossCompileLinuxNode(UnrealTargetPlatform.Win64));
+				BranchConfig.AddNode(new RootEditorCrossCompileLinuxNode(BranchConfig, UnrealTargetPlatform.Win64));
 			}
             if (!bPreflightBuild)
             {
@@ -1154,16 +1203,16 @@ partial class GUBP
 
 		// Convert the GUBPNodes and GUBPAggregates maps into lists
 		AllNodes = new List<NodeInfo>();
-		AllNodes.AddRange(BranchConfig.GUBPNodes.Values.Select(x => new NodeInfo{ Name = x.GetFullName(), Node = x }));
+		AllNodes.AddRange(BranchConfig.GUBPNodes.Values.Select(x => new NodeInfo{ Name = x.GetFullName(), Node = x, AgentMemoryRequirement = x.AgentMemoryRequirement(BranchConfig) }));
 
 		AllAggregates = new List<AggregateInfo>();
 		AllAggregates.AddRange(BranchConfig.GUBPAggregates.Values.Select(x => new AggregateInfo{ Name = x.GetFullName(), Node = x }));
 
 		// Calculate the frequencies for each node
-		FindFrequenciesForNodes(AllNodes, FrequencyOverrides);
+		FindFrequenciesForNodes(BranchConfig, AllNodes, FrequencyOverrides);
 
 		// Get the email list for each node
-		FindEmailsForNodes(AllNodes);
+		FindEmailsForNodes(BranchName, AllNodes);
 
 		// Add aggregate nodes for every project in the branch
 		foreach (BranchInfo.BranchUProject GameProj in Branch.AllProjects)
diff --git a/Engine/Source/Programs/AutomationTool/GUBP/LegacyNodes.cs b/Engine/Source/Programs/AutomationTool/GUBP/LegacyNodes.cs
--- a/Engine/Source/Programs/AutomationTool/GUBP/LegacyNodes.cs
+++ b/Engine/Source/Programs/AutomationTool/GUBP/LegacyNodes.cs
@@ -66,7 +66,7 @@ partial class GUBP
 		{
 			return "";
 		}
-        public virtual int AgentMemoryRequirement(GUBP bp)
+        public virtual int AgentMemoryRequirement(GUBP.GUBPBranchConfig BranchConfig)
         {
             return 0;
         }
@@ -80,7 +80,7 @@ partial class GUBP
         /// The builder has a increasing index, specified with -TimeIndex=N
         /// If N mod (1 lshift CISFrequencyQuantumShift()) is not 0, the node is skipped
         /// </summary>
-        public virtual int CISFrequencyQuantumShift(GUBP bp)
+		public virtual int CISFrequencyQuantumShift(GUBP.GUBPBranchConfig BranchConfig)
         {
             return 0;
         }
@@ -310,10 +310,12 @@ partial class GUBP
 
     public class CompileNode : HostPlatformNode
     {
+		protected GUBP.GUBPBranchConfig BranchConfig;
 		bool bDependentOnCompileTools = true;
-        public CompileNode(UnrealTargetPlatform InHostPlatform, bool DependentOnCompileTools = true)
+        public CompileNode(GUBP.GUBPBranchConfig InBranchConfig, UnrealTargetPlatform InHostPlatform, bool DependentOnCompileTools = true)
             : base(InHostPlatform)
         {
+			BranchConfig = InBranchConfig;
 			bDependentOnCompileTools = DependentOnCompileTools;
             if (DependentOnCompileTools)
             {
@@ -345,7 +347,7 @@ partial class GUBP
             UE4Build.BuildAgenda Agenda = GetAgenda(bp);
             if (Agenda != null)
             {
-                bool ReallyDeleteBuildProducts = DeleteBuildProducts() && !bp.bForceIncrementalCompile;
+				bool ReallyDeleteBuildProducts = DeleteBuildProducts() && !BranchConfig.bForceIncrementalCompile;
                 Agenda.DoRetries = false; // these would delete build products
 				bool UseParallelExecutor = bDependentOnCompileTools && (HostPlatform == UnrealTargetPlatform.Win64);
 				UE4Build.Build(Agenda, InDeleteBuildProducts: ReallyDeleteBuildProducts, InUpdateVersionFiles: false, InForceNoXGE: true, InForceUnity: true, InUseParallelExecutor: UseParallelExecutor);
@@ -359,7 +361,7 @@ partial class GUBP
                     AddBuildProduct(Product);
                 }
                 RemoveOveralppingBuildProducts();
-                if (bp.bSignBuildProducts)
+                if (CommandUtils.IsBuildMachine)
                 {
                     // Sign everything we built
 					using(TelemetryStopwatch SignStopwatch = new TelemetryStopwatch("SignBuildProducts"))
@@ -369,7 +371,7 @@ partial class GUBP
                 }
                 PostBuildProducts(bp);
             }
-            if (Agenda == null || (BuildProducts.Count == 0 && bp.bForceIncrementalCompile))
+			if (Agenda == null || (BuildProducts.Count == 0 && BranchConfig.bForceIncrementalCompile))
             {
                 SaveRecordOfSuccessAndAddToBuildProducts("Nothing to actually compile");
             }
@@ -382,12 +384,12 @@ partial class GUBP
 
     public class ToolsForCompileNode : CompileNode
     {
-		int TimeIndex;
+		bool bHasLauncherParam;
 
-        public ToolsForCompileNode(UnrealTargetPlatform InHostPlatform, int InTimeIndex)
-            : base(InHostPlatform, false)
+        public ToolsForCompileNode(GUBP.GUBPBranchConfig InBranchConfig, UnrealTargetPlatform InHostPlatform, bool bInHasLauncherParam)
+            : base(InBranchConfig, InHostPlatform, false)
         {
-			TimeIndex = InTimeIndex;
+			bHasLauncherParam = bInHasLauncherParam;
 
 			if (InHostPlatform != UnrealTargetPlatform.Win64)
             {
@@ -422,18 +424,18 @@ partial class GUBP
         {
             return true;
         }
-        public override int AgentMemoryRequirement(GUBP bp)
+		public override int AgentMemoryRequirement(GUBP.GUBPBranchConfig BranchConfig)
         {
-            if (bp.ParseParam("Launcher") || TimeIndex != 0  && HostPlatform != UnrealTargetPlatform.Mac)
+			if (bHasLauncherParam)
             {
-                return base.AgentMemoryRequirement(bp);
+                return base.AgentMemoryRequirement(BranchConfig);
             }
             return 32;
         }
         public override UE4Build.BuildAgenda GetAgenda(GUBP bp)
         {
             var Agenda = new UE4Build.BuildAgenda();
-			if (HostPlatform == UnrealTargetPlatform.Win64 && !bp.bForceIncrementalCompile)
+			if (HostPlatform == UnrealTargetPlatform.Win64 && !BranchConfig.bForceIncrementalCompile)
             {
                 Agenda.DotNetProjects.AddRange(
                     new string[] 
@@ -464,13 +466,13 @@ partial class GUBP
 
     public class RootEditorNode : CompileNode
     {
-        public RootEditorNode(UnrealTargetPlatform InHostPlatform)
-            : base(InHostPlatform)
+        public RootEditorNode(GUBP.GUBPBranchConfig InBranchConfig, UnrealTargetPlatform InHostPlatform)
+            : base(InBranchConfig, InHostPlatform)
         {
             if (InHostPlatform != UnrealTargetPlatform.Win64)
             {
-            AgentSharingGroup = "Editor" + StaticGetHostPlatformSuffix(InHostPlatform);
-        }
+				AgentSharingGroup = "Editor" + StaticGetHostPlatformSuffix(InHostPlatform);
+			}
         }
         public static string StaticGetFullName(UnrealTargetPlatform InHostPlatform)
         {
@@ -488,7 +490,7 @@ partial class GUBP
 		{
 			base.DoBuild(bp);
 
-			if(!bp.BranchOptions.bNoInstalledEngine)
+			if(!BranchConfig.BranchOptions.bNoInstalledEngine)
 			{
 				FileFilter Filter = new FileFilter();
 				Filter.Include("/Engine/Intermediate/Build/" + HostPlatform.ToString() + "/UE4Editor/Inc/...");
@@ -524,14 +526,14 @@ partial class GUBP
             var Agenda = new UE4Build.BuildAgenda();
 
             string AddArgs = "-nobuilduht";
-			if(!bp.BranchOptions.bNoInstalledEngine)
+			if (!BranchConfig.BranchOptions.bNoInstalledEngine)
 			{
 				AddArgs += " -precompile";
 			}
             Agenda.AddTargets(
-                new string[] { bp.Branch.BaseEngineProject.Properties.Targets[TargetRules.TargetType.Editor].TargetName },
+                new string[] { BranchConfig.Branch.BaseEngineProject.Properties.Targets[TargetRules.TargetType.Editor].TargetName },
                 HostPlatform, UnrealTargetConfiguration.Development, InAddArgs: AddArgs);
-            foreach (var ProgramTarget in bp.Branch.BaseEngineProject.Properties.Programs)
+            foreach (var ProgramTarget in BranchConfig.Branch.BaseEngineProject.Properties.Programs)
             {
                 if (ProgramTarget.Rules.GUBP_AlwaysBuildWithBaseEditor() && ProgramTarget.Rules.SupportsPlatform(HostPlatform))
                 {
@@ -542,12 +544,12 @@ partial class GUBP
         }
         void DeleteStaleDLLs(GUBP bp)
         {
-            if (bp.bForceIncrementalCompile)
+			if (BranchConfig.bForceIncrementalCompile)
             {
                 return;
             }
-            var Targets = new List<string>{bp.Branch.BaseEngineProject.Properties.Targets[TargetRules.TargetType.Editor].TargetName};
-            foreach (var ProgramTarget in bp.Branch.BaseEngineProject.Properties.Programs)
+            var Targets = new List<string>{BranchConfig.Branch.BaseEngineProject.Properties.Targets[TargetRules.TargetType.Editor].TargetName};
+            foreach (var ProgramTarget in BranchConfig.Branch.BaseEngineProject.Properties.Programs)
             {
                 if (ProgramTarget.Rules.GUBP_AlwaysBuildWithBaseEditor() && ProgramTarget.Rules.SupportsPlatform(HostPlatform))
                 {
@@ -613,8 +615,8 @@ partial class GUBP
     }
 	public class RootEditorCrossCompileLinuxNode : CompileNode
 	{
-		public RootEditorCrossCompileLinuxNode(UnrealTargetPlatform InHostPlatform)
-			: base(InHostPlatform)
+		public RootEditorCrossCompileLinuxNode(GUBP.GUBPBranchConfig InBranchConfig, UnrealTargetPlatform InHostPlatform)
+			: base(InBranchConfig, InHostPlatform)
 		{
 			AddDependency(RootEditorNode.StaticGetFullName(UnrealTargetPlatform.Win64));
 			AddDependency(ToolsForCompileNode.StaticGetFullName(UnrealTargetPlatform.Win64));
@@ -627,9 +629,9 @@ partial class GUBP
 		{
 			return StaticGetFullName();
 		}
-		public override int CISFrequencyQuantumShift(GUBP bp)
+		public override int CISFrequencyQuantumShift(GUBP.GUBPBranchConfig BranchConfig)
 		{
-			return base.CISFrequencyQuantumShift(bp) + 3;
+			return base.CISFrequencyQuantumShift(BranchConfig) + 3;
 		}
 		public override UE4Build.BuildAgenda GetAgenda(GUBP bp)
 		{
@@ -637,9 +639,9 @@ partial class GUBP
 
 			string AddArgs = "-nobuilduht";
 			Agenda.AddTargets(
-				new string[] { bp.Branch.BaseEngineProject.Properties.Targets[TargetRules.TargetType.Editor].TargetName },
+				new string[] { BranchConfig.Branch.BaseEngineProject.Properties.Targets[TargetRules.TargetType.Editor].TargetName },
 				UnrealTargetPlatform.Linux, UnrealTargetConfiguration.Development, InAddArgs: AddArgs);
-			 foreach (var ProgramTarget in bp.Branch.BaseEngineProject.Properties.Programs)
+			 foreach (var ProgramTarget in BranchConfig.Branch.BaseEngineProject.Properties.Programs)
             {
                 if (ProgramTarget.Rules.GUBP_AlwaysBuildWithBaseEditor() && ProgramTarget.Rules.SupportsPlatform(UnrealTargetPlatform.Linux))
                 {
@@ -650,12 +652,12 @@ partial class GUBP
 		}
         void DeleteStaleDLLs(GUBP bp)
         {
-            if (bp.bForceIncrementalCompile)
+			if (BranchConfig.bForceIncrementalCompile)
             {
                 return;
             }
-            var Targets = new List<string> { bp.Branch.BaseEngineProject.Properties.Targets[TargetRules.TargetType.Editor].TargetName };
-            foreach (var ProgramTarget in bp.Branch.BaseEngineProject.Properties.Programs)
+            var Targets = new List<string> { BranchConfig.Branch.BaseEngineProject.Properties.Targets[TargetRules.TargetType.Editor].TargetName };
+            foreach (var ProgramTarget in BranchConfig.Branch.BaseEngineProject.Properties.Programs)
             {
                 if (ProgramTarget.Rules.GUBP_AlwaysBuildWithBaseEditor() && ProgramTarget.Rules.SupportsPlatform(UnrealTargetPlatform.Linux))
                 {
@@ -721,12 +723,12 @@ partial class GUBP
 	}
     public class ToolsNode : CompileNode
     {
-        public ToolsNode(GUBP bp, UnrealTargetPlatform InHostPlatform)
-            : base(InHostPlatform)
+        public ToolsNode(GUBP.GUBPBranchConfig InBranchConfig, UnrealTargetPlatform InHostPlatform)
+            : base(InBranchConfig, InHostPlatform)
         {
-			if(!bp.BranchOptions.bNoEditorDependenciesForTools)
+			if(!BranchConfig.BranchOptions.bNoEditorDependenciesForTools)
 			{
-            AddPseudodependency(RootEditorNode.StaticGetFullName(HostPlatform));
+				AddPseudodependency(RootEditorNode.StaticGetFullName(HostPlatform));
 			}
             AgentSharingGroup = "ToolsGroup" + StaticGetHostPlatformSuffix(HostPlatform);
         }
@@ -752,7 +754,7 @@ partial class GUBP
 
 			if (HostPlatform == UnrealTargetPlatform.Win64)
             {
-				if (!bp.bForceIncrementalCompile)
+				if (!BranchConfig.bForceIncrementalCompile)
 				{
 					Agenda.DotNetProjects.AddRange(
 						new string[] 
@@ -767,12 +769,12 @@ partial class GUBP
 							CombinePaths(@"Engine\Source\Programs\NetworkProfiler\NetworkProfiler.sln"),   
 						}
                     );
-                if (!bp.bForceIncrementalCompile)
+				if (!BranchConfig.bForceIncrementalCompile)
                 {
 					Agenda.SwarmProject = CombinePaths(@"Engine\Source\Programs\UnrealSwarm\UnrealSwarm.sln");
 				}
-				
-				bool WithIOS = !bp.BranchOptions.PlatformsToRemove.Contains(UnrealTargetPlatform.IOS);
+
+				bool WithIOS = !BranchConfig.BranchOptions.PlatformsToRemove.Contains(UnrealTargetPlatform.IOS);
 				if ( WithIOS )
 				{
 					Agenda.IOSDotNetProjects.AddRange(
@@ -786,7 +788,7 @@ partial class GUBP
                     );
 				}
 
-				bool WithHTML5 = !bp.BranchOptions.PlatformsToRemove.Contains(UnrealTargetPlatform.HTML5);
+				bool WithHTML5 = !BranchConfig.BranchOptions.PlatformsToRemove.Contains(UnrealTargetPlatform.HTML5);
 				if (WithHTML5)
 				{
 					Agenda.HTML5DotNetProjects.AddRange(
@@ -800,24 +802,24 @@ partial class GUBP
 
             string AddArgs = "-nobuilduht -skipactionhistory -CopyAppBundleBackToDevice";
 
-            foreach (var ProgramTarget in bp.Branch.BaseEngineProject.Properties.Programs)
+            foreach (var ProgramTarget in BranchConfig.Branch.BaseEngineProject.Properties.Programs)
             {
                 bool bInternalOnly;
                 bool SeparateNode;
 				bool CrossCompile;
                 if (ProgramTarget.Rules.GUBP_AlwaysBuildWithTools(HostPlatform, out bInternalOnly, out SeparateNode, out CrossCompile) && ProgramTarget.Rules.SupportsPlatform(HostPlatform) && !bInternalOnly && !SeparateNode)
                 {
-                    if (!bp.BranchOptions.ExcludeNodes.Contains(ProgramTarget.TargetName))
-                    {
-                    foreach (var Plat in ProgramTarget.Rules.GUBP_ToolPlatforms(HostPlatform))
+					if (!BranchConfig.BranchOptions.ExcludeNodes.Contains(ProgramTarget.TargetName))
                     {
-                        foreach (var Config in ProgramTarget.Rules.GUBP_ToolConfigs(HostPlatform))
-                        {
-                            Agenda.AddTargets(new string[] { ProgramTarget.TargetName }, Plat, Config, InAddArgs: AddArgs);
-                        }
-                    }
-                }
-            }
+						foreach (var Plat in ProgramTarget.Rules.GUBP_ToolPlatforms(HostPlatform))
+						{
+							foreach (var Config in ProgramTarget.Rules.GUBP_ToolConfigs(HostPlatform))
+							{
+								Agenda.AddTargets(new string[] { ProgramTarget.TargetName }, Plat, Config, InAddArgs: AddArgs);
+							}
+						}
+					}
+				}
             }
 
             return Agenda;
@@ -825,8 +827,8 @@ partial class GUBP
     }
 	public class ToolsCrossCompileNode : CompileNode
 	{
-		public ToolsCrossCompileNode(UnrealTargetPlatform InHostPlatform)
-			: base(InHostPlatform)
+		public ToolsCrossCompileNode(GUBP.GUBPBranchConfig InBranchConfig, UnrealTargetPlatform InHostPlatform)
+			: base(InBranchConfig, InHostPlatform)
 		{
 			AddPseudodependency(RootEditorCrossCompileLinuxNode.StaticGetFullName());
 			AgentSharingGroup = "ToolsCrossCompileGroup" + StaticGetHostPlatformSuffix(HostPlatform);
@@ -853,7 +855,7 @@ partial class GUBP
 
 			string AddArgs = "-nobuilduht -skipactionhistory -CopyAppBundleBackToDevice";
 
-			foreach (var ProgramTarget in bp.Branch.BaseEngineProject.Properties.Programs)
+			foreach (var ProgramTarget in BranchConfig.Branch.BaseEngineProject.Properties.Programs)
 			{
 				bool bInternalOnly;
 				bool SeparateNode;
@@ -873,8 +875,8 @@ partial class GUBP
     public class SingleToolsNode : CompileNode
     {
         SingleTargetProperties ProgramTarget;
-        public SingleToolsNode(UnrealTargetPlatform InHostPlatform, SingleTargetProperties InProgramTarget)
-            : base(InHostPlatform)
+        public SingleToolsNode(GUBP.GUBPBranchConfig InBranchConfig, UnrealTargetPlatform InHostPlatform, SingleTargetProperties InProgramTarget)
+            : base(InBranchConfig, InHostPlatform)
         {
             ProgramTarget = InProgramTarget;
             AddPseudodependency(RootEditorNode.StaticGetFullName(HostPlatform));
@@ -915,12 +917,12 @@ partial class GUBP
 
     public class InternalToolsNode : CompileNode
     {
-        public InternalToolsNode(GUBP bp, UnrealTargetPlatform InHostPlatform)
-            : base(InHostPlatform)
+        public InternalToolsNode(GUBP.GUBPBranchConfig InBranchConfig, UnrealTargetPlatform InHostPlatform)
+            : base(InBranchConfig, InHostPlatform)
         {
-			if(!bp.BranchOptions.bNoEditorDependenciesForTools)
+			if(!InBranchConfig.BranchOptions.bNoEditorDependenciesForTools)
 			{
-            AddPseudodependency(RootEditorNode.StaticGetFullName(HostPlatform));
+				AddPseudodependency(RootEditorNode.StaticGetFullName(HostPlatform));
 			}
             AgentSharingGroup = "ToolsGroup" + StaticGetHostPlatformSuffix(HostPlatform);
         }
@@ -936,9 +938,9 @@ partial class GUBP
         {
             return base.Priority() - 2;
         }
-        public override int CISFrequencyQuantumShift(GUBP bp)
+		public override int CISFrequencyQuantumShift(GUBP.GUBPBranchConfig BranchConfig)
         {
-            int Result = base.CISFrequencyQuantumShift(bp) + 1;
+            int Result = base.CISFrequencyQuantumShift(BranchConfig) + 1;
             return Result;
         }
         public override bool DeleteBuildProducts()
@@ -978,7 +980,7 @@ partial class GUBP
             }
             string AddArgs = "-nobuilduht -skipactionhistory -CopyAppBundleBackToDevice";
 
-            foreach (var ProgramTarget in bp.Branch.BaseEngineProject.Properties.Programs)
+			foreach (var ProgramTarget in BranchConfig.Branch.BaseEngineProject.Properties.Programs)
             {
                 bool bInternalOnly;
                 bool SeparateNode;
@@ -1006,17 +1008,17 @@ partial class GUBP
     public class SingleInternalToolsNode : CompileNode
     {
         SingleTargetProperties ProgramTarget;
-        public SingleInternalToolsNode(UnrealTargetPlatform InHostPlatform, SingleTargetProperties InProgramTarget)
-            : base(InHostPlatform)
+        public SingleInternalToolsNode(GUBP.GUBPBranchConfig InBranchConfig, UnrealTargetPlatform InHostPlatform, SingleTargetProperties InProgramTarget)
+            : base(InBranchConfig, InHostPlatform)
         {
 			SetupSingleInternalToolsNode(InProgramTarget);
         }
-		public SingleInternalToolsNode(GUBP bp, GUBPBranchConfig BranchConfig, UnrealTargetPlatform InHostPlatform, SingleTargetProperties InProgramTarget)
-            : base(InHostPlatform)
+		public SingleInternalToolsNode(GUBP bp, GUBPBranchConfig InBranchConfig, UnrealTargetPlatform InHostPlatform, SingleTargetProperties InProgramTarget)
+			: base(InBranchConfig, InHostPlatform)
         {
 			// Don't add rooteditor dependency if it isn't in the graph
 			var bRootEditorNodeDoesExit = BranchConfig.HasNode(RootEditorNode.StaticGetFullName(HostPlatform));
-			SetupSingleInternalToolsNode(InProgramTarget, !bRootEditorNodeDoesExit && bp.BranchOptions.bNoEditorDependenciesForTools);
+			SetupSingleInternalToolsNode(InProgramTarget, !bRootEditorNodeDoesExit && BranchConfig.BranchOptions.bNoEditorDependenciesForTools);
         }
 		private void SetupSingleInternalToolsNode(SingleTargetProperties InProgramTarget, bool bSkipRootEditorPsuedoDependency = true)
 		{
@@ -1036,9 +1038,9 @@ partial class GUBP
         {
             return StaticGetFullName(HostPlatform, ProgramTarget);
         }
-        public override int CISFrequencyQuantumShift(GUBP bp)
+		public override int CISFrequencyQuantumShift(GUBP.GUBPBranchConfig BranchConfig)
         {
-            int Result = base.CISFrequencyQuantumShift(bp) + 1;                             
+            int Result = base.CISFrequencyQuantumShift(BranchConfig) + 1;                             
             return Result;
         }
         public override float Priority()
@@ -1070,12 +1072,12 @@ partial class GUBP
     {
 		List<BranchInfo.BranchUProject> GameProjects = new List<BranchInfo.BranchUProject>();
 
-        public EditorGameNode(GUBP bp, UnrealTargetPlatform InHostPlatform, BranchInfo.BranchUProject InGameProj)
-            : base(InHostPlatform)
+        public EditorGameNode(GUBP.GUBPBranchConfig InBranchConfig, UnrealTargetPlatform InHostPlatform, BranchInfo.BranchUProject InGameProj)
+            : base(InBranchConfig, InHostPlatform)
         {
             if (InHostPlatform != UnrealTargetPlatform.Win64)
             {
-            AgentSharingGroup = "Editor"  + StaticGetHostPlatformSuffix(InHostPlatform);
+				AgentSharingGroup = "Editor"  + StaticGetHostPlatformSuffix(InHostPlatform);
             }
             GameProjects.Add(InGameProj);
 			AddDependency(RootEditorNode.StaticGetFullName(InHostPlatform));						
@@ -1122,7 +1124,7 @@ partial class GUBP
 
             string Args = "-nobuilduht -skipactionhistory -skipnonhostplatforms -CopyAppBundleBackToDevice -forceheadergeneration";
 			
-			if(!bp.BranchOptions.bNoInstalledEngine)
+			if(!BranchConfig.BranchOptions.bNoInstalledEngine)
 			{
 				Args += " -precompile";
 			}
@@ -1169,24 +1171,24 @@ partial class GUBP
 			return false;
 		}
 
-		public static UnrealTargetPlatform GetDefaultBuildPlatform(GUBP bp)
+		public static UnrealTargetPlatform GetDefaultBuildPlatform(List<UnrealTargetPlatform> HostPlatforms)
 		{
-			if(bp.HostPlatforms.Contains(UnrealTargetPlatform.Win64))
+			if(HostPlatforms.Contains(UnrealTargetPlatform.Win64))
 			{
 				return UnrealTargetPlatform.Win64;
 			}
-			else if(bp.HostPlatforms.Contains(UnrealTargetPlatform.Mac))
+			else if(HostPlatforms.Contains(UnrealTargetPlatform.Mac))
 			{
 				return UnrealTargetPlatform.Mac;
 			}
 			else
 			{
-				if (UnrealBuildTool.BuildHostPlatform.Current.Platform == UnrealTargetPlatform.Linux && bp.HostPlatforms[0] != UnrealTargetPlatform.Linux)
+				if (UnrealBuildTool.BuildHostPlatform.Current.Platform == UnrealTargetPlatform.Linux && HostPlatforms[0] != UnrealTargetPlatform.Linux)
 				{
-					throw new AutomationException("Linux is not (yet?) able to cross-compile nodes for platform {0}, did you forget -NoPC / -NoMac?", bp.HostPlatforms[0]);
+					throw new AutomationException("Linux is not (yet?) able to cross-compile nodes for platform {0}, did you forget -NoPC / -NoMac?", HostPlatforms[0]);
 				}
 
-				return bp.HostPlatforms[0];
+				return HostPlatforms[0];
 			}
 		}
 
@@ -1200,9 +1202,9 @@ partial class GUBP
             return StaticGetFullName(HostPlatform);
         }
 
-		public override int CISFrequencyQuantumShift(GUBP bp)
+		public override int CISFrequencyQuantumShift(GUBP.GUBPBranchConfig BranchConfig)
 		{
-			return base.CISFrequencyQuantumShift(bp) + 2;
+			return base.CISFrequencyQuantumShift(BranchConfig) + 2;
 		}
 		public override void DoBuild(GUBP bp)
         {
@@ -1253,8 +1255,8 @@ partial class GUBP
 		bool WithXp;
 		bool Precompiled; // If true, just builds targets which generate static libraries for the -UsePrecompiled option to UBT. If false, just build those that don't.
 
-        public GamePlatformMonolithicsNode(GUBP bp, GUBPBranchConfig BranchConfig, UnrealTargetPlatform InHostPlatform, List<UnrealTargetPlatform> InActivePlatforms, BranchInfo.BranchUProject InGameProj, UnrealTargetPlatform InTargetPlatform, bool InWithXp = false, bool InPrecompiled = false)
-            : base(InHostPlatform)
+        public GamePlatformMonolithicsNode(GUBPBranchConfig InBranchConfig, UnrealTargetPlatform InHostPlatform, List<UnrealTargetPlatform> InActivePlatforms, BranchInfo.BranchUProject InGameProj, UnrealTargetPlatform InTargetPlatform, bool InWithXp = false, bool InPrecompiled = false)
+            : base(InBranchConfig, InHostPlatform)
         {
             GameProj = InGameProj;
             TargetPlatform = InTargetPlatform;
@@ -1268,22 +1270,22 @@ partial class GUBP
 				AddDependency(ToolsNode.StaticGetFullName(InHostPlatform));
 			}
 
-            if (InGameProj.GameName != bp.Branch.BaseEngineProject.GameName && GameProj.Properties.Targets.ContainsKey(TargetRules.TargetType.Editor))
+            if (InGameProj.GameName != BranchConfig.Branch.BaseEngineProject.GameName && GameProj.Properties.Targets.ContainsKey(TargetRules.TargetType.Editor))
             {
-				if (!bp.BranchOptions.ExcludePlatformsForEditor.Contains(InHostPlatform))
+				if (!BranchConfig.BranchOptions.ExcludePlatformsForEditor.Contains(InHostPlatform))
 				{
 					AddPseudodependency(EditorGameNode.StaticGetFullName(InHostPlatform, GameProj));
 				}
-                if (BranchConfig.HasNode(GamePlatformMonolithicsNode.StaticGetFullName(HostPlatform, bp.Branch.BaseEngineProject, TargetPlatform)))
+				if (BranchConfig.HasNode(GamePlatformMonolithicsNode.StaticGetFullName(HostPlatform, BranchConfig.Branch.BaseEngineProject, TargetPlatform)))
                 {
-                    AddPseudodependency(GamePlatformMonolithicsNode.StaticGetFullName(HostPlatform, bp.Branch.BaseEngineProject, TargetPlatform));
+					AddPseudodependency(GamePlatformMonolithicsNode.StaticGetFullName(HostPlatform, BranchConfig.Branch.BaseEngineProject, TargetPlatform));
                 }
             }
             else
             {
-                if (TargetPlatform != InHostPlatform && BranchConfig.HasNode(GamePlatformMonolithicsNode.StaticGetFullName(InHostPlatform, bp.Branch.BaseEngineProject, InHostPlatform, Precompiled: Precompiled)))
+				if (TargetPlatform != InHostPlatform && BranchConfig.HasNode(GamePlatformMonolithicsNode.StaticGetFullName(InHostPlatform, BranchConfig.Branch.BaseEngineProject, InHostPlatform, Precompiled: Precompiled)))
                 {
-                    AddPseudodependency(GamePlatformMonolithicsNode.StaticGetFullName(InHostPlatform, bp.Branch.BaseEngineProject, InHostPlatform, Precompiled: Precompiled));
+					AddPseudodependency(GamePlatformMonolithicsNode.StaticGetFullName(InHostPlatform, BranchConfig.Branch.BaseEngineProject, InHostPlatform, Precompiled: Precompiled));
                 }
             }
             if (InGameProj.Options(InHostPlatform).bTestWithShared)  /// compiling templates is only for testing purposes, and we will group them to avoid saturating the farm
@@ -1357,18 +1359,18 @@ partial class GUBP
         {
             return true;
         }
-        public override int AgentMemoryRequirement(GUBP bp)
+		public override int AgentMemoryRequirement(GUBP.GUBPBranchConfig BranchConfig)
         {
-            if (bp.BranchOptions.EnhanceAgentRequirements.Contains(StaticGetFullName(HostPlatform, GameProj, TargetPlatform, WithXp, Precompiled)))
+			if (BranchConfig.BranchOptions.EnhanceAgentRequirements.Contains(StaticGetFullName(HostPlatform, GameProj, TargetPlatform, WithXp, Precompiled)))
             {
                 return 64;
             }
-            return base.AgentMemoryRequirement(bp);
+            return base.AgentMemoryRequirement(BranchConfig);
         }
-        public override int CISFrequencyQuantumShift(GUBP bp)
+        public override int CISFrequencyQuantumShift(GUBP.GUBPBranchConfig BranchConfig)
         {            
-            int Result = base.CISFrequencyQuantumShift(bp);
-            if(GameProj.GameName != bp.Branch.BaseEngineProject.GameName || !Precompiled)
+            int Result = base.CISFrequencyQuantumShift(BranchConfig);
+			if (GameProj.GameName != BranchConfig.Branch.BaseEngineProject.GameName || !Precompiled)
             {
                 Result += 3; //only every 80m
             }
@@ -1498,7 +1500,7 @@ partial class GUBP
 							
 							foreach (var Config in Configs)
 							{
-								if (GameProj.GameName == bp.Branch.BaseEngineProject.GameName)
+								if (GameProj.GameName == BranchConfig.Branch.BaseEngineProject.GameName)
 								{
 									Agenda.AddTargets(new string[] { Target.TargetName }, TargetPlatform, Config, InAddArgs: Args);
 								}
@@ -1662,7 +1664,7 @@ partial class GUBP
     {
         BranchInfo.BranchUProject GameProj;		
 
-        public GameAggregatePromotableNode(GUBP bp, List<UnrealTargetPlatform> InHostPlatforms, List<UnrealTargetPlatform> InActivePlatforms, BranchInfo.BranchUProject InGameProj, bool IsSeparate, bool bNoIOSOnPC)
+        public GameAggregatePromotableNode(GUBP.GUBPBranchConfig BranchConfig, List<UnrealTargetPlatform> InHostPlatforms, List<UnrealTargetPlatform> InActivePlatforms, BranchInfo.BranchUProject InGameProj, bool IsSeparate, bool bNoIOSOnPC)
             : base(InHostPlatforms, InGameProj.GameName)
         {
             GameProj = InGameProj;
@@ -1670,12 +1672,12 @@ partial class GUBP
             foreach (var HostPlatform in HostPlatforms)
             {
                 AddDependency(RootEditorNode.StaticGetFullName(HostPlatform));
-			    if(!bp.BranchOptions.PromotablesWithoutTools.Contains(GameProj.GameName))
+			    if(!BranchConfig.BranchOptions.PromotablesWithoutTools.Contains(GameProj.GameName))
                 {
                     AddDependency(ToolsNode.StaticGetFullName(HostPlatform));
                     AddDependency(InternalToolsNode.StaticGetFullName(HostPlatform));
                 }
-                if (InGameProj.GameName != bp.Branch.BaseEngineProject.GameName && GameProj.Properties.Targets.ContainsKey(TargetRules.TargetType.Editor))
+                if (InGameProj.GameName != BranchConfig.Branch.BaseEngineProject.GameName && GameProj.Properties.Targets.ContainsKey(TargetRules.TargetType.Editor))
                 {
                     AddDependency(EditorGameNode.StaticGetFullName(HostPlatform, GameProj));
                 }				
@@ -1726,20 +1728,20 @@ partial class GUBP
     public class SharedAggregatePromotableNode : AggregatePromotableNode
     {
 
-        public SharedAggregatePromotableNode(GUBP bp, List<UnrealTargetPlatform> InHostPlatforms, List<UnrealTargetPlatform> InActivePlatforms)
-            : base(InHostPlatforms, "Shared")
+        public SharedAggregatePromotableNode(GUBP.GUBPBranchConfig BranchConfig, List<UnrealTargetPlatform> InActivePlatforms)
+            : base(BranchConfig.HostPlatforms, "Shared")
         {
             foreach (var HostPlatform in HostPlatforms)
             {
                 AddDependency(EditorAndToolsNode.StaticGetFullName(HostPlatform));
                 {
-                    var Options = bp.Branch.BaseEngineProject.Options(HostPlatform);
+                    var Options = BranchConfig.Branch.BaseEngineProject.Options(HostPlatform);
                     if (Options.bIsPromotable && !Options.bSeparateGamePromotion)
                     {
-                        AddDependency(GameAggregatePromotableNode.StaticGetFullName(bp.Branch.BaseEngineProject));
+                        AddDependency(GameAggregatePromotableNode.StaticGetFullName(BranchConfig.Branch.BaseEngineProject));
                     }
                 }
-                foreach (var CodeProj in bp.Branch.CodeProjects)
+                foreach (var CodeProj in BranchConfig.Branch.CodeProjects)
                 {
                     var Options = CodeProj.Options(HostPlatform);
                     if (!Options.bSeparateGamePromotion)
@@ -1763,10 +1765,10 @@ partial class GUBP
                     AddDependency(ToolsCrossCompileNode.StaticGetFullName(HostPlatform));
                 }
             }
-			if(!bp.BranchOptions.bNoInstalledEngine)
+			if (!BranchConfig.BranchOptions.bNoInstalledEngine)
 			{
-			AddDependency(MakeFeaturePacksNode.StaticGetFullName(MakeFeaturePacksNode.GetDefaultBuildPlatform(bp)));
-        }
+				AddDependency(MakeFeaturePacksNode.StaticGetFullName(MakeFeaturePacksNode.GetDefaultBuildPlatform(BranchConfig.HostPlatforms)));
+			}
         }
 		public override bool IsSeparatePromotable()
 		{
@@ -2007,11 +2009,13 @@ partial class GUBP
 
     public class LabelPromotableNode : GUBPNode
     {
+		GUBP.GUBPBranchConfig BranchConfig;
         string PromotionLabelPrefix;        
         protected bool bLabelPromoted; // true if this is the promoted version
 
-        public LabelPromotableNode(string InPromotionLabelPrefix, string InPromotionLabelSuffix, bool bInLabelPromoted)
+        public LabelPromotableNode(GUBP.GUBPBranchConfig InBranchConfig, string InPromotionLabelPrefix, string InPromotionLabelSuffix, bool bInLabelPromoted)
         {
+			BranchConfig = InBranchConfig;
             PromotionLabelPrefix = InPromotionLabelPrefix;            
             bLabelPromoted = bInLabelPromoted;
             AddDependency(WaitForPromotionUserInput.StaticGetFullName(PromotionLabelPrefix, InPromotionLabelSuffix, bLabelPromoted));
@@ -2047,7 +2051,7 @@ partial class GUBP
         {
             BuildProducts = new List<string>();
 
-            if (P4Enabled && !bp.bPreflightBuild)
+			if (P4Enabled && !BranchConfig.bPreflightBuild)
             {
                 if (AllDependencyBuildProducts.Count == 0)
                 {
@@ -2061,8 +2065,8 @@ partial class GUBP
                 }
                 else
                 {
-					int WorkingCL = P4.CreateChange(P4Env.Client, String.Format("GUBP Node {0} built from changelist {1}", GetFullName(), bp.CL));
-                    Log("Build from {0}    Working in {1}", bp.CL, WorkingCL);
+					int WorkingCL = P4.CreateChange(P4Env.Client, String.Format("GUBP Node {0} built from changelist {1}", GetFullName(), BranchConfig.CL));
+                    Log("Build from {0}    Working in {1}", BranchConfig.CL, WorkingCL);
 
                     var ProductsToSubmit = new List<String>();
 
@@ -2112,8 +2116,8 @@ partial class GUBP
     public class GameLabelPromotableNode : LabelPromotableNode
     {
         BranchInfo.BranchUProject GameProj;        
-        public GameLabelPromotableNode(GUBP bp, BranchInfo.BranchUProject InGameProj, bool bInLabelPromoted)
-            : base(InGameProj.GameName, "", bInLabelPromoted)
+        public GameLabelPromotableNode(GUBP.GUBPBranchConfig InBranchConfig, BranchInfo.BranchUProject InGameProj, bool bInLabelPromoted)
+			: base(InBranchConfig, InGameProj.GameName, "", bInLabelPromoted)
         {
             GameProj = InGameProj;
         }
@@ -2130,8 +2134,8 @@ partial class GUBP
 
     public class SharedLabelPromotableNode : LabelPromotableNode
     {        
-        public SharedLabelPromotableNode(GUBP bp, bool bInLabelPromoted)
-            : base("Shared", IsMainBranch(), bInLabelPromoted)
+        public SharedLabelPromotableNode(GUBP.GUBPBranchConfig InBranchConfig, bool bInLabelPromoted)
+            : base(InBranchConfig, "Shared", IsMainBranch(), bInLabelPromoted)
         {
         }
 
@@ -2195,7 +2199,7 @@ partial class GUBP
         bool bIsMassive;
 		
 
-        public CookNode(GUBP bp, GUBPBranchConfig BranchConfig, UnrealTargetPlatform InHostPlatform, BranchInfo.BranchUProject InGameProj, UnrealTargetPlatform InTargetPlatform, string InCookPlatform)
+        public CookNode(GUBPBranchConfig BranchConfig, UnrealTargetPlatform InHostPlatform, BranchInfo.BranchUProject InGameProj, UnrealTargetPlatform InTargetPlatform, string InCookPlatform)
             : base(InHostPlatform)
         {
             GameProj = InGameProj;
@@ -2205,7 +2209,7 @@ partial class GUBP
             AddDependency(EditorAndToolsNode.StaticGetFullName(HostPlatform));
             bool bIsShared = false;
             // is this the "base game" or a non code project?
-            if (InGameProj.GameName != bp.Branch.BaseEngineProject.GameName && GameProj.Properties.Targets.ContainsKey(TargetRules.TargetType.Editor))
+            if (InGameProj.GameName != BranchConfig.Branch.BaseEngineProject.GameName && GameProj.Properties.Targets.ContainsKey(TargetRules.TargetType.Editor))
             {
                 var Options = InGameProj.Options(HostPlatform);
                 bIsMassive = Options.bIsMassive;				
@@ -2222,15 +2226,15 @@ partial class GUBP
                 {
                     bIsShared = true;
                 }
-				if(!bp.BranchOptions.bNoMonolithicDependenciesForCooks)
+				if (!BranchConfig.BranchOptions.bNoMonolithicDependenciesForCooks)
 				{
-                AddPseudodependency(GamePlatformMonolithicsNode.StaticGetFullName(HostPlatform, GameProj, TargetPlatform));
-            }
+					AddPseudodependency(GamePlatformMonolithicsNode.StaticGetFullName(HostPlatform, GameProj, TargetPlatform));
+				}
             }
             else
             {
                 bIsShared = true;				
-                AddPseudodependency(GamePlatformMonolithicsNode.StaticGetFullName(HostPlatform, bp.Branch.BaseEngineProject, TargetPlatform));
+                AddPseudodependency(GamePlatformMonolithicsNode.StaticGetFullName(HostPlatform, BranchConfig.Branch.BaseEngineProject, TargetPlatform));
             }
             if (bIsShared)
             {
@@ -2239,17 +2243,17 @@ partial class GUBP
                 AgentSharingGroup = "SharedCooks" + StaticGetHostPlatformSuffix(HostPlatform);
 
                 // If the cook fails for the base engine, don't bother trying
-                if (InGameProj.GameName != bp.Branch.BaseEngineProject.GameName && BranchConfig.HasNode(CookNode.StaticGetFullName(HostPlatform, bp.Branch.BaseEngineProject, CookPlatform)))
+                if (InGameProj.GameName != BranchConfig.Branch.BaseEngineProject.GameName && BranchConfig.HasNode(CookNode.StaticGetFullName(HostPlatform, BranchConfig.Branch.BaseEngineProject, CookPlatform)))
                 {
-                    AddPseudodependency(CookNode.StaticGetFullName(HostPlatform, bp.Branch.BaseEngineProject, CookPlatform));
+                    AddPseudodependency(CookNode.StaticGetFullName(HostPlatform, BranchConfig.Branch.BaseEngineProject, CookPlatform));
                 }
 
                 // If the base cook platform fails, don't bother trying other ones
                 string BaseCookedPlatform = Platform.Platforms[HostPlatform].GetCookPlatform(false, false, "");
-                if (InGameProj.GameName == bp.Branch.BaseEngineProject.GameName && CookPlatform != BaseCookedPlatform &&
-                    BranchConfig.HasNode(CookNode.StaticGetFullName(HostPlatform, bp.Branch.BaseEngineProject, BaseCookedPlatform)))
+                if (InGameProj.GameName == BranchConfig.Branch.BaseEngineProject.GameName && CookPlatform != BaseCookedPlatform &&
+                    BranchConfig.HasNode(CookNode.StaticGetFullName(HostPlatform, BranchConfig.Branch.BaseEngineProject, BaseCookedPlatform)))
                 {
-                    AddPseudodependency(CookNode.StaticGetFullName(HostPlatform, bp.Branch.BaseEngineProject, BaseCookedPlatform));
+                    AddPseudodependency(CookNode.StaticGetFullName(HostPlatform, BranchConfig.Branch.BaseEngineProject, BaseCookedPlatform));
                 }
             }			
         }
@@ -2265,15 +2269,15 @@ partial class GUBP
         {
             return GameProj.GameName;
         }
-        public override int CISFrequencyQuantumShift(GUBP bp)
+		public override int CISFrequencyQuantumShift(GUBP.GUBPBranchConfig BranchConfig)
         {
-            return base.CISFrequencyQuantumShift(bp) + 4 + (bIsMassive ? 1 : 0);
+            return base.CISFrequencyQuantumShift(BranchConfig) + 4 + (bIsMassive ? 1 : 0);
         }
         public override float Priority()
         {
             return 10.0f;
         }
-        public override int AgentMemoryRequirement(GUBP bp)
+		public override int AgentMemoryRequirement(GUBP.GUBPBranchConfig BranchConfig)
         {
             return bIsMassive ? 32 : 0;
         }
@@ -2316,7 +2320,7 @@ partial class GUBP
         BranchInfo.BranchUProject GameProj;
         UnrealTargetPlatform TargetPlatform;
 
-        public GamePlatformCookedAndCompiledNode(GUBP bp, UnrealTargetPlatform InHostPlatform, BranchInfo.BranchUProject InGameProj, UnrealTargetPlatform InTargetPlatform, bool bCodeProject)
+        public GamePlatformCookedAndCompiledNode(GUBP.GUBPBranchConfig BranchConfig, UnrealTargetPlatform InHostPlatform, BranchInfo.BranchUProject InGameProj, UnrealTargetPlatform InTargetPlatform, bool bCodeProject)
             : base(InHostPlatform)
         {
             GameProj = InGameProj;
@@ -2351,16 +2355,16 @@ partial class GUBP
                 {
                     if (Kind == TargetRules.TargetType.Game) //for now, non-code projects don't do client or server.
                     {
-                        if (bp.Branch.BaseEngineProject.Properties.Targets.ContainsKey(Kind))
+                        if (BranchConfig.Branch.BaseEngineProject.Properties.Targets.ContainsKey(Kind))
                         {
-                            var Target = bp.Branch.BaseEngineProject.Properties.Targets[Kind];
+                            var Target = BranchConfig.Branch.BaseEngineProject.Properties.Targets[Kind];
                             var Platforms = Target.Rules.GUBP_GetPlatforms_MonolithicOnly(HostPlatform);
                             if (Platforms.Contains(TargetPlatform) && Target.Rules.SupportsPlatform(TargetPlatform))
                             {
                                 //@todo how do we get the client target platform?
                                 string CookedPlatform = Platform.Platforms[TargetPlatform].GetCookPlatform(Kind == TargetRules.TargetType.Server, Kind == TargetRules.TargetType.Client, "");
                                 AddDependency(CookNode.StaticGetFullName(HostPlatform, GameProj, CookedPlatform));
-                                AddDependency(GamePlatformMonolithicsNode.StaticGetFullName(HostPlatform, bp.Branch.BaseEngineProject, TargetPlatform));
+                                AddDependency(GamePlatformMonolithicsNode.StaticGetFullName(HostPlatform, BranchConfig.Branch.BaseEngineProject, TargetPlatform));
                             }
                         }
                     }
@@ -2384,6 +2388,7 @@ partial class GUBP
 
     public class FormalBuildNode : HostPlatformNode
     {
+		GUBP.GUBPBranchConfig BranchConfig;
         BranchInfo.BranchUProject GameProj;
 
         //CAUTION, these are lists, but it isn't clear that lists really work on all platforms, so we stick to one node per build
@@ -2395,7 +2400,7 @@ partial class GUBP
 		bool bIsCode;
         UnrealBuildTool.TargetRules.TargetType GameOrClient;
 
-        public FormalBuildNode(GUBP bp,
+        public FormalBuildNode(GUBP.GUBPBranchConfig InBranchConfig,
             BranchInfo.BranchUProject InGameProj,
             UnrealTargetPlatform InHostPlatform,
             List<UnrealTargetPlatform> InClientTargetPlatforms = null,
@@ -2406,6 +2411,7 @@ partial class GUBP
             )
             : base(InHostPlatform)
         {
+			BranchConfig = InBranchConfig;
             GameProj = InGameProj;
             ClientTargetPlatforms = InClientTargetPlatforms;
             ServerTargetPlatforms = InServerTargetPlatforms;
@@ -2419,7 +2425,7 @@ partial class GUBP
             {
                 GameOrClient = TargetRules.TargetType.Client;
             }
-			if (InGameProj.GameName != bp.Branch.BaseEngineProject.GameName && GameProj.Properties.Targets.ContainsKey(TargetRules.TargetType.Editor))
+			if (InGameProj.GameName != BranchConfig.Branch.BaseEngineProject.GameName && GameProj.Properties.Targets.ContainsKey(TargetRules.TargetType.Editor))
 			{
 				bIsCode = true;
 			}
@@ -2433,7 +2439,7 @@ partial class GUBP
             if (!WorkingGameProject.Properties.Targets.ContainsKey(TargetRules.TargetType.Editor))
             {
                 // this is a codeless project, use the base project
-                WorkingGameProject = bp.Branch.BaseEngineProject;
+                WorkingGameProject = BranchConfig.Branch.BaseEngineProject;
             }
 
             var AllTargetPlatforms = new List<UnrealTargetPlatform>();					
@@ -2579,18 +2585,18 @@ partial class GUBP
         {
             return base.Priority() + 20.0f;
         }
-		public override int CISFrequencyQuantumShift(GUBP bp)
+		public override int CISFrequencyQuantumShift(GUBP.GUBPBranchConfig BranchConfig)
 		{
-			return base.CISFrequencyQuantumShift(bp) + 3;
+			return base.CISFrequencyQuantumShift(BranchConfig) + 3;
 		}
-        public static string GetArchiveDirectory(GUBP bp, BranchInfo.BranchUProject InGameProj, UnrealTargetPlatform InHostPlatform, List<UnrealTargetPlatform> InClientTargetPlatforms = null, List<UnrealTargetConfiguration> InClientConfigs = null, List<UnrealTargetPlatform> InServerTargetPlatforms = null, List<UnrealTargetConfiguration> InServerConfigs = null, bool InClientNotGame = false)
+		public static string GetArchiveDirectory(GUBP.GUBPBranchConfig BranchConfig, BranchInfo.BranchUProject InGameProj, UnrealTargetPlatform InHostPlatform, List<UnrealTargetPlatform> InClientTargetPlatforms = null, List<UnrealTargetConfiguration> InClientConfigs = null, List<UnrealTargetPlatform> InServerTargetPlatforms = null, List<UnrealTargetConfiguration> InServerConfigs = null, bool InClientNotGame = false)
         {
             string BaseDir = TempStorage.ResolveSharedBuildDirectory(InGameProj.GameName);
             string NodeName = StaticGetFullName(InGameProj, InHostPlatform, InClientTargetPlatforms, InClientConfigs, InServerTargetPlatforms, InServerConfigs, InClientNotGame);
             string Inner = P4Env.BuildRootEscaped + "-CL-" + P4Env.ChangelistString;
-            if (bp.bPreflightBuild)
+            if (BranchConfig.bPreflightBuild)
             {
-                Inner = Inner + bp.PreflightMangleSuffix;
+                Inner = Inner + BranchConfig.PreflightMangleSuffix;
             }
             string ArchiveDirectory = CombinePaths(BaseDir, NodeName, Inner);
             return ArchiveDirectory;
@@ -2706,7 +2712,7 @@ partial class GUBP
 			string IntermediateArchiveDirectory = FinalArchiveDirectory;
             if (P4Enabled)
             {
-                FinalArchiveDirectory = GetArchiveDirectory(bp, GameProj, HostPlatform, ClientTargetPlatforms, ClientConfigs, ServerTargetPlatforms, ServerConfigs, ClientNotGame);
+                FinalArchiveDirectory = GetArchiveDirectory(BranchConfig, GameProj, HostPlatform, ClientTargetPlatforms, ClientConfigs, ServerTargetPlatforms, ServerConfigs, ClientNotGame);
 				IntermediateArchiveDirectory = FinalArchiveDirectory;
 				// Xbox One packaging does not function with remote file systems. Use a temp local directory to package and then move files into final location.
 				if (bXboxOneTarget)
@@ -2768,12 +2774,13 @@ partial class GUBP
 
     public class FormalBuildTestNode : TestNode
     {
+		GUBP.GUBPBranchConfig BranchConfig;
         BranchInfo.BranchUProject GameProj;
         UnrealTargetPlatform ClientTargetPlatform;
         UnrealTargetConfiguration ClientConfig;
         UnrealBuildTool.TargetRules.TargetType GameOrClient;
 
-        public FormalBuildTestNode(GUBP bp,
+        public FormalBuildTestNode(GUBP.GUBPBranchConfig InBranchConfig,
             BranchInfo.BranchUProject InGameProj,
             UnrealTargetPlatform InHostPlatform,
             UnrealTargetPlatform InClientTargetPlatform,
@@ -2781,6 +2788,7 @@ partial class GUBP
             )
             : base(InHostPlatform)
         {
+			BranchConfig = InBranchConfig;
             GameProj = InGameProj;
             ClientTargetPlatform = InClientTargetPlatform;
             ClientConfig = InClientConfig;
@@ -2791,7 +2799,7 @@ partial class GUBP
             if (!WorkingGameProject.Properties.Targets.ContainsKey(TargetRules.TargetType.Editor))
             {
                 // this is a codeless project, use the base project
-                WorkingGameProject = bp.Branch.BaseEngineProject;
+                WorkingGameProject = BranchConfig.Branch.BaseEngineProject;
             }
             if (!WorkingGameProject.Properties.Targets.ContainsKey(GameOrClient))
             {
@@ -2827,9 +2835,9 @@ partial class GUBP
         {
             return base.Priority() - 20.0f;
         }
-        public override int CISFrequencyQuantumShift(GUBP bp)
+		public override int CISFrequencyQuantumShift(GUBP.GUBPBranchConfig BranchConfig)
         {
-            return base.CISFrequencyQuantumShift(bp) + 3;
+            return base.CISFrequencyQuantumShift(BranchConfig) + 3;
         }
         public override void DoTest(GUBP bp)
         {
@@ -2839,7 +2847,7 @@ partial class GUBP
                 ProjectArg = " -project=\"" + GameProj.FilePath + "\"";
             }
 
-            string ArchiveDirectory = FormalBuildNode.GetArchiveDirectory(bp, GameProj, HostPlatform, new List<UnrealTargetPlatform>() { ClientTargetPlatform }, InClientConfigs: new List<UnrealTargetConfiguration>() { ClientConfig }, InClientNotGame: GameOrClient == TargetRules.TargetType.Client);
+            string ArchiveDirectory = FormalBuildNode.GetArchiveDirectory(BranchConfig, GameProj, HostPlatform, new List<UnrealTargetPlatform>() { ClientTargetPlatform }, InClientConfigs: new List<UnrealTargetConfiguration>() { ClientConfig }, InClientNotGame: GameOrClient == TargetRules.TargetType.Client);
             if (!DirectoryExists_NoExceptions(ArchiveDirectory))
             {
                 throw new AutomationException("Archive directory does not exist {0}, so we can't test the build.", ArchiveDirectory);
@@ -2884,18 +2892,18 @@ partial class GUBP
 		{
 			return StaticGetFullName(HostPlatform, ProgramTarget);
 		}
-		public override int CISFrequencyQuantumShift(GUBP bp)
+		public override int CISFrequencyQuantumShift(GUBP.GUBPBranchConfig BranchConfig)
 		{
-			int Result = base.CISFrequencyQuantumShift(bp) + 2;
+			int Result = base.CISFrequencyQuantumShift(BranchConfig) + 2;
 			if (HostPlatform == UnrealTargetPlatform.Mac)
 			{
 				Result += 1;
 			}		
 			return Result;
 		}
-		public override int AgentMemoryRequirement(GUBP bp)
+		public override int AgentMemoryRequirement(GUBP.GUBPBranchConfig BranchConfig)
 		{
-			int Result = base.AgentMemoryRequirement(bp);
+			int Result = base.AgentMemoryRequirement(BranchConfig);
 			if (HostPlatform == UnrealTargetPlatform.Mac)
 			{
 				Result = 32;
@@ -2919,9 +2927,12 @@ partial class GUBP
 
     public class NonUnityTestNode : TestNode
     {
-        public NonUnityTestNode(UnrealTargetPlatform InHostPlatform)
+		GUBP.GUBPBranchConfig BranchConfig;
+
+        public NonUnityTestNode(GUBP.GUBPBranchConfig InBranchConfig, UnrealTargetPlatform InHostPlatform)
             : base(InHostPlatform)
         {
+			BranchConfig = InBranchConfig;
             AddPseudodependency(RootEditorNode.StaticGetFullName(HostPlatform));
         }
         public static string StaticGetFullName(UnrealTargetPlatform InHostPlatform)
@@ -2932,18 +2943,18 @@ partial class GUBP
         {
             return StaticGetFullName(HostPlatform);
         }
-        public override int CISFrequencyQuantumShift(GUBP bp)
+		public override int CISFrequencyQuantumShift(GUBP.GUBPBranchConfig BranchConfig)
         {
-            int Result = base.CISFrequencyQuantumShift(bp) + 2;
+            int Result = base.CISFrequencyQuantumShift(BranchConfig) + 2;
             if(HostPlatform == UnrealTargetPlatform.Mac)
             {
                 Result += 1;
             }
             return Result;
         }
-        public override int AgentMemoryRequirement(GUBP bp)
+		public override int AgentMemoryRequirement(GUBP.GUBPBranchConfig BranchConfig)
         {
-            int Result = base.AgentMemoryRequirement(bp);
+            int Result = base.AgentMemoryRequirement(BranchConfig);
             if(HostPlatform == UnrealTargetPlatform.Mac)
             {
                 Result = 32;
@@ -2957,14 +2968,14 @@ partial class GUBP
             
             Agenda.AddTargets(new string[] { "UnrealHeaderTool" }, HostPlatform, UnrealTargetConfiguration.Development);
             Agenda.AddTargets(
-                new string[] { bp.Branch.BaseEngineProject.Properties.Targets[TargetRules.TargetType.Editor].TargetName },
+                new string[] { BranchConfig.Branch.BaseEngineProject.Properties.Targets[TargetRules.TargetType.Editor].TargetName },
                 HostPlatform, UnrealTargetConfiguration.Development, InAddArgs: "-skipnonhostplatforms -shadowvariableerrors");
 
             foreach (var Kind in BranchInfo.MonolithicKinds)
             {
-                if (bp.Branch.BaseEngineProject.Properties.Targets.ContainsKey(Kind))
+                if (BranchConfig.Branch.BaseEngineProject.Properties.Targets.ContainsKey(Kind))
                 {
-                    var Target = bp.Branch.BaseEngineProject.Properties.Targets[Kind];
+                    var Target = BranchConfig.Branch.BaseEngineProject.Properties.Targets[Kind];
                     Agenda.AddTargets(new string[] { Target.TargetName }, HostPlatform, UnrealTargetConfiguration.Development, InAddArgs: "-shadowvariableerrors");
                 }
             }
@@ -2978,12 +2989,15 @@ partial class GUBP
 
     public class IOSOnPCTestNode : TestNode
     {
-        public IOSOnPCTestNode(GUBP bp)
+		GUBP.GUBPBranchConfig BranchConfig;
+
+        public IOSOnPCTestNode(GUBP.GUBPBranchConfig InBranchConfig)
             : base(UnrealTargetPlatform.Win64)
         {
+			BranchConfig = InBranchConfig;
             AddDependency(ToolsForCompileNode.StaticGetFullName(UnrealTargetPlatform.Win64));
             AddDependency(ToolsNode.StaticGetFullName(UnrealTargetPlatform.Win64));
-            AddPseudodependency(GamePlatformMonolithicsNode.StaticGetFullName(UnrealTargetPlatform.Mac, bp.Branch.BaseEngineProject, UnrealTargetPlatform.IOS));
+			AddPseudodependency(GamePlatformMonolithicsNode.StaticGetFullName(UnrealTargetPlatform.Mac, BranchConfig.Branch.BaseEngineProject, UnrealTargetPlatform.IOS));
         }
         public static string StaticGetFullName()
         {
@@ -2993,16 +3007,16 @@ partial class GUBP
         {
             return StaticGetFullName();
         }
-        public override int CISFrequencyQuantumShift(GUBP bp)
+		public override int CISFrequencyQuantumShift(GUBP.GUBPBranchConfig BranchConfig)
         {
-            return base.CISFrequencyQuantumShift(bp) + 3;
+            return base.CISFrequencyQuantumShift(BranchConfig) + 3;
         }
         public override void DoTest(GUBP bp)
         {
             var Build = new UE4Build(bp);
             var Agenda = new UE4Build.BuildAgenda();
 
-            Agenda.AddTargets(new string[] { bp.Branch.BaseEngineProject.Properties.Targets[TargetRules.TargetType.Game].TargetName }, UnrealTargetPlatform.IOS, UnrealTargetConfiguration.Development);
+            Agenda.AddTargets(new string[] { BranchConfig.Branch.BaseEngineProject.Properties.Targets[TargetRules.TargetType.Game].TargetName }, UnrealTargetPlatform.IOS, UnrealTargetConfiguration.Development);
 
             Build.Build(Agenda, InDeleteBuildProducts: true, InUpdateVersionFiles: false);
 
@@ -3013,9 +3027,12 @@ partial class GUBP
 	
 	public class VSExpressTestNode : TestNode
 	{
-		public VSExpressTestNode(GUBP bp)
+		GUBP.GUBPBranchConfig BranchConfig;
+
+		public VSExpressTestNode(GUBP.GUBPBranchConfig InBranchConfig)
 			: base(UnrealTargetPlatform.Win64)
 		{
+			BranchConfig = InBranchConfig;
 			AddDependency(ToolsForCompileNode.StaticGetFullName(UnrealTargetPlatform.Win64));
 			AddDependency(RootEditorNode.StaticGetFullName(UnrealTargetPlatform.Win64));
 		}
@@ -3027,9 +3044,9 @@ partial class GUBP
 		{
 			return StaticGetFullName();
 		}
-		public override int CISFrequencyQuantumShift(GUBP bp)
+		public override int CISFrequencyQuantumShift(GUBP.GUBPBranchConfig BranchConfig)
 		{
-			return base.CISFrequencyQuantumShift(bp) + 3;
+			return base.CISFrequencyQuantumShift(BranchConfig) + 3;
 		}
 		public override string ECAgentString()
 		{
@@ -3043,9 +3060,9 @@ partial class GUBP
 			string AddArgs = "-nobuilduht";
 
 			Agenda.AddTargets(
-				new string[] { bp.Branch.BaseEngineProject.Properties.Targets[TargetRules.TargetType.Editor].TargetName },
+				new string[] { BranchConfig.Branch.BaseEngineProject.Properties.Targets[TargetRules.TargetType.Editor].TargetName },
 				HostPlatform, UnrealTargetConfiguration.Development, InAddArgs: AddArgs);
-			foreach (var ProgramTarget in bp.Branch.BaseEngineProject.Properties.Programs)
+			foreach (var ProgramTarget in BranchConfig.Branch.BaseEngineProject.Properties.Programs)
 			{
 				if (ProgramTarget.Rules.GUBP_AlwaysBuildWithBaseEditor() && ProgramTarget.Rules.SupportsPlatform(HostPlatform))
 				{
@@ -3067,7 +3084,7 @@ partial class GUBP
         List<UnrealTargetPlatform> DependsOnCooked;
         float ECPriority;
 
-        public UATTestNode(GUBP bp, GUBPBranchConfig BranchConfig, UnrealTargetPlatform InHostPlatform, BranchInfo.BranchUProject InGameProj, string InTestName, string InUATCommandLine, string InAgentSharingGroup, bool InDependsOnEditor = true, List<UnrealTargetPlatform> InDependsOnCooked = null, float InECPriority = 0.0f)
+        public UATTestNode(GUBPBranchConfig BranchConfig, UnrealTargetPlatform InHostPlatform, BranchInfo.BranchUProject InGameProj, string InTestName, string InUATCommandLine, string InAgentSharingGroup, bool InDependsOnEditor = true, List<UnrealTargetPlatform> InDependsOnCooked = null, float InECPriority = 0.0f)
             : base(InHostPlatform)
         {
             AgentSharingGroup = InAgentSharingGroup;
@@ -3088,7 +3105,7 @@ partial class GUBP
             if (DependsOnEditor)
             {
                 AddDependency(EditorAndToolsNode.StaticGetFullName(HostPlatform));
-                if (GameProj.GameName != bp.Branch.BaseEngineProject.GameName)
+                if (GameProj.GameName != BranchConfig.Branch.BaseEngineProject.GameName)
                 {
                     if (GameProj.Properties.Targets.ContainsKey(TargetRules.TargetType.Editor))
                     {
@@ -3103,11 +3120,11 @@ partial class GUBP
            
             AddPseudodependency(WaitForTestShared.StaticGetFullName());
             // If the same test fails for the base engine, don't bother trying
-            if (InGameProj.GameName != bp.Branch.BaseEngineProject.GameName)
+            if (InGameProj.GameName != BranchConfig.Branch.BaseEngineProject.GameName)
             {
-                if (BranchConfig.HasNode(UATTestNode.StaticGetFullName(HostPlatform, bp.Branch.BaseEngineProject, TestName)))
+                if (BranchConfig.HasNode(UATTestNode.StaticGetFullName(HostPlatform, BranchConfig.Branch.BaseEngineProject, TestName)))
                 {
-                    AddPseudodependency(UATTestNode.StaticGetFullName(HostPlatform, bp.Branch.BaseEngineProject, InTestName));
+                    AddPseudodependency(UATTestNode.StaticGetFullName(HostPlatform, BranchConfig.Branch.BaseEngineProject, InTestName));
                 }
                 else
                 {
@@ -3115,22 +3132,22 @@ partial class GUBP
                     foreach (var Plat in DependsOnCooked)
                     {
                         var PlatTestName = "CookedGameTest_"  + Plat.ToString();
-                        if (BranchConfig.HasNode(UATTestNode.StaticGetFullName(HostPlatform, bp.Branch.BaseEngineProject, PlatTestName)))
+                        if (BranchConfig.HasNode(UATTestNode.StaticGetFullName(HostPlatform, BranchConfig.Branch.BaseEngineProject, PlatTestName)))
                         {
-                            AddPseudodependency(UATTestNode.StaticGetFullName(HostPlatform, bp.Branch.BaseEngineProject, PlatTestName));
+                            AddPseudodependency(UATTestNode.StaticGetFullName(HostPlatform, BranchConfig.Branch.BaseEngineProject, PlatTestName));
                             bFoundACook = true;
                         }
                     }
 
-                    if (!bFoundACook && BranchConfig.HasNode(UATTestNode.StaticGetFullName(HostPlatform, bp.Branch.BaseEngineProject, "EditorTest")))
+                    if (!bFoundACook && BranchConfig.HasNode(UATTestNode.StaticGetFullName(HostPlatform, BranchConfig.Branch.BaseEngineProject, "EditorTest")))
                     {
-                        AddPseudodependency(UATTestNode.StaticGetFullName(HostPlatform, bp.Branch.BaseEngineProject, "EditorTest"));
+                        AddPseudodependency(UATTestNode.StaticGetFullName(HostPlatform, BranchConfig.Branch.BaseEngineProject, "EditorTest"));
                     }
 
                 }
             }
 
-            if (InGameProj.GameName == bp.Branch.BaseEngineProject.GameName)
+            if (InGameProj.GameName == BranchConfig.Branch.BaseEngineProject.GameName)
             {
                 ECPriority = ECPriority + 1.0f;
             }
@@ -3156,9 +3173,9 @@ partial class GUBP
 		{
 			return true;
 		}
-        public override int CISFrequencyQuantumShift(GUBP bp)
+		public override int CISFrequencyQuantumShift(GUBP.GUBPBranchConfig BranchConfig)
         {
-            return base.CISFrequencyQuantumShift(bp) + 5;
+            return base.CISFrequencyQuantumShift(BranchConfig) + 5;
         }
         public override string ECAgentString()
         {
@@ -3210,7 +3227,7 @@ partial class GUBP
         BranchInfo.BranchUProject GameProj;
         string AggregateName;
 
-        public GameAggregateNode(GUBP bp, UnrealTargetPlatform InHostPlatform, BranchInfo.BranchUProject InGameProj, string InAggregateName, List<string> Dependencies)
+        public GameAggregateNode(UnrealTargetPlatform InHostPlatform, BranchInfo.BranchUProject InGameProj, string InAggregateName, List<string> Dependencies)
             : base(InHostPlatform)
         {
             GameProj = InGameProj;
@@ -3274,9 +3291,9 @@ partial class GUBP
             BuildProducts = new List<string>();
             SaveRecordOfSuccessAndAddToBuildProducts();
         }
-        public override int CISFrequencyQuantumShift(GUBP bp)
+		public override int CISFrequencyQuantumShift(GUBP.GUBPBranchConfig BranchConfig)
         {
-            return base.CISFrequencyQuantumShift(bp) + 3;
+            return base.CISFrequencyQuantumShift(BranchConfig) + 3;
         }
     };
 }
diff --git a/Engine/Source/Programs/AutomationTool/Scripts/BuildDocumentation.Automation.cs b/Engine/Source/Programs/AutomationTool/Scripts/BuildDocumentation.Automation.cs
--- a/Engine/Source/Programs/AutomationTool/Scripts/BuildDocumentation.Automation.cs
+++ b/Engine/Source/Programs/AutomationTool/Scripts/BuildDocumentation.Automation.cs
@@ -10,8 +10,8 @@ using UnrealBuildTool;
 
 class ToolsForDocumentationNode : GUBP.CompileNode
 {
-    public ToolsForDocumentationNode(UnrealTargetPlatform InHostPlatform)
-        : base(InHostPlatform, false)
+    public ToolsForDocumentationNode(GUBP.GUBPBranchConfig InBranchConfig, UnrealTargetPlatform InHostPlatform)
+        : base(InBranchConfig, InHostPlatform, false)
     {
 		AgentSharingGroup = "Documentation" + StaticGetHostPlatformSuffix(InHostPlatform);
     }
@@ -26,9 +26,9 @@ class ToolsForDocumentationNode : GUBP.CompileNode
         return StaticGetFullName(HostPlatform);
     }
 
-	public override int CISFrequencyQuantumShift(GUBP bp)
+	public override int CISFrequencyQuantumShift(GUBP.GUBPBranchConfig BranchConfig)
 	{
-		return base.CISFrequencyQuantumShift(bp) + 3;
+		return base.CISFrequencyQuantumShift(BranchConfig) + 3;
 	}
 
 	public override UE4Build.BuildAgenda GetAgenda(GUBP bp)
@@ -48,9 +48,9 @@ class DocumentationNode : GUBP.GUBPNode
 		AddDependency(ToolsForDocumentationNode.StaticGetFullName(InHostPlatform));
 	}
 
-	public override int CISFrequencyQuantumShift(GUBP bp)
+	public override int CISFrequencyQuantumShift(GUBP.GUBPBranchConfig BranchConfig)
 	{
-		return base.CISFrequencyQuantumShift(bp) + 3;
+		return base.CISFrequencyQuantumShift(BranchConfig) + 3;
 	}
 
 	protected void ExecuteApiDocTool(string Arguments, string LogName)
@@ -138,9 +138,9 @@ public class DocumentationNodeAdder : GUBP.GUBPNodeAdder
 {
 	public override void AddNodes(GUBP bp, GUBP.GUBPBranchConfig BranchConfig, UnrealTargetPlatform InHostPlatform, List<UnrealTargetPlatform> InActivePlatforms)
 	{
-		if(InHostPlatform == UnrealTargetPlatform.Win64 && !bp.BranchOptions.bNoDocumentation)
+		if(InHostPlatform == UnrealTargetPlatform.Win64 && !BranchConfig.BranchOptions.bNoDocumentation)
 		{
-			BranchConfig.AddNode(new ToolsForDocumentationNode(InHostPlatform));
+			BranchConfig.AddNode(new ToolsForDocumentationNode(BranchConfig, InHostPlatform));
 			BranchConfig.AddNode(new CodeDocumentationNode(InHostPlatform));
 			BranchConfig.AddNode(new BlueprintDocumentationNode(InHostPlatform));
 		}
diff --git a/Engine/Source/Programs/AutomationTool/Scripts/RocketBuild.Automation.cs b/Engine/Source/Programs/AutomationTool/Scripts/RocketBuild.Automation.cs
--- a/Engine/Source/Programs/AutomationTool/Scripts/RocketBuild.Automation.cs
+++ b/Engine/Source/Programs/AutomationTool/Scripts/RocketBuild.Automation.cs
@@ -76,7 +76,7 @@ namespace Rocket
 
 		public override void AddNodes(GUBP bp, GUBP.GUBPBranchConfig BranchConfig, UnrealTargetPlatform HostPlatform, List<UnrealTargetPlatform> ActivePlatforms)
 		{
-			if(!bp.BranchOptions.bNoInstalledEngine)
+			if (!BranchConfig.BranchOptions.bNoInstalledEngine)
 			{
 				// Find all the target platforms for this host platform.
 				List<UnrealTargetPlatform> TargetPlatforms = GetTargetPlatforms(bp, HostPlatform);
@@ -97,17 +97,17 @@ namespace Rocket
 				// Strip all the target platforms that are built on this host
 				foreach (UnrealTargetPlatform TargetPlatform in TargetPlatforms)
 				{
-					if (GetSourceHostPlatform(bp, HostPlatform, TargetPlatform) == HostPlatform && StripRocketNode.IsRequiredForPlatform(TargetPlatform))
+					if (GetSourceHostPlatform(BranchConfig.HostPlatforms, HostPlatform, TargetPlatform) == HostPlatform && StripRocketNode.IsRequiredForPlatform(TargetPlatform))
 					{
-						BranchConfig.AddNode(new StripRocketMonolithicsNode(bp, BranchConfig, HostPlatform, TargetPlatform, StrippedDir));
+						BranchConfig.AddNode(new StripRocketMonolithicsNode(BranchConfig, HostPlatform, TargetPlatform, StrippedDir));
 					}
 				}
 
 				// Build the DDC
-				BranchConfig.AddNode(new BuildDerivedDataCacheNode(HostPlatform, GetCookPlatforms(HostPlatform, TargetPlatforms), CurrentFeaturePacks));
+				BranchConfig.AddNode(new BuildDerivedDataCacheNode(BranchConfig, HostPlatform, GetCookPlatforms(HostPlatform, TargetPlatforms), CurrentFeaturePacks));
 
 				// Generate a list of files that needs to be copied for each target platform
-				BranchConfig.AddNode(new FilterRocketNode(bp, BranchConfig, HostPlatform, TargetPlatforms, CurrentFeaturePacks, CurrentTemplates));
+				BranchConfig.AddNode(new FilterRocketNode(BranchConfig, HostPlatform, TargetPlatforms, CurrentFeaturePacks, CurrentTemplates));
 
 				// Copy the install to the output directory
 				string LocalOutputDir = CommandUtils.CombinePaths(CommandUtils.CmdEnv.LocalRoot, "LocalBuilds", "Rocket", CommandUtils.GetGenericPlatformName(HostPlatform));
@@ -124,14 +124,14 @@ namespace Rocket
 					string GitConfigRelativePath = "Engine/Build/Git/UnrealBot.ini";
 					if(CommandUtils.FileExists(CommandUtils.CombinePaths(CommandUtils.CmdEnv.LocalRoot, GitConfigRelativePath)))
 					{
-						BranchConfig.AddNode(new BuildGitPromotable(bp, HostPlatform, GitConfigRelativePath));
+						BranchConfig.AddNode(new BuildGitPromotable(HostPlatform, BranchConfig.HostPlatforms, GitConfigRelativePath));
 						PromotableNode.AddDependency(BuildGitPromotable.StaticGetFullName(HostPlatform));
 					}
 				}
 
 				// Get the output directory for the build zips
 				string PublishedEngineDir;
-				if(ShouldDoSeriousThingsLikeP4CheckinAndPostToMCP(bp))
+				if (ShouldDoSeriousThingsLikeP4CheckinAndPostToMCP(BranchConfig))
 				{
 					PublishedEngineDir = CommandUtils.CombinePaths(CommandUtils.RootSharedTempStorageDirectory(), "Rocket", "Automated", GetBuildLabel(), CommandUtils.GetGenericPlatformName(HostPlatform));
 				}
@@ -142,7 +142,7 @@ namespace Rocket
 
 				// Publish the install to the network
 				BranchConfig.AddNode(new PublishRocketNode(HostPlatform, LocalOutputDir, PublishedEngineDir));
-				BranchConfig.AddNode(new PublishRocketSymbolsNode(bp, HostPlatform, TargetPlatforms, PublishedEngineDir + "Symbols"));
+				BranchConfig.AddNode(new PublishRocketSymbolsNode(BranchConfig, HostPlatform, TargetPlatforms, PublishedEngineDir + "Symbols"));
 
 				// Add a dependency on this being published as part of the shared promotable being labeled
 				GUBP.SharedLabelPromotableSuccessNode LabelPromotableNode = (GUBP.SharedLabelPromotableSuccessNode)BranchConfig.FindAggregateNode(GUBP.SharedLabelPromotableSuccessNode.StaticGetFullName());
@@ -155,7 +155,7 @@ namespace Rocket
 				WaitForPromotionNode.AddDependency(PublishRocketSymbolsNode.StaticGetFullName(HostPlatform));
 
 				// Push everything behind the promotion triggers if we're doing things on the build machines
-				if(ShouldDoSeriousThingsLikeP4CheckinAndPostToMCP(bp) || bp.ParseParam("WithRocketPromotable"))
+				if (ShouldDoSeriousThingsLikeP4CheckinAndPostToMCP(BranchConfig) || bp.ParseParam("WithRocketPromotable"))
 				{
 					string WaitForTrigger = GUBP.WaitForSharedPromotionUserInput.StaticGetFullName(false);
 
@@ -247,22 +247,22 @@ namespace Rocket
 			return CommandUtils.CombineCommandletParams(CookPlatforms.Distinct().ToArray());
 		}
 
-		public static bool ShouldDoSeriousThingsLikeP4CheckinAndPostToMCP(GUBP bp)
+		public static bool ShouldDoSeriousThingsLikeP4CheckinAndPostToMCP(GUBP.GUBPBranchConfig BranchConfig)
 		{
-			return CommandUtils.P4Enabled && CommandUtils.AllowSubmit && !bp.bPreflightBuild; // we don't do serious things in a preflight
+			return CommandUtils.P4Enabled && CommandUtils.AllowSubmit && !BranchConfig.bPreflightBuild; // we don't do serious things in a preflight
 		}
 
-		public static UnrealTargetPlatform GetSourceHostPlatform(GUBP bp, UnrealTargetPlatform HostPlatform, UnrealTargetPlatform TargetPlatform)
+		public static UnrealTargetPlatform GetSourceHostPlatform(List<UnrealTargetPlatform> HostPlatforms, UnrealTargetPlatform HostPlatform, UnrealTargetPlatform TargetPlatform)
 		{
-            if (TargetPlatform == UnrealTargetPlatform.HTML5 && HostPlatform == UnrealTargetPlatform.Mac && bp.HostPlatforms.Contains(UnrealTargetPlatform.Win64))
+            if (TargetPlatform == UnrealTargetPlatform.HTML5 && HostPlatform == UnrealTargetPlatform.Mac && HostPlatforms.Contains(UnrealTargetPlatform.Win64))
             {
                 return UnrealTargetPlatform.Win64;
             }
-			if (TargetPlatform == UnrealTargetPlatform.Android && HostPlatform == UnrealTargetPlatform.Mac && bp.HostPlatforms.Contains(UnrealTargetPlatform.Win64))
+			if (TargetPlatform == UnrealTargetPlatform.Android && HostPlatform == UnrealTargetPlatform.Mac && HostPlatforms.Contains(UnrealTargetPlatform.Win64))
 			{
 				return UnrealTargetPlatform.Win64;
 			}
-			if(TargetPlatform == UnrealTargetPlatform.IOS && HostPlatform == UnrealTargetPlatform.Win64 && bp.HostPlatforms.Contains(UnrealTargetPlatform.Mac))
+			if(TargetPlatform == UnrealTargetPlatform.IOS && HostPlatform == UnrealTargetPlatform.Win64 && HostPlatforms.Contains(UnrealTargetPlatform.Mac))
 			{
 				return UnrealTargetPlatform.Mac;
 			}
@@ -287,15 +287,15 @@ namespace Rocket
 	{
 		string ConfigRelativePath;
 
-		public BuildGitPromotable(GUBP bp, UnrealTargetPlatform HostPlatform, string InConfigRelativePath) : base(HostPlatform)
+		public BuildGitPromotable(UnrealTargetPlatform HostPlatform, List<UnrealTargetPlatform> ForHostPlatforms, string InConfigRelativePath) : base(HostPlatform)
 		{
 			ConfigRelativePath = InConfigRelativePath;
 
-			foreach(UnrealTargetPlatform OtherHostPlatform in bp.HostPlatforms)
+			foreach(UnrealTargetPlatform ForHostPlatform in ForHostPlatforms)
 			{
-				AddDependency(GUBP.RootEditorNode.StaticGetFullName(OtherHostPlatform));
-				AddDependency(GUBP.ToolsNode.StaticGetFullName(OtherHostPlatform));
-				AddDependency(GUBP.InternalToolsNode.StaticGetFullName(OtherHostPlatform));
+				AddDependency(GUBP.RootEditorNode.StaticGetFullName(ForHostPlatform));
+				AddDependency(GUBP.ToolsNode.StaticGetFullName(ForHostPlatform));
+				AddDependency(GUBP.InternalToolsNode.StaticGetFullName(ForHostPlatform));
 			}
 		}
 
@@ -451,15 +451,15 @@ namespace Rocket
 		BranchInfo.BranchUProject Project;
 		bool bIsCodeTargetPlatform;
 
-		public StripRocketMonolithicsNode(GUBP bp, GUBP.GUBPBranchConfig InBranchConfig, UnrealTargetPlatform InHostPlatform, UnrealTargetPlatform InTargetPlatform, string InStrippedDir) : base(InBranchConfig, InHostPlatform, InTargetPlatform, InStrippedDir)
+		public StripRocketMonolithicsNode(GUBP.GUBPBranchConfig InBranchConfig, UnrealTargetPlatform InHostPlatform, UnrealTargetPlatform InTargetPlatform, string InStrippedDir) : base(InBranchConfig, InHostPlatform, InTargetPlatform, InStrippedDir)
 		{
-			Project = bp.Branch.BaseEngineProject;
+			Project = InBranchConfig.Branch.BaseEngineProject;
 			bIsCodeTargetPlatform = RocketBuild.IsCodeTargetPlatform(InHostPlatform, InTargetPlatform);
 
 			GUBP.GUBPNode Node = InBranchConfig.FindNode(GUBP.GamePlatformMonolithicsNode.StaticGetFullName(HostPlatform, Project, InTargetPlatform, Precompiled: bIsCodeTargetPlatform));
 			if(String.IsNullOrEmpty(Node.AgentSharingGroup))
 			{
-				Node.AgentSharingGroup = bp.Branch.BaseEngineProject.GameName + "_MonolithicsGroup_" + InTargetPlatform + StaticGetHostPlatformSuffix(InHostPlatform);
+				Node.AgentSharingGroup = BranchConfig.Branch.BaseEngineProject.GameName + "_MonolithicsGroup_" + InTargetPlatform + StaticGetHostPlatformSuffix(InHostPlatform);
 			}
 			AddNodeToStrip(Node.GetFullName());
 
@@ -497,7 +497,7 @@ namespace Rocket
 		public readonly string DepotManifestPath;
 		public Dictionary<string, string> StrippedNodeManifestPaths = new Dictionary<string, string>();
 
-		public FilterRocketNode(GUBP bp, GUBP.GUBPBranchConfig InBranchConfig, UnrealTargetPlatform InHostPlatform, List<UnrealTargetPlatform> InTargetPlatforms, string[] InCurrentFeaturePacks, string[] InCurrentTemplates)
+		public FilterRocketNode(GUBP.GUBPBranchConfig InBranchConfig, UnrealTargetPlatform InHostPlatform, List<UnrealTargetPlatform> InTargetPlatforms, string[] InCurrentFeaturePacks, string[] InCurrentTemplates)
 			: base(InHostPlatform)
 		{
 			BranchConfig = InBranchConfig;
@@ -515,9 +515,9 @@ namespace Rocket
 			// Add all the monolithic builds from their appropriate source host platform
 			foreach(UnrealTargetPlatform TargetPlatform in TargetPlatforms)
 			{
-				UnrealTargetPlatform SourceHostPlatform = RocketBuild.GetSourceHostPlatform(bp, HostPlatform, TargetPlatform);
+				UnrealTargetPlatform SourceHostPlatform = RocketBuild.GetSourceHostPlatform(BranchConfig.HostPlatforms, HostPlatform, TargetPlatform);
 				bool bIsCodeTargetPlatform = RocketBuild.IsCodeTargetPlatform(SourceHostPlatform, TargetPlatform);
-				AddDependency(GUBP.GamePlatformMonolithicsNode.StaticGetFullName(SourceHostPlatform, bp.Branch.BaseEngineProject, TargetPlatform, Precompiled: bIsCodeTargetPlatform));
+				AddDependency(GUBP.GamePlatformMonolithicsNode.StaticGetFullName(SourceHostPlatform, BranchConfig.Branch.BaseEngineProject, TargetPlatform, Precompiled: bIsCodeTargetPlatform));
 			}
 
 			// Also add stripped symbols for all the target platforms that require it
@@ -526,25 +526,25 @@ namespace Rocket
 			{
 				if(StripRocketNode.IsRequiredForPlatform(TargetPlatform))
 				{
-					UnrealTargetPlatform SourceHostPlatform = RocketBuild.GetSourceHostPlatform(bp, HostPlatform, TargetPlatform);
-					string StripNode = StripRocketMonolithicsNode.StaticGetFullName(SourceHostPlatform, bp.Branch.BaseEngineProject, TargetPlatform, RocketBuild.IsCodeTargetPlatform(SourceHostPlatform, TargetPlatform));
+					UnrealTargetPlatform SourceHostPlatform = RocketBuild.GetSourceHostPlatform(BranchConfig.HostPlatforms, HostPlatform, TargetPlatform);
+					string StripNode = StripRocketMonolithicsNode.StaticGetFullName(SourceHostPlatform, BranchConfig.Branch.BaseEngineProject, TargetPlatform, RocketBuild.IsCodeTargetPlatform(SourceHostPlatform, TargetPlatform));
 					AddDependency(StripNode);
 					StrippedNodeNames.Add(StripNode);
 				}
 			}
 
 			// Add win64 tools on Mac, to get the win64 build of UBT, UAT and IPP
-			if (HostPlatform == UnrealTargetPlatform.Mac && bp.HostPlatforms.Contains(UnrealTargetPlatform.Win64))
+			if (HostPlatform == UnrealTargetPlatform.Mac && BranchConfig.HostPlatforms.Contains(UnrealTargetPlatform.Win64))
 			{
 				AddDependency(GUBP.ToolsNode.StaticGetFullName(UnrealTargetPlatform.Win64));
 				AddDependency(GUBP.ToolsForCompileNode.StaticGetFullName(UnrealTargetPlatform.Win64));
 			}
 
 			// Add all the feature packs
-			AddDependency(GUBP.MakeFeaturePacksNode.StaticGetFullName(GUBP.MakeFeaturePacksNode.GetDefaultBuildPlatform(bp)));
+			AddDependency(GUBP.MakeFeaturePacksNode.StaticGetFullName(GUBP.MakeFeaturePacksNode.GetDefaultBuildPlatform(BranchConfig.HostPlatforms)));
 
 			// Find all the host platforms we need
-			SourceHostPlatforms = TargetPlatforms.Select(x => RocketBuild.GetSourceHostPlatform(bp, HostPlatform, x)).Distinct().ToList();
+			SourceHostPlatforms = TargetPlatforms.Select(x => RocketBuild.GetSourceHostPlatform(BranchConfig.HostPlatforms, HostPlatform, x)).Distinct().ToList();
 			if(!SourceHostPlatforms.Contains(HostPlatform))
 			{
 				SourceHostPlatforms.Add(HostPlatform);
@@ -567,9 +567,9 @@ namespace Rocket
 			}
 		}
 
-		public override int CISFrequencyQuantumShift(GUBP bp)
+		public override int CISFrequencyQuantumShift(GUBP.GUBPBranchConfig BranchConfig)
 		{
-			return base.CISFrequencyQuantumShift(bp) + 2;
+			return base.CISFrequencyQuantumShift(BranchConfig) + 2;
 		}
 
 		public static string StaticGetFullName(UnrealTargetPlatform InHostPlatform)
@@ -593,7 +593,7 @@ namespace Rocket
 			AddRuleForBuildProducts(Filter, BranchConfig, GUBP.ToolsNode.StaticGetFullName(HostPlatform), FileFilterType.Include);
 
 			// Include win64 tools on Mac, to get the win64 build of UBT, UAT and IPP
-			if (HostPlatform == UnrealTargetPlatform.Mac && bp.HostPlatforms.Contains(UnrealTargetPlatform.Win64))
+			if (HostPlatform == UnrealTargetPlatform.Mac && BranchConfig.HostPlatforms.Contains(UnrealTargetPlatform.Win64))
 			{
 				AddRuleForBuildProducts(Filter, BranchConfig, GUBP.ToolsNode.StaticGetFullName(UnrealTargetPlatform.Win64), FileFilterType.Include);
 				AddRuleForBuildProducts(Filter, BranchConfig, GUBP.ToolsForCompileNode.StaticGetFullName(UnrealTargetPlatform.Win64), FileFilterType.Include);
@@ -607,32 +607,32 @@ namespace Rocket
 			{
 				if(RocketBuild.IsCodeTargetPlatform(HostPlatform, TargetPlatform))
 				{
-					UnrealTargetPlatform SourceHostPlatform = RocketBuild.GetSourceHostPlatform(bp, HostPlatform, TargetPlatform);
-					string FileListPath = GUBP.GamePlatformMonolithicsNode.StaticGetBuildDependenciesPath(SourceHostPlatform, bp.Branch.BaseEngineProject, TargetPlatform);
+					UnrealTargetPlatform SourceHostPlatform = RocketBuild.GetSourceHostPlatform(BranchConfig.HostPlatforms, HostPlatform, TargetPlatform);
+					string FileListPath = GUBP.GamePlatformMonolithicsNode.StaticGetBuildDependenciesPath(SourceHostPlatform, BranchConfig.Branch.BaseEngineProject, TargetPlatform);
 					Filter.AddRuleForFiles(UnrealBuildTool.Utils.ReadClass<UnrealBuildTool.ExternalFileList>(FileListPath).FileNames, CommandUtils.CmdEnv.LocalRoot, FileFilterType.Include);
-					UnzipAndAddRuleForHeaders(GUBP.GamePlatformMonolithicsNode.StaticGetArchivedHeadersPath(SourceHostPlatform, bp.Branch.BaseEngineProject, TargetPlatform), Filter, FileFilterType.Include);
+					UnzipAndAddRuleForHeaders(GUBP.GamePlatformMonolithicsNode.StaticGetArchivedHeadersPath(SourceHostPlatform, BranchConfig.Branch.BaseEngineProject, TargetPlatform), Filter, FileFilterType.Include);
 				}
 			}
 
 			// Add the monolithic binaries
 			foreach(UnrealTargetPlatform TargetPlatform in TargetPlatforms)
 			{
-				UnrealTargetPlatform SourceHostPlatform = RocketBuild.GetSourceHostPlatform(bp, HostPlatform, TargetPlatform);
+				UnrealTargetPlatform SourceHostPlatform = RocketBuild.GetSourceHostPlatform(BranchConfig.HostPlatforms, HostPlatform, TargetPlatform);
 				bool bIsCodeTargetPlatform = RocketBuild.IsCodeTargetPlatform(SourceHostPlatform, TargetPlatform);
-				AddRuleForBuildProducts(Filter, BranchConfig, GUBP.GamePlatformMonolithicsNode.StaticGetFullName(SourceHostPlatform, bp.Branch.BaseEngineProject, TargetPlatform, Precompiled: bIsCodeTargetPlatform), FileFilterType.Include);
+				AddRuleForBuildProducts(Filter, BranchConfig, GUBP.GamePlatformMonolithicsNode.StaticGetFullName(SourceHostPlatform, BranchConfig.Branch.BaseEngineProject, TargetPlatform, Precompiled: bIsCodeTargetPlatform), FileFilterType.Include);
 			}
 
 			// Include the feature packs
 			foreach(string CurrentFeaturePack in CurrentFeaturePacks)
 			{
-				BranchInfo.BranchUProject Project = bp.Branch.FindGameChecked(CurrentFeaturePack);
+				BranchInfo.BranchUProject Project = BranchConfig.Branch.FindGameChecked(CurrentFeaturePack);
 				Filter.AddRuleForFile(GUBP.MakeFeaturePacksNode.GetOutputFile(Project), CommandUtils.CmdEnv.LocalRoot, FileFilterType.Include);
 			}
 
 			// Include all the templates
 			foreach (string Template in CurrentTemplates)
 			{
-				BranchInfo.BranchUProject Project = bp.Branch.FindGameChecked(Template);
+				BranchInfo.BranchUProject Project = BranchConfig.Branch.FindGameChecked(Template);
 				Filter.Include("/" + Utils.StripBaseDirectory(Path.GetDirectoryName(Project.FilePath), CommandUtils.CmdEnv.LocalRoot).Replace('\\', '/') + "/...");
 			}
 
@@ -762,8 +762,8 @@ namespace Rocket
 			// Extract all the headers for code target platforms
 			foreach(UnrealTargetPlatform CodeTargetPlatform in CodeTargetPlatforms)
 			{
-				UnrealTargetPlatform SourceHostPlatform = RocketBuild.GetSourceHostPlatform(bp, HostPlatform, CodeTargetPlatform);
-				string ZipFileName = GUBP.GamePlatformMonolithicsNode.StaticGetArchivedHeadersPath(SourceHostPlatform, bp.Branch.BaseEngineProject, CodeTargetPlatform);
+				UnrealTargetPlatform SourceHostPlatform = RocketBuild.GetSourceHostPlatform(BranchConfig.HostPlatforms, HostPlatform, CodeTargetPlatform);
+				string ZipFileName = GUBP.GamePlatformMonolithicsNode.StaticGetArchivedHeadersPath(SourceHostPlatform, BranchConfig.Branch.BaseEngineProject, CodeTargetPlatform);
 				CommandUtils.UnzipFiles(ZipFileName, CommandUtils.CmdEnv.LocalRoot);
 			}
 
@@ -874,10 +874,12 @@ namespace Rocket
 
 	public class PublishRocketSymbolsNode : GUBP.HostPlatformNode
 	{
+		GUBP.GUBPBranchConfig BranchConfig;
 		string SymbolsOutputDir;
 
-		public PublishRocketSymbolsNode(GUBP bp, UnrealTargetPlatform HostPlatform, IEnumerable<UnrealTargetPlatform> TargetPlatforms, string InSymbolsOutputDir) : base(HostPlatform)
+		public PublishRocketSymbolsNode(GUBP.GUBPBranchConfig InBranchConfig, UnrealTargetPlatform HostPlatform, IEnumerable<UnrealTargetPlatform> TargetPlatforms, string InSymbolsOutputDir) : base(HostPlatform)
 		{
+			BranchConfig = InBranchConfig;
 			SymbolsOutputDir = InSymbolsOutputDir;
 
 			AddDependency(GUBP.WaitForSharedPromotionUserInput.StaticGetFullName(false));
@@ -887,10 +889,10 @@ namespace Rocket
 
 			foreach(UnrealTargetPlatform TargetPlatform in TargetPlatforms)
 			{
-				if(HostPlatform == RocketBuild.GetSourceHostPlatform(bp, HostPlatform, TargetPlatform))
+				if(HostPlatform == RocketBuild.GetSourceHostPlatform(BranchConfig.HostPlatforms, HostPlatform, TargetPlatform))
 				{
 					bool bIsCodeTargetPlatform = RocketBuild.IsCodeTargetPlatform(HostPlatform, TargetPlatform);
-					AddDependency(GUBP.GamePlatformMonolithicsNode.StaticGetFullName(HostPlatform, bp.Branch.BaseEngineProject, TargetPlatform, Precompiled: bIsCodeTargetPlatform));
+					AddDependency(GUBP.GamePlatformMonolithicsNode.StaticGetFullName(HostPlatform, BranchConfig.Branch.BaseEngineProject, TargetPlatform, Precompiled: bIsCodeTargetPlatform));
 				}
 			}
 
@@ -909,7 +911,7 @@ namespace Rocket
 
 		public override void DoBuild(GUBP bp)
 		{
-			if(RocketBuild.ShouldDoSeriousThingsLikeP4CheckinAndPostToMCP(bp))
+			if (RocketBuild.ShouldDoSeriousThingsLikeP4CheckinAndPostToMCP(BranchConfig))
 			{
 				// Make a lookup for all the known debug extensions, and filter all the dependency build products against that
 				HashSet<string> DebugExtensions = new HashSet<string>(Platform.Platforms.Values.SelectMany(x => x.GetDebugFileExtentions()).Distinct().ToArray(), StringComparer.InvariantCultureIgnoreCase);
@@ -932,14 +934,16 @@ namespace Rocket
 
 	public class BuildDerivedDataCacheNode : GUBP.HostPlatformNode
 	{
+		GUBP.GUBPBranchConfig BranchConfig;
 		string TargetPlatforms;
 		string[] ProjectNames;
 		public readonly string SavedDir;
 		public readonly string SavedManifestPath;
 
-		public BuildDerivedDataCacheNode(UnrealTargetPlatform InHostPlatform, string InTargetPlatforms, string[] InProjectNames)
+		public BuildDerivedDataCacheNode(GUBP.GUBPBranchConfig InBranchConfig, UnrealTargetPlatform InHostPlatform, string InTargetPlatforms, string[] InProjectNames)
 			: base(InHostPlatform)
 		{
+			BranchConfig = InBranchConfig;
 			TargetPlatforms = InTargetPlatforms;
 			ProjectNames = InProjectNames;
 			SavedDir = CommandUtils.CombinePaths(CommandUtils.CmdEnv.LocalRoot, "Engine", "Saved", "Rocket", HostPlatform.ToString());
@@ -954,9 +958,9 @@ namespace Rocket
             return "BuildDerivedDataCache" + StaticGetHostPlatformSuffix(InHostPlatform);
 		}
 
-		public override int CISFrequencyQuantumShift(GUBP bp)
+		public override int CISFrequencyQuantumShift(GUBP.GUBPBranchConfig BranchConfig)
 		{
-			return base.CISFrequencyQuantumShift(bp) + 2;
+			return base.CISFrequencyQuantumShift(BranchConfig) + 2;
 		}
 
 		public override string GetFullName()
@@ -977,7 +981,7 @@ namespace Rocket
 				List<BranchInfo.BranchUProject> Projects = new List<BranchInfo.BranchUProject>();
 				foreach(string ProjectName in ProjectNames)
 				{
-					BranchInfo.BranchUProject Project = bp.Branch.FindGameChecked(ProjectName);
+					BranchInfo.BranchUProject Project = BranchConfig.Branch.FindGameChecked(ProjectName);
 					if(!Project.Properties.bIsCodeBasedProject)
 					{
 						Projects.Add(Project);
