Move BuildNode into a subfolder.

diff --git a/Engine/Source/Programs/AutomationTool/GUBP/BuildNode.cs b/Engine/Source/Programs/AutomationTool/GUBP/BuildNode.cs
deleted file mode 100644
diff --git a/Engine/Source/Programs/AutomationTool/GUBP/GUBP.Automation.csproj b/Engine/Source/Programs/AutomationTool/GUBP/GUBP.Automation.csproj
--- a/Engine/Source/Programs/AutomationTool/GUBP/GUBP.Automation.csproj
+++ b/Engine/Source/Programs/AutomationTool/GUBP/GUBP.Automation.csproj
@@ -58,11 +58,11 @@
     <Compile Include="..\..\DotNETCommon\MetaData.cs">
       <Link>Properties\MetaData.cs</Link>
     </Compile>
-    <Compile Include="BuildNode.cs" />
     <Compile Include="ElectricCommander.cs" />
     <Compile Include="GUBP.cs" />
     <Compile Include="LegacyBranchSetup.cs" />
     <Compile Include="LegacyNodes.cs" />
+    <Compile Include="Nodes\BuildNode.cs" />
     <Compile Include="TempStorage.cs">
       <SubType>Code</SubType>
     </Compile>
@@ -79,8 +79,6 @@
       <Private>False</Private>
     </ProjectReference>
   </ItemGroup>
-  <ItemGroup>
-    <Folder Include="Nodes\" />
-  </ItemGroup>
+  <ItemGroup />
   <Import Project="$(MSBuildToolsPath)\Microsoft.CSharp.targets" />
 </Project>
\ No newline at end of file
diff --git a/Engine/Source/Programs/AutomationTool/GUBP/Nodes/BuildNode.cs b/Engine/Source/Programs/AutomationTool/GUBP/Nodes/BuildNode.cs
new file mode 100644
--- /dev/null
+++ b/Engine/Source/Programs/AutomationTool/GUBP/Nodes/BuildNode.cs
@@ -0,0 +1,41 @@
+﻿using System;
+using System.Collections.Generic;
+using System.Diagnostics;
+using System.Linq;
+using System.Text;
+
+namespace AutomationTool
+{
+	[DebuggerDisplay("{Name}")]
+	class NodeInfo
+	{
+		public string Name;
+		public GUBP.GUBPNode Node;
+		public NodeInfo[] Dependencies;
+		public NodeInfo[] PseudoDependencies;
+		public NodeInfo[] AllDirectDependencies;
+		public NodeInfo[] AllIndirectDependencies;
+		public NodeInfo[] ControllingTriggers;
+		public int FrequencyShift;
+		public bool IsComplete;
+		public string[] RecipientsForFailureEmails;
+		public bool AddSubmittersToFailureEmails;
+		public int AgentMemoryRequirement;
+
+		public string ControllingTriggerDotName
+		{
+			get { return String.Join(".", ControllingTriggers.Select(x => x.Name)); }
+		}
+
+		public bool DependsOn(NodeInfo Node)
+		{
+			return AllIndirectDependencies.Contains(Node);
+		}
+
+		public override string ToString()
+		{
+			System.Diagnostics.Trace.TraceWarning("Implicit conversion from NodeInfo to string\n{0}", Environment.StackTrace);
+			return Name;
+		}
+	}
+}
