Move AggregateNode into its own file, and rename NodeInfo to BuildNode.

diff --git a/Engine/Source/Programs/AutomationTool/GUBP/GUBP.Automation.csproj b/Engine/Source/Programs/AutomationTool/GUBP/GUBP.Automation.csproj
--- a/Engine/Source/Programs/AutomationTool/GUBP/GUBP.Automation.csproj
+++ b/Engine/Source/Programs/AutomationTool/GUBP/GUBP.Automation.csproj
@@ -62,6 +62,7 @@
     <Compile Include="GUBP.cs" />
     <Compile Include="LegacyBranchSetup.cs" />
     <Compile Include="LegacyNodes.cs" />
+    <Compile Include="Nodes\AggregateNode.cs" />
     <Compile Include="Nodes\BuildNode.cs" />
     <Compile Include="TempStorage.cs">
       <SubType>Code</SubType>
diff --git a/Engine/Source/Programs/AutomationTool/GUBP/GUBP.cs b/Engine/Source/Programs/AutomationTool/GUBP/GUBP.cs
--- a/Engine/Source/Programs/AutomationTool/GUBP/GUBP.cs
+++ b/Engine/Source/Programs/AutomationTool/GUBP/GUBP.cs
@@ -17,14 +17,6 @@ public partial class GUBP : BuildCommand
 	const string FailedTempStorageSuffix = "_Failed";
 	const string SucceededTempStorageSuffix = "_Succeeded";
 
-	[DebuggerDisplay("{Name}")]
-	class AggregateInfo
-	{
-		public string Name;
-		public GUBPAggregateNode Node;
-		public NodeInfo[] Dependencies;
-	}
-
     class NodeHistory
     {
         public int LastSucceeded = 0;
@@ -42,9 +34,9 @@ public partial class GUBP : BuildCommand
 	/// Recursively update the ControllingTriggers array for each of the nodes passed in. 
 	/// </summary>
 	/// <param name="NodesToDo">Nodes to update</param>
-	static void FindControllingTriggers(IEnumerable<NodeInfo> NodesToDo)
+	static void FindControllingTriggers(IEnumerable<BuildNode> NodesToDo)
 	{
-		foreach(NodeInfo NodeToDo in NodesToDo)
+		foreach(BuildNode NodeToDo in NodesToDo)
 		{
 			FindControllingTriggers(NodeToDo);
 		}
@@ -54,24 +46,24 @@ public partial class GUBP : BuildCommand
 	/// Recursively find the controlling triggers for the given node.
 	/// </summary>
 	/// <param name="NodeToDo">Node to find the controlling triggers for</param>
-    static void FindControllingTriggers(NodeInfo NodeToDo)
+    static void FindControllingTriggers(BuildNode NodeToDo)
     {
 		if(NodeToDo.ControllingTriggers == null)
 		{
-			NodeToDo.ControllingTriggers = new NodeInfo[0];
+			NodeToDo.ControllingTriggers = new BuildNode[0];
 
 			// Find all the dependencies of this node
-			List<NodeInfo> AllDependencies = new List<NodeInfo>();
+			List<BuildNode> AllDependencies = new List<BuildNode>();
 			AllDependencies.AddRange(NodeToDo.Dependencies);
 			AllDependencies.AddRange(NodeToDo.PseudoDependencies);
 
 			// Find the immediate trigger controlling this one
-			List<NodeInfo> PreviousTriggers = new List<NodeInfo>();
-			foreach (NodeInfo Dependency in AllDependencies)
+			List<BuildNode> PreviousTriggers = new List<BuildNode>();
+			foreach (BuildNode Dependency in AllDependencies)
 			{
 				FindControllingTriggers(Dependency);
 
-				NodeInfo PreviousTrigger = null;
+				BuildNode PreviousTrigger = null;
 				if(Dependency.Node.TriggerNode())
 				{
 					PreviousTrigger = Dependency;
@@ -99,7 +91,7 @@ public partial class GUBP : BuildCommand
 			// Update the list of controlling triggers
 			if (PreviousTriggers.Count == 1)
 			{
-				List<NodeInfo> ControllingTriggers = new List<NodeInfo>();
+				List<BuildNode> ControllingTriggers = new List<BuildNode>();
 				ControllingTriggers.AddRange(PreviousTriggers[0].ControllingTriggers);
 				ControllingTriggers.Add(PreviousTriggers[0]);
 				NodeToDo.ControllingTriggers = ControllingTriggers.ToArray();
@@ -128,20 +120,20 @@ public partial class GUBP : BuildCommand
 	/// Update the frequency shift for all build nodes to ensure that each node's frequency is at least that of its dependencies.
 	/// </summary>
 	/// <param name="Nodes">Sequence of nodes to compute frequencies for</param>
-	static void ComputeDependentFrequencies(IEnumerable<NodeInfo> Nodes)
+	static void ComputeDependentFrequencies(IEnumerable<BuildNode> Nodes)
 	{
-		foreach(NodeInfo Node in Nodes)
+		foreach(BuildNode Node in Nodes)
 		{
-			foreach(NodeInfo IndirectDependency in Node.AllIndirectDependencies)
+			foreach(BuildNode IndirectDependency in Node.AllIndirectDependencies)
 			{
 				Node.FrequencyShift = Math.Max(Node.FrequencyShift, IndirectDependency.FrequencyShift);
 			}
 		}
 	}
 
-	static void FindCompletionState(IEnumerable<NodeInfo> NodesToDo, string StoreName, bool LocalOnly)
+	static void FindCompletionState(IEnumerable<BuildNode> NodesToDo, string StoreName, bool LocalOnly)
 	{
-		foreach(NodeInfo NodeToDo in NodesToDo)
+		foreach(BuildNode NodeToDo in NodesToDo)
 		{
 			string NodeStoreName = StoreName + "-" + NodeToDo.Name;
 			string GameNameIfAny = NodeToDo.Node.GameNameIfAnyForTempStorage();
@@ -293,13 +285,13 @@ public partial class GUBP : BuildCommand
 
     }
 
-    void PrintNodes(GUBP bp, List<NodeInfo> Nodes, IEnumerable<AggregateInfo> Aggregates, List<NodeInfo> UnfinishedTriggers, int TimeQuantum)
+    void PrintNodes(GUBP bp, List<BuildNode> Nodes, IEnumerable<AggregateNode> Aggregates, List<BuildNode> UnfinishedTriggers, int TimeQuantum)
     {
-		AggregateInfo[] MatchingAggregates = Aggregates.Where(x => x.Dependencies.All(y => Nodes.Contains(y))).ToArray();
+		AggregateNode[] MatchingAggregates = Aggregates.Where(x => x.Dependencies.All(y => Nodes.Contains(y))).ToArray();
 		if (MatchingAggregates.Length > 0)
 		{
 			Log("*********** Aggregates");
-			foreach (AggregateInfo Aggregate in MatchingAggregates.OrderBy(x => x.Name))
+			foreach (AggregateNode Aggregate in MatchingAggregates.OrderBy(x => x.Name))
 			{
 				StringBuilder Note = new StringBuilder("    " + Aggregate.Name);
 				if (Aggregate.Node.IsPromotableAggregate())
@@ -317,7 +309,7 @@ public partial class GUBP : BuildCommand
         string LastAgentGroup = "";
 
         Log("*********** Desired And Dependent Nodes, in order.");
-        foreach (NodeInfo NodeToDo in Nodes)
+        foreach (BuildNode NodeToDo in Nodes)
         {
             string MyControllingTrigger = NodeToDo.ControllingTriggerDotName;
             if (MyControllingTrigger != LastControllingTrigger)
@@ -391,11 +383,11 @@ public partial class GUBP : BuildCommand
 
             if (bShowDependencies)
             {
-                foreach (NodeInfo Dep in NodeToDo.Dependencies)
+                foreach (BuildNode Dep in NodeToDo.Dependencies)
                 {
                     Log("            dep> {0}", Dep.Name);
                 }
-                foreach (NodeInfo Dep in NodeToDo.PseudoDependencies)
+                foreach (BuildNode Dep in NodeToDo.PseudoDependencies)
                 {
                     Log("           pdep> {0}", Dep.Name);
                 }
@@ -403,15 +395,15 @@ public partial class GUBP : BuildCommand
         }
     }
 
-    static void SaveGraphVisualization(List<NodeInfo> Nodes)
+    static void SaveGraphVisualization(List<BuildNode> Nodes)
     {
         List<GraphNode> GraphNodes = new List<GraphNode>();
 
-        Dictionary<NodeInfo, GraphNode> NodeToGraphNodeMap = new Dictionary<NodeInfo, GraphNode>();
+        Dictionary<BuildNode, GraphNode> NodeToGraphNodeMap = new Dictionary<BuildNode, GraphNode>();
 
         for (int NodeIndex = 0; NodeIndex < Nodes.Count; ++NodeIndex)
         {
-            NodeInfo Node = Nodes[NodeIndex];
+            BuildNode Node = Nodes[NodeIndex];
 
             GraphNode GraphNode = new GraphNode()
             {
@@ -427,10 +419,10 @@ public partial class GUBP : BuildCommand
 
         for (int NodeIndex = 0; NodeIndex < Nodes.Count; ++NodeIndex)
         {
-            NodeInfo Node = Nodes[NodeIndex];
+            BuildNode Node = Nodes[NodeIndex];
             GraphNode NodeGraphNode = NodeToGraphNodeMap[Node];
 
-            foreach (NodeInfo Dep in Node.Dependencies)
+            foreach (BuildNode Dep in Node.Dependencies)
             {
                 GraphNode PrerequisiteFileGraphNode;
                 if (NodeToGraphNodeMap.TryGetValue(Dep, out PrerequisiteFileGraphNode))
@@ -448,7 +440,7 @@ public partial class GUBP : BuildCommand
                 }
 
             }
-            foreach (NodeInfo Dep in Node.PseudoDependencies)
+            foreach (BuildNode Dep in Node.PseudoDependencies)
             {
                 GraphNode PrerequisiteFileGraphNode;
                 if (NodeToGraphNodeMap.TryGetValue(Dep, out PrerequisiteFileGraphNode))
@@ -489,7 +481,7 @@ public partial class GUBP : BuildCommand
         Result.Sort();
         return Result;
     }
-    void SaveStatus(NodeInfo NodeToDo, string Suffix, string NodeStoreName, bool bSaveSharedTempStorage, string GameNameIfAny, string JobStepIDForFailure = null)
+    void SaveStatus(BuildNode NodeToDo, string Suffix, string NodeStoreName, bool bSaveSharedTempStorage, string GameNameIfAny, string JobStepIDForFailure = null)
     {
         string Contents = "Just a status record: " + Suffix;
         if (!String.IsNullOrEmpty(JobStepIDForFailure) && IsBuildMachine)
@@ -536,24 +528,24 @@ public partial class GUBP : BuildCommand
         return Result;
     }
 
-    static List<NodeInfo> TopologicalSort(HashSet<NodeInfo> NodesToDo, NodeInfo ExplicitTrigger, bool SubSort, bool DoNotConsiderCompletion)
+    static List<BuildNode> TopologicalSort(HashSet<BuildNode> NodesToDo, BuildNode ExplicitTrigger, bool SubSort, bool DoNotConsiderCompletion)
     {
         DateTime StartTime = DateTime.UtcNow;
 
-        List<NodeInfo> OrdereredToDo = new List<NodeInfo>();
+        List<BuildNode> OrdereredToDo = new List<BuildNode>();
 
-        Dictionary<string, List<NodeInfo>> SortedAgentGroupChains = new Dictionary<string, List<NodeInfo>>();
+        Dictionary<string, List<BuildNode>> SortedAgentGroupChains = new Dictionary<string, List<BuildNode>>();
         if (!SubSort)
         {
-            Dictionary<string, List<NodeInfo>> AgentGroupChains = new Dictionary<string, List<NodeInfo>>();
-            foreach (NodeInfo NodeToDo in NodesToDo)
+            Dictionary<string, List<BuildNode>> AgentGroupChains = new Dictionary<string, List<BuildNode>>();
+            foreach (BuildNode NodeToDo in NodesToDo)
             {
                 string MyAgentGroup = NodeToDo.Node.AgentSharingGroup;
                 if (MyAgentGroup != "")
                 {
                     if (!AgentGroupChains.ContainsKey(MyAgentGroup))
                     {
-                        AgentGroupChains.Add(MyAgentGroup, new List<NodeInfo> { NodeToDo });
+                        AgentGroupChains.Add(MyAgentGroup, new List<BuildNode> { NodeToDo });
                     }
                     else
                     {
@@ -561,9 +553,9 @@ public partial class GUBP : BuildCommand
                     }
                 }
             }
-            foreach (KeyValuePair<string, List<NodeInfo>> Chain in AgentGroupChains)
+            foreach (KeyValuePair<string, List<BuildNode>> Chain in AgentGroupChains)
             {
-                SortedAgentGroupChains.Add(Chain.Key, TopologicalSort(new HashSet<NodeInfo>(Chain.Value), ExplicitTrigger, true, DoNotConsiderCompletion));
+                SortedAgentGroupChains.Add(Chain.Key, TopologicalSort(new HashSet<BuildNode>(Chain.Value), ExplicitTrigger, true, DoNotConsiderCompletion));
             }
             Log("***************Done with recursion");
         }
@@ -573,12 +565,12 @@ public partial class GUBP : BuildCommand
         {
             bool bProgressMade = false;
             float BestPriority = -1E20f;
-            NodeInfo BestNode = null;
+            BuildNode BestNode = null;
             bool BestPseudoReady = false;
             HashSet<string> NonReadyAgentGroups = new HashSet<string>();
             HashSet<string> NonPeudoReadyAgentGroups = new HashSet<string>();
             HashSet<string> ExaminedAgentGroups = new HashSet<string>();
-            foreach (NodeInfo NodeToDo in NodesToDo)
+            foreach (BuildNode NodeToDo in NodesToDo)
             {
                 bool bReady = true;
 				bool bPseudoReady = true;
@@ -593,9 +585,9 @@ public partial class GUBP : BuildCommand
                     else
                     {
                         ExaminedAgentGroups.Add(NodeToDo.Node.AgentSharingGroup);
-                        foreach (NodeInfo ChainNode in SortedAgentGroupChains[NodeToDo.Node.AgentSharingGroup])
+                        foreach (BuildNode ChainNode in SortedAgentGroupChains[NodeToDo.Node.AgentSharingGroup])
                         {
-                            foreach (NodeInfo Dep in ChainNode.Dependencies)
+                            foreach (BuildNode Dep in ChainNode.Dependencies)
                             {
                                 if (!SortedAgentGroupChains[NodeToDo.Node.AgentSharingGroup].Contains(Dep) && NodesToDo.Contains(Dep))
                                 {
@@ -608,7 +600,7 @@ public partial class GUBP : BuildCommand
                                 NonReadyAgentGroups.Add(NodeToDo.Node.AgentSharingGroup);
                                 break;
                             }
-                            foreach (NodeInfo Dep in ChainNode.PseudoDependencies)
+                            foreach (BuildNode Dep in ChainNode.PseudoDependencies)
                             {
                                 if (!SortedAgentGroupChains[NodeToDo.Node.AgentSharingGroup].Contains(Dep) && NodesToDo.Contains(Dep))
                                 {
@@ -622,7 +614,7 @@ public partial class GUBP : BuildCommand
                 }
                 else
                 {
-                    foreach (NodeInfo Dep in NodeToDo.Dependencies)
+                    foreach (BuildNode Dep in NodeToDo.Dependencies)
                     {
                         if (NodesToDo.Contains(Dep))
                         {
@@ -630,7 +622,7 @@ public partial class GUBP : BuildCommand
                             break;
                         }
                     }
-                    foreach (NodeInfo Dep in NodeToDo.PseudoDependencies)
+                    foreach (BuildNode Dep in NodeToDo.PseudoDependencies)
                     {
                         if (NodesToDo.Contains(Dep))
                         {
@@ -702,7 +694,7 @@ public partial class GUBP : BuildCommand
             {
                 if (!SubSort && BestNode.Node.AgentSharingGroup != "")
                 {
-                    foreach (NodeInfo ChainNode in SortedAgentGroupChains[BestNode.Node.AgentSharingGroup])
+                    foreach (BuildNode ChainNode in SortedAgentGroupChains[BestNode.Node.AgentSharingGroup])
                     {
                         OrdereredToDo.Add(ChainNode);
                         NodesToDo.Remove(ChainNode);
@@ -718,14 +710,14 @@ public partial class GUBP : BuildCommand
             if (!bProgressMade && NodesToDo.Count > 0)
             {
                 Log("Cycle in GUBP, could not resolve:");
-                foreach (NodeInfo NodeToDo in NodesToDo)
+                foreach (BuildNode NodeToDo in NodesToDo)
                 {
                     string Deps = "";
                     if (!SubSort && NodeToDo.Node.AgentSharingGroup != "")
                     {
-                        foreach (NodeInfo ChainNode in SortedAgentGroupChains[NodeToDo.Node.AgentSharingGroup])
+                        foreach (BuildNode ChainNode in SortedAgentGroupChains[NodeToDo.Node.AgentSharingGroup])
                         {
-                            foreach (NodeInfo Dep in ChainNode.Dependencies)
+                            foreach (BuildNode Dep in ChainNode.Dependencies)
                             {
                                 if (!SortedAgentGroupChains[NodeToDo.Node.AgentSharingGroup].Contains(Dep) && NodesToDo.Contains(Dep))
                                 {
@@ -734,14 +726,14 @@ public partial class GUBP : BuildCommand
                             }
                         }
                     }
-                    foreach (NodeInfo Dep in NodeToDo.Dependencies)
+                    foreach (BuildNode Dep in NodeToDo.Dependencies)
                     {
                         if (NodesToDo.Contains(Dep))
                         {
                             Deps = Deps + Dep.Name + " ";
                         }
                     }
-                    foreach (NodeInfo Dep in NodeToDo.PseudoDependencies)
+                    foreach (BuildNode Dep in NodeToDo.PseudoDependencies)
                     {
                         if (NodesToDo.Contains(Dep))
                         {
@@ -762,7 +754,7 @@ public partial class GUBP : BuildCommand
         return OrdereredToDo;
     }
 
-    static string GetJobStepPath(NodeInfo Dep)
+    static string GetJobStepPath(BuildNode Dep)
     {
 		if (Dep.Node.AgentSharingGroup == "")
 		{
@@ -773,12 +765,12 @@ public partial class GUBP : BuildCommand
             return "jobSteps[" + Dep.Node.AgentSharingGroup + "]/jobSteps[" + Dep.Name + "]";
         }
     }
-    static string GetJobStep(string ParentPath, NodeInfo Dep)
+    static string GetJobStep(string ParentPath, BuildNode Dep)
     {
         return ParentPath + "/" + GetJobStepPath(Dep);
     }
 
-    static NodeHistory FindNodeHistory(NodeInfo NodeToDo, string CLString, string StoreName)
+    static NodeHistory FindNodeHistory(BuildNode NodeToDo, string CLString, string StoreName)
     {
 		NodeHistory History = null;
 
@@ -821,7 +813,7 @@ public partial class GUBP : BuildCommand
 
 		return History;
     }
-	void GetFailureEmails(NodeInfo NodeToDo, NodeHistory History, string CLString, string StoreName, bool OnlyLateUpdates = false)
+	void GetFailureEmails(BuildNode NodeToDo, NodeHistory History, string CLString, string StoreName, bool OnlyLateUpdates = false)
 	{
         string FailCauserEMails = "";
         string EMailNote = "";
@@ -939,7 +931,7 @@ public partial class GUBP : BuildCommand
         return true;
     }
 
-	int FindLastNonDuplicateFail(NodeInfo NodeToDo, NodeHistory History, string CLString, string StoreName)
+	int FindLastNonDuplicateFail(BuildNode NodeToDo, NodeHistory History, string CLString, string StoreName)
     {
         int Result = P4Env.Changelist;
 
@@ -998,7 +990,7 @@ public partial class GUBP : BuildCommand
         }
         return Result;
     }
-    List<string> GetECPropsForNode(NodeInfo NodeToDo, string CLString, bool OnlyLateUpdates)
+    List<string> GetECPropsForNode(BuildNode NodeToDo, string CLString, bool OnlyLateUpdates)
     {
         List<string> ECProps = new List<string>();
 		ECProps.Add("FailEmails/" + NodeToDo.Name + "=" + String.Join(" ", NodeToDo.RecipientsForFailureEmails));
@@ -1019,7 +1011,7 @@ public partial class GUBP : BuildCommand
         return ECProps;
     }
 
-    void UpdateECProps(NodeInfo NodeToDo, string CLString)
+    void UpdateECProps(BuildNode NodeToDo, string CLString)
     {
         try
         {
@@ -1037,7 +1029,7 @@ public partial class GUBP : BuildCommand
             Log(System.Diagnostics.TraceEventType.Warning, LogUtils.FormatException(Ex));
         }
     }
-	void UpdateECBuildTime(NodeInfo NodeToDo, double BuildDuration)
+	void UpdateECBuildTime(BuildNode NodeToDo, double BuildDuration)
 	{
 		try
 		{
@@ -1188,8 +1180,8 @@ public partial class GUBP : BuildCommand
 			ExplicitTriggerName = ParseParamValue("TriggerNode", "");
 		}
 
-		List<NodeInfo> AllNodes;
-		List<AggregateInfo> AllAggregates;
+		List<BuildNode> AllNodes;
+		List<AggregateNode> AllAggregates;
 		int TimeQuantum = 20;
 		AddNodesForBranch(CL, HostPlatforms, bPreflightBuild, PreflightMangleSuffix, out AllNodes, out AllAggregates, ref TimeQuantum);
 
@@ -1215,14 +1207,14 @@ public partial class GUBP : BuildCommand
 		}
 
 		bool WithoutLinux = ParseParam("NoLinux");
-		HashSet<NodeInfo> NodesToDo = ParseNodesToDo(AllNodes, AllAggregates, WithoutLinux, TimeIndex, bPreflightBuild, CommanderSetup);
+		HashSet<BuildNode> NodesToDo = ParseNodesToDo(AllNodes, AllAggregates, WithoutLinux, TimeIndex, bPreflightBuild, CommanderSetup);
 
-		NodeInfo ExplicitTrigger = null;
+		BuildNode ExplicitTrigger = null;
 		if (CommanderSetup)
         {
             if (!String.IsNullOrEmpty(ExplicitTriggerName))
             {
-                foreach (NodeInfo Node in AllNodes)
+                foreach (BuildNode Node in AllNodes)
                 {
                     if (Node.Name.Equals(ExplicitTriggerName, StringComparison.InvariantCultureIgnoreCase))
                     {
@@ -1243,7 +1235,7 @@ public partial class GUBP : BuildCommand
             {
                 if (bSkipTriggers)
                 {
-                    foreach (NodeInfo Node in AllNodes)
+                    foreach (BuildNode Node in AllNodes)
                     {
                         if (Node.Node.TriggerNode())
                         {
@@ -1254,9 +1246,9 @@ public partial class GUBP : BuildCommand
             }
         }
 
-        List<NodeInfo> OrderedToDo = TopologicalSort(NodesToDo, ExplicitTrigger, false, false);
+        List<BuildNode> OrderedToDo = TopologicalSort(NodesToDo, ExplicitTrigger, false, false);
 
-		List<NodeInfo> UnfinishedTriggers = FindUnfinishedTriggers(bSkipTriggers, ExplicitTrigger, OrderedToDo);
+		List<BuildNode> UnfinishedTriggers = FindUnfinishedTriggers(bSkipTriggers, ExplicitTrigger, OrderedToDo);
 
 		PrintNodes(this, OrderedToDo, AllAggregates, UnfinishedTriggers, TimeQuantum);
 
@@ -1279,10 +1271,10 @@ public partial class GUBP : BuildCommand
         PrintRunTime();
 	}
 
-	private static List<AggregateInfo> FindPromotables(IEnumerable<AggregateInfo> NodesToDo)
+	private static List<AggregateNode> FindPromotables(IEnumerable<AggregateNode> NodesToDo)
 	{
-		List<AggregateInfo> SeparatePromotables = new List<AggregateInfo>();
-		foreach (AggregateInfo NodeToDo in NodesToDo)
+		List<AggregateNode> SeparatePromotables = new List<AggregateNode>();
+		foreach (AggregateNode NodeToDo in NodesToDo)
 		{
 			if (NodeToDo.Node.IsSeparatePromotable())
 			{
@@ -1292,13 +1284,13 @@ public partial class GUBP : BuildCommand
 		return SeparatePromotables;
 	}
 
-	private static Dictionary<NodeInfo, List<AggregateInfo>> FindDependentPromotables(IEnumerable<NodeInfo> NodesToDo, IEnumerable<AggregateInfo> SeparatePromotions)
+	private static Dictionary<BuildNode, List<AggregateNode>> FindDependentPromotables(IEnumerable<BuildNode> NodesToDo, IEnumerable<AggregateNode> SeparatePromotions)
 	{
-		Dictionary<NodeInfo, List<AggregateInfo>> DependentPromotions = NodesToDo.ToDictionary(x => x, x => new List<AggregateInfo>());
-		foreach (AggregateInfo SeparatePromotion in SeparatePromotions)
+		Dictionary<BuildNode, List<AggregateNode>> DependentPromotions = NodesToDo.ToDictionary(x => x, x => new List<AggregateNode>());
+		foreach (AggregateNode SeparatePromotion in SeparatePromotions)
 		{
-			NodeInfo[] Dependencies = SeparatePromotion.Dependencies.SelectMany(x => x.AllIndirectDependencies).Distinct().ToArray();
-			foreach (NodeInfo Dependency in Dependencies)
+			BuildNode[] Dependencies = SeparatePromotion.Dependencies.SelectMany(x => x.AllIndirectDependencies).Distinct().ToArray();
+			foreach (BuildNode Dependency in Dependencies)
 			{
 				DependentPromotions[Dependency].Add(SeparatePromotion);
 			}
@@ -1311,26 +1303,26 @@ public partial class GUBP : BuildCommand
 	/// </summary>
 	/// <param name="AggregateNameToInfo">Map of aggregate names to their info objects</param>
 	/// <param name="NodeNameToInfo">Map of node names to their info objects</param>
-	private static void LinkGraph(IEnumerable<AggregateInfo> Aggregates, IEnumerable<NodeInfo> Nodes)
+	private static void LinkGraph(IEnumerable<AggregateNode> Aggregates, IEnumerable<BuildNode> Nodes)
 	{
-		Dictionary<string, NodeInfo> NodeNameToInfo = new Dictionary<string,NodeInfo>();
-		foreach(NodeInfo Node in Nodes)
+		Dictionary<string, BuildNode> NodeNameToInfo = new Dictionary<string,BuildNode>();
+		foreach(BuildNode Node in Nodes)
 		{
 			NodeNameToInfo.Add(Node.Name, Node);
 		}
 
-		Dictionary<string, AggregateInfo> AggregateNameToInfo = new Dictionary<string, AggregateInfo>();
-		foreach(AggregateInfo Aggregate in Aggregates)
+		Dictionary<string, AggregateNode> AggregateNameToInfo = new Dictionary<string, AggregateNode>();
+		foreach(AggregateNode Aggregate in Aggregates)
 		{
 			AggregateNameToInfo.Add(Aggregate.Name, Aggregate);
 		}
 
 		int NumErrors = 0;
-		foreach (AggregateInfo AggregateNode in AggregateNameToInfo.Values)
+		foreach (AggregateNode AggregateNode in AggregateNameToInfo.Values)
 		{
 			LinkAggregate(AggregateNode, AggregateNameToInfo, NodeNameToInfo, ref NumErrors);
 		}
-		foreach (NodeInfo BuildNode in NodeNameToInfo.Values)
+		foreach (BuildNode BuildNode in NodeNameToInfo.Values)
 		{
 			LinkNode(BuildNode, AggregateNameToInfo, NodeNameToInfo, ref NumErrors);
 		}
@@ -1347,16 +1339,16 @@ public partial class GUBP : BuildCommand
 	/// <param name="AggregateNameToInfo">Map of other aggregate names to their corresponding instance.</param>
 	/// <param name="NodeNameToInfo">Map from node names to their corresponding instance.</param>
 	/// <param name="NumErrors">The number of errors output so far. Incremented if resolving this aggregate fails.</param>
-	private static void LinkAggregate(AggregateInfo Aggregate, Dictionary<string, AggregateInfo> AggregateNameToInfo, Dictionary<string, NodeInfo> NodeNameToInfo, ref int NumErrors)
+	private static void LinkAggregate(AggregateNode Aggregate, Dictionary<string, AggregateNode> AggregateNameToInfo, Dictionary<string, BuildNode> NodeNameToInfo, ref int NumErrors)
 	{
 		if (Aggregate.Dependencies == null)
 		{
-			Aggregate.Dependencies = new NodeInfo[0];
+			Aggregate.Dependencies = new BuildNode[0];
 
-			HashSet<NodeInfo> Dependencies = new HashSet<NodeInfo>();
+			HashSet<BuildNode> Dependencies = new HashSet<BuildNode>();
 			foreach (string DependencyName in Aggregate.Node.Dependencies)
 			{
-				AggregateInfo AggregateDependency;
+				AggregateNode AggregateDependency;
 				if(AggregateNameToInfo.TryGetValue(DependencyName, out AggregateDependency))
 				{
 					LinkAggregate(AggregateDependency, AggregateNameToInfo, NodeNameToInfo, ref NumErrors);
@@ -1364,7 +1356,7 @@ public partial class GUBP : BuildCommand
 					continue;
 				}
 
-				NodeInfo Dependency;
+				BuildNode Dependency;
 				if(NodeNameToInfo.TryGetValue(DependencyName, out Dependency))
 				{
 					Dependencies.Add(Dependency);
@@ -1385,12 +1377,12 @@ public partial class GUBP : BuildCommand
 	/// <param name="AggregateNameToInfo">Map of other aggregate names to their corresponding instance.</param>
 	/// <param name="NodeNameToInfo">Map from node names to their corresponding instance.</param>
 	/// <param name="NumErrors">The number of errors output so far. Incremented if resolving this aggregate fails.</param>
-	private static void LinkNode(NodeInfo Node, Dictionary<string, AggregateInfo> AggregateNameToInfo, Dictionary<string, NodeInfo> NodeNameToInfo, ref int NumErrors)
+	private static void LinkNode(BuildNode Node, Dictionary<string, AggregateNode> AggregateNameToInfo, Dictionary<string, BuildNode> NodeNameToInfo, ref int NumErrors)
 	{
 		if(Node.Dependencies == null)
 		{
 			// Find all the dependencies
-			HashSet<NodeInfo> Dependencies = new HashSet<NodeInfo>();
+			HashSet<BuildNode> Dependencies = new HashSet<BuildNode>();
 			foreach (string DependencyName in Node.Node.FullNamesOfDependencies)
 			{
 				if (!ResolveDependencies(DependencyName, AggregateNameToInfo, NodeNameToInfo, Dependencies))
@@ -1402,7 +1394,7 @@ public partial class GUBP : BuildCommand
 			Node.Dependencies = Dependencies.ToArray();
 
 			// Find all the pseudo-dependencies
-			HashSet<NodeInfo> PseudoDependencies = new HashSet<NodeInfo>();
+			HashSet<BuildNode> PseudoDependencies = new HashSet<BuildNode>();
 			foreach (string PseudoDependencyName in Node.Node.FullNamesOfPseudosependencies)
 			{
 				if (!ResolveDependencies(PseudoDependencyName, AggregateNameToInfo, NodeNameToInfo, PseudoDependencies))
@@ -1417,8 +1409,8 @@ public partial class GUBP : BuildCommand
 			Node.AllDirectDependencies = Node.Dependencies.Union(Node.PseudoDependencies).ToArray();
 
 			// Recursively find the dependencies for all the dependencies
-			HashSet<NodeInfo> IndirectDependenices = new HashSet<NodeInfo>(Node.AllDirectDependencies);
-			foreach(NodeInfo DirectDependency in Node.AllDirectDependencies)
+			HashSet<BuildNode> IndirectDependenices = new HashSet<BuildNode>(Node.AllDirectDependencies);
+			foreach(BuildNode DirectDependency in Node.AllDirectDependencies)
 			{
 				LinkNode(DirectDependency, AggregateNameToInfo, NodeNameToInfo, ref NumErrors);
 				IndirectDependenices.UnionWith(DirectDependency.AllIndirectDependencies);
@@ -1442,16 +1434,16 @@ public partial class GUBP : BuildCommand
 	/// <param name="NodeNameToInfo">Map from node names to their corresponding info instance.</param>
 	/// <param name="Dependencies">The set of dependencies to add to.</param>
 	/// <returns>True if the name was found (and the dependencies list was updated), false otherwise.</returns>
-	private static bool ResolveDependencies(string Name, Dictionary<string, AggregateInfo> AggregateNameToInfo, Dictionary<string, NodeInfo> NodeNameToInfo, HashSet<NodeInfo> Dependencies)
+	private static bool ResolveDependencies(string Name, Dictionary<string, AggregateNode> AggregateNameToInfo, Dictionary<string, BuildNode> NodeNameToInfo, HashSet<BuildNode> Dependencies)
 	{
-		AggregateInfo AggregateDependency;
+		AggregateNode AggregateDependency;
 		if (AggregateNameToInfo.TryGetValue(Name, out AggregateDependency))
 		{
 			Dependencies.UnionWith(AggregateDependency.Dependencies);
 			return true;
 		}
 
-		NodeInfo NodeDependency;
+		BuildNode NodeDependency;
 		if (NodeNameToInfo.TryGetValue(Name, out NodeDependency))
 		{
 			Dependencies.Add(NodeDependency);
@@ -1461,7 +1453,7 @@ public partial class GUBP : BuildCommand
 		return false;
 	}
 
-	private HashSet<NodeInfo> ParseNodesToDo(IEnumerable<NodeInfo> PotentialNodes, IEnumerable<AggregateInfo> PotentialAggregates, bool WithoutLinux, int TimeIndex, bool bPreflightBuild, bool CommanderSetup)
+	private HashSet<BuildNode> ParseNodesToDo(IEnumerable<BuildNode> PotentialNodes, IEnumerable<AggregateNode> PotentialAggregates, bool WithoutLinux, int TimeIndex, bool bPreflightBuild, bool CommanderSetup)
 	{
 		List<string> NamesToDo = new List<string>();
 
@@ -1477,11 +1469,11 @@ public partial class GUBP : BuildCommand
 			NamesToDo.AddRange(GameParam.Split('+').Select(x => FullGameAggregateNode.StaticGetFullName(x)));
 		}
 
-		HashSet<NodeInfo> NodesToDo = new HashSet<NodeInfo>();
+		HashSet<BuildNode> NodesToDo = new HashSet<BuildNode>();
 		foreach (string NameToDo in NamesToDo)
 		{
 			int FoundNames = 0;
-			foreach (NodeInfo PotentialNode in PotentialNodes)
+			foreach (BuildNode PotentialNode in PotentialNodes)
 			{
 				if (String.Compare(PotentialNode.Name, NameToDo, StringComparison.InvariantCultureIgnoreCase) == 0 || String.Compare(PotentialNode.Node.AgentSharingGroup, NameToDo, StringComparison.InvariantCultureIgnoreCase) == 0)
 				{
@@ -1489,7 +1481,7 @@ public partial class GUBP : BuildCommand
 					FoundNames++;
 				}
 			}
-			foreach (AggregateInfo PotentialAggregate in PotentialAggregates)
+			foreach (AggregateNode PotentialAggregate in PotentialAggregates)
 			{
 				if (String.Compare(PotentialAggregate.Name, NameToDo, StringComparison.InvariantCultureIgnoreCase) == 0)
 				{
@@ -1511,7 +1503,7 @@ public partial class GUBP : BuildCommand
 		else if (TimeIndex != 0)
 		{
 			LogVerbose("Check to make sure we didn't ask for nodes that will be culled by time index");
-			foreach (NodeInfo NodeToDo in NodesToDo)
+			foreach (BuildNode NodeToDo in NodesToDo)
 			{
 				if (TimeIndex % (1 << NodeToDo.FrequencyShift) != 0)
 				{
@@ -1521,7 +1513,7 @@ public partial class GUBP : BuildCommand
 		}
 
 		LogVerbose("Desired Nodes");
-		foreach (NodeInfo NodeToDo in NodesToDo)
+		foreach (BuildNode NodeToDo in NodesToDo)
 		{
 			LogVerbose("  {0}", NodeToDo.Name);
 		}
@@ -1531,10 +1523,10 @@ public partial class GUBP : BuildCommand
 		while (!bDoneWithDependencies)
 		{
 			bDoneWithDependencies = true;
-			HashSet<NodeInfo> Fringe = new HashSet<NodeInfo>();
-			foreach (NodeInfo NodeToDo in NodesToDo)
+			HashSet<BuildNode> Fringe = new HashSet<BuildNode>();
+			foreach (BuildNode NodeToDo in NodesToDo)
 			{
-				foreach (NodeInfo Dep in NodeToDo.Dependencies)
+				foreach (BuildNode Dep in NodeToDo.Dependencies)
 				{
 					if (!NodesToDo.Contains(Dep))
 					{
@@ -1549,8 +1541,8 @@ public partial class GUBP : BuildCommand
 		if (TimeIndex != 0)
 		{
 			LogVerbose("Culling based on time index");
-			HashSet<NodeInfo> NewNodesToDo = new HashSet<NodeInfo>();
-			foreach (NodeInfo NodeToDo in NodesToDo)
+			HashSet<BuildNode> NewNodesToDo = new HashSet<BuildNode>();
+			foreach (BuildNode NodeToDo in NodesToDo)
 			{
 				if (TimeIndex % (1 << NodeToDo.FrequencyShift) == 0)
 				{
@@ -1567,8 +1559,8 @@ public partial class GUBP : BuildCommand
 		//Remove Plat if specified
 		if (WithoutLinux)
 		{
-			HashSet<NodeInfo> NewNodesToDo = new HashSet<NodeInfo>();
-			foreach (NodeInfo NodeToDo in NodesToDo)
+			HashSet<BuildNode> NewNodesToDo = new HashSet<BuildNode>();
+			foreach (BuildNode NodeToDo in NodesToDo)
 			{
 				if (!NodeToDo.Node.GetFullName().Contains("Linux"))
 				{
@@ -1584,8 +1576,8 @@ public partial class GUBP : BuildCommand
 		if (bPreflightBuild)
 		{
 			LogVerbose("Culling triggers and downstream for preflight builds ");
-			HashSet<NodeInfo> NewNodesToDo = new HashSet<NodeInfo>();
-			foreach (NodeInfo NodeToDo in NodesToDo)
+			HashSet<BuildNode> NewNodesToDo = new HashSet<BuildNode>();
+			foreach (BuildNode NodeToDo in NodesToDo)
 			{
 				string TriggerDot = NodeToDo.ControllingTriggerDotName;
 				if (TriggerDot == "" && !NodeToDo.Node.TriggerNode())
@@ -1603,19 +1595,19 @@ public partial class GUBP : BuildCommand
 		return NodesToDo;
 	}
 
-	private static Dictionary<string, string> GetFullNodeDependedOnBy(HashSet<NodeInfo> NodesToDo)
+	private static Dictionary<string, string> GetFullNodeDependedOnBy(HashSet<BuildNode> NodesToDo)
 	{
 		//find things that depend on our nodes and setup commander dictionary
 		Dictionary<string, string> FullNodeDependedOnBy = new Dictionary<string, string>();
-		foreach (NodeInfo NodeToDo in NodesToDo)
+		foreach (BuildNode NodeToDo in NodesToDo)
 		{
 			FullNodeDependedOnBy[NodeToDo.Name] = "";
 		}
-		foreach (NodeInfo NodeToDo in NodesToDo)
+		foreach (BuildNode NodeToDo in NodesToDo)
 		{
 			if (!NodeToDo.Node.IsTest())
 			{
-				foreach (NodeInfo Dep in NodeToDo.AllDirectDependencies)
+				foreach (BuildNode Dep in NodeToDo.AllDirectDependencies)
 				{
 					string CurrentValue;
 					if (!FullNodeDependedOnBy.TryGetValue(Dep.Name, out CurrentValue) || CurrentValue.Length == 0)
@@ -1745,13 +1737,13 @@ public partial class GUBP : BuildCommand
 		throw new AutomationException("Failed to update the CIS counter after 20 tries.");
 	}
 
-	private List<NodeInfo> FindUnfinishedTriggers(bool bSkipTriggers, NodeInfo ExplicitTrigger, List<NodeInfo> OrdereredToDo)
+	private List<BuildNode> FindUnfinishedTriggers(bool bSkipTriggers, BuildNode ExplicitTrigger, List<BuildNode> OrdereredToDo)
 	{
 		// find all unfinished triggers, excepting the one we are triggering right now
-		List<NodeInfo> UnfinishedTriggers = new List<NodeInfo>();
+		List<BuildNode> UnfinishedTriggers = new List<BuildNode>();
 		if (!bSkipTriggers)
 		{
-			foreach (NodeInfo NodeToDo in OrdereredToDo)
+			foreach (BuildNode NodeToDo in OrdereredToDo)
 			{
 				if (NodeToDo.Node.TriggerNode() && !NodeToDo.IsComplete)
 				{
@@ -1765,16 +1757,16 @@ public partial class GUBP : BuildCommand
 		return UnfinishedTriggers;
 	}
 
-	private void CheckSortOrder(List<NodeInfo> OrdereredToDo)
+	private void CheckSortOrder(List<BuildNode> OrdereredToDo)
 	{
-		foreach (NodeInfo NodeToDo in OrdereredToDo)
+		foreach (BuildNode NodeToDo in OrdereredToDo)
 		{
 			if (NodeToDo.Node.TriggerNode() && (NodeToDo.Node.IsSticky() || NodeToDo.IsComplete)) // these sticky triggers are ok, everything is already completed anyway
 			{
 				continue;
 			}
 			int MyIndex = OrdereredToDo.IndexOf(NodeToDo);
-			foreach (NodeInfo Dep in NodeToDo.Dependencies)
+			foreach (BuildNode Dep in NodeToDo.Dependencies)
 			{
 				int DepIndex = OrdereredToDo.IndexOf(Dep);
 				if (DepIndex >= MyIndex)
@@ -1782,7 +1774,7 @@ public partial class GUBP : BuildCommand
 					throw new AutomationException("Topological sort error, node {0} has a dependency of {1} which sorted after it.", NodeToDo.Name, Dep.Name);
 				}
 			}
-			foreach (NodeInfo Dep in NodeToDo.PseudoDependencies)
+			foreach (BuildNode Dep in NodeToDo.PseudoDependencies)
 			{
 				int DepIndex = OrdereredToDo.IndexOf(Dep);
 				if (DepIndex >= MyIndex)
@@ -1793,17 +1785,17 @@ public partial class GUBP : BuildCommand
 		}
 	}
 
-	private void DoCommanderSetup(IEnumerable<NodeInfo> AllNodes, IEnumerable<AggregateInfo> AllAggregates, HashSet<NodeInfo> NodesToDo, List<NodeInfo> OrdereredToDo, int TimeIndex, int TimeQuantum, bool bSkipTriggers, bool bFake, bool bFakeEC, string CLString, NodeInfo ExplicitTrigger, List<NodeInfo> UnfinishedTriggers, string FakeFail, bool bPreflightBuild)
+	private void DoCommanderSetup(IEnumerable<BuildNode> AllNodes, IEnumerable<AggregateNode> AllAggregates, HashSet<BuildNode> NodesToDo, List<BuildNode> OrdereredToDo, int TimeIndex, int TimeQuantum, bool bSkipTriggers, bool bFake, bool bFakeEC, string CLString, BuildNode ExplicitTrigger, List<BuildNode> UnfinishedTriggers, string FakeFail, bool bPreflightBuild)
 	{
 		Dictionary<string, string> FullNodeDependedOnBy = GetFullNodeDependedOnBy(NodesToDo);
 
-		List<AggregateInfo> SeparatePromotables = FindPromotables(AllAggregates);
-		Dictionary<NodeInfo, List<AggregateInfo>> DependentPromotions = FindDependentPromotables(AllNodes, SeparatePromotables);
+		List<AggregateNode> SeparatePromotables = FindPromotables(AllAggregates);
+		Dictionary<BuildNode, List<AggregateNode>> DependentPromotions = FindDependentPromotables(AllNodes, SeparatePromotables);
 
-		List<NodeInfo> SortedNodes = TopologicalSort(new HashSet<NodeInfo>(AllNodes), null, SubSort: false, DoNotConsiderCompletion: true);
+		List<BuildNode> SortedNodes = TopologicalSort(new HashSet<BuildNode>(AllNodes), null, SubSort: false, DoNotConsiderCompletion: true);
 		Log("******* {0} GUBP Nodes", SortedNodes.Count);
 
-		Dictionary<NodeInfo, int> FullNodeListSortKey = GetDisplayOrder(SortedNodes);
+		Dictionary<BuildNode, int> FullNodeListSortKey = GetDisplayOrder(SortedNodes);
 
 		if (OrdereredToDo.Count == 0)
 		{
@@ -1811,15 +1803,15 @@ public partial class GUBP : BuildCommand
 		}
 		List<string> ECProps = new List<string>();
 		ECProps.Add(String.Format("TimeIndex={0}", TimeIndex));
-		foreach (NodeInfo Node in SortedNodes)
+		foreach (BuildNode Node in SortedNodes)
 		{
 			ECProps.Add(string.Format("AllNodes/{0}={1}", Node.Name, GetNodeForAllNodesProperty(Node, TimeQuantum)));
 		}
-		foreach (NodeInfo Node in SortedNodes)
+		foreach (BuildNode Node in SortedNodes)
 		{
 			ECProps.Add(string.Format("DirectDependencies/{0}={1}", Node.Name, String.Join(" ", Node.AllDirectDependencies.Select(x => x.Name))));
 		}
-		foreach (KeyValuePair<NodeInfo, int> NodePair in FullNodeListSortKey)
+		foreach (KeyValuePair<BuildNode, int> NodePair in FullNodeListSortKey)
 		{
 			ECProps.Add(string.Format("SortKey/{0}={1}", NodePair.Key.Name, NodePair.Value));
 		}
@@ -1827,11 +1819,11 @@ public partial class GUBP : BuildCommand
 		{
 			ECProps.Add(string.Format("DependedOnBy/{0}={1}", NodePair.Key, NodePair.Value));
 		}
-		foreach (KeyValuePair<NodeInfo, List<AggregateInfo>> NodePair in DependentPromotions)
+		foreach (KeyValuePair<BuildNode, List<AggregateNode>> NodePair in DependentPromotions)
 		{
 			ECProps.Add(string.Format("DependentPromotions/{0}={1}", NodePair.Key.Name, String.Join(" ", NodePair.Value.Select(x => x.Name))));
 		}
-		foreach (AggregateInfo Node in SeparatePromotables)
+		foreach (AggregateNode Node in SeparatePromotables)
 		{
 			ECProps.Add(string.Format("PossiblePromotables/{0}={1}", Node.Name, ""));
 		}
@@ -1847,11 +1839,11 @@ public partial class GUBP : BuildCommand
 
 		using(TelemetryStopwatch StartFilterTimer = new TelemetryStopwatch("FilterNodes"))
 		{
-			List<NodeInfo> FilteredOrdereredToDo = new List<NodeInfo>();
+			List<BuildNode> FilteredOrdereredToDo = new List<BuildNode>();
 			// remove nodes that have unfinished triggers
-			foreach (NodeInfo NodeToDo in OrdereredToDo)
+			foreach (BuildNode NodeToDo in OrdereredToDo)
 			{
-				NodeInfo ControllingTrigger = (NodeToDo.ControllingTriggers.Length > 0)? NodeToDo.ControllingTriggers.Last() : null;
+				BuildNode ControllingTrigger = (NodeToDo.ControllingTriggers.Length > 0)? NodeToDo.ControllingTriggers.Last() : null;
 				bool bNoUnfinishedTriggers = !UnfinishedTriggers.Contains(ControllingTrigger);
 
 				if (bNoUnfinishedTriggers)
@@ -1882,7 +1874,7 @@ public partial class GUBP : BuildCommand
 		// here we are just making sure everything before the explicit trigger is completed.
 		if (ExplicitTrigger != null)
 		{
-			foreach (NodeInfo NodeToDo in OrdereredToDo)
+			foreach (BuildNode NodeToDo in OrdereredToDo)
 			{
 				if (!NodeToDo.IsComplete && NodeToDo != ExplicitTrigger && !NodeToDo.DependsOn(ExplicitTrigger)) // if something is already finished, we don't put it into EC
 				{
@@ -1891,7 +1883,7 @@ public partial class GUBP : BuildCommand
 			}
 		}
 
-		NodeInfo LastSticky = null;
+		BuildNode LastSticky = null;
 		bool HitNonSticky = false;
 		bool bHaveECNodes = false;
 		List<string> StepList = new List<string>();
@@ -1902,7 +1894,7 @@ public partial class GUBP : BuildCommand
 		StepList.Add("$ec->setTimeout(600);");
 		StepList.Add("my $batch = $ec->newBatch(\"serial\");");
 		// sticky nodes are ones that we run on the main agent. We run then first and they must not be intermixed with parallel jobs
-		foreach (NodeInfo NodeToDo in OrdereredToDo)
+		foreach (BuildNode NodeToDo in OrdereredToDo)
 		{
 			if (!NodeToDo.IsComplete) // if something is already finished, we don't put it into EC
 			{
@@ -1937,9 +1929,9 @@ public partial class GUBP : BuildCommand
 				bHasNoop = true;
 			}
 
-			Dictionary<string, List<NodeInfo>> AgentGroupChains = new Dictionary<string, List<NodeInfo>>();
-			List<NodeInfo> StickyChain = new List<NodeInfo>();
-			foreach (NodeInfo NodeToDo in OrdereredToDo)
+			Dictionary<string, List<BuildNode>> AgentGroupChains = new Dictionary<string, List<BuildNode>>();
+			List<BuildNode> StickyChain = new List<BuildNode>();
+			foreach (BuildNode NodeToDo in OrdereredToDo)
 			{
 				if (!NodeToDo.IsComplete) // if something is already finished, we don't put it into EC  
 				{
@@ -1948,7 +1940,7 @@ public partial class GUBP : BuildCommand
 					{
 						if (!AgentGroupChains.ContainsKey(MyAgentGroup))
 						{
-							AgentGroupChains.Add(MyAgentGroup, new List<NodeInfo> { NodeToDo });
+							AgentGroupChains.Add(MyAgentGroup, new List<BuildNode> { NodeToDo });
 						}
 						else
 						{
@@ -1964,7 +1956,7 @@ public partial class GUBP : BuildCommand
 					}
 				}
 			}
-			foreach (NodeInfo NodeToDo in OrdereredToDo)
+			foreach (BuildNode NodeToDo in OrdereredToDo)
 			{
 				if (!NodeToDo.IsComplete) // if something is already finished, we don't put it into EC  
 				{
@@ -2000,9 +1992,9 @@ public partial class GUBP : BuildCommand
 					}
 					Args = Args + "]";
 
-					List<NodeInfo> UncompletedEcDeps = new List<NodeInfo>();
+					List<BuildNode> UncompletedEcDeps = new List<BuildNode>();
 					{
-						foreach (NodeInfo Dep in NodeToDo.AllDirectDependencies)
+						foreach (BuildNode Dep in NodeToDo.AllDirectDependencies)
 						{
 							if (!Dep.IsComplete && OrdereredToDo.Contains(Dep)) // if something is already finished, we don't put it into EC
 							{
@@ -2028,7 +2020,7 @@ public partial class GUBP : BuildCommand
 						bDoNestedJobstep = true;
 						NodeParentPath = ParentPath + "/jobSteps[" + MyAgentGroup + "]";
 
-						List<NodeInfo> MyChain = AgentGroupChains[MyAgentGroup];
+						List<BuildNode> MyChain = AgentGroupChains[MyAgentGroup];
 						int MyIndex = MyChain.IndexOf(NodeToDo);
 						if (MyIndex <= 0)
 						{
@@ -2133,7 +2125,7 @@ public partial class GUBP : BuildCommand
 
 					if (MyAgentGroup != "" && !bDoNestedJobstep)
 					{
-						List<NodeInfo> MyChain = AgentGroupChains[MyAgentGroup];
+						List<BuildNode> MyChain = AgentGroupChains[MyAgentGroup];
 						int MyIndex = MyChain.IndexOf(NodeToDo);
 						if (MyIndex == MyChain.Count - 1)
 						{
@@ -2167,16 +2159,16 @@ public partial class GUBP : BuildCommand
 		}
 	}
 
-	private string GetPreConditionForNode(List<NodeInfo> OrdereredToDo, string PreconditionParentPath, bool bHasNoop, Dictionary<string, List<NodeInfo>> AgentGroupChains, List<NodeInfo> StickyChain, NodeInfo NodeToDo, List<NodeInfo> UncompletedEcDeps)
+	private string GetPreConditionForNode(List<BuildNode> OrdereredToDo, string PreconditionParentPath, bool bHasNoop, Dictionary<string, List<BuildNode>> AgentGroupChains, List<BuildNode> StickyChain, BuildNode NodeToDo, List<BuildNode> UncompletedEcDeps)
 	{
-		List<NodeInfo> PreConditionUncompletedEcDeps = new List<NodeInfo>();
+		List<BuildNode> PreConditionUncompletedEcDeps = new List<BuildNode>();
 		if(NodeToDo.Node.AgentSharingGroup == "")
 		{
-			PreConditionUncompletedEcDeps = new List<NodeInfo>(UncompletedEcDeps);
+			PreConditionUncompletedEcDeps = new List<BuildNode>(UncompletedEcDeps);
 		}
 		else 
 		{
-			List<NodeInfo> MyChain = AgentGroupChains[NodeToDo.Node.AgentSharingGroup];
+			List<BuildNode> MyChain = AgentGroupChains[NodeToDo.Node.AgentSharingGroup];
 			int MyIndex = MyChain.IndexOf(NodeToDo);
 			if (MyIndex > 0)
 			{
@@ -2185,9 +2177,9 @@ public partial class GUBP : BuildCommand
 			else
 			{
 				// to avoid idle agents (and also EC doesn't actually reserve our agent!), we promote all dependencies to the first one
-				foreach (NodeInfo Chain in MyChain)
+				foreach (BuildNode Chain in MyChain)
 				{
-					foreach (NodeInfo Dep in Chain.AllDirectDependencies)
+					foreach (BuildNode Dep in Chain.AllDirectDependencies)
 					{
 						if (!Dep.IsComplete && OrdereredToDo.Contains(Dep)) // if something is already finished, we don't put it into EC
 						{
@@ -2206,7 +2198,7 @@ public partial class GUBP : BuildCommand
 		}
 		if (NodeToDo.Node.IsSticky())
 		{
-			List<NodeInfo> MyChain = StickyChain;
+			List<BuildNode> MyChain = StickyChain;
 			int MyIndex = MyChain.IndexOf(NodeToDo);
 			if (MyIndex > 0)
 			{
@@ -2217,7 +2209,7 @@ public partial class GUBP : BuildCommand
 			}
 			else
 			{
-				foreach (NodeInfo Dep in NodeToDo.AllDirectDependencies)
+				foreach (BuildNode Dep in NodeToDo.AllDirectDependencies)
 				{
 					if (!Dep.IsComplete && OrdereredToDo.Contains(Dep)) // if something is already finished, we don't put it into EC
 					{
@@ -2252,7 +2244,7 @@ public partial class GUBP : BuildCommand
 		return PreCondition;
 	}
 
-	private static string GetNodeForAllNodesProperty(NodeInfo Node, int TimeQuantum)
+	private static string GetNodeForAllNodesProperty(BuildNode Node, int TimeQuantum)
 	{
 		string Note = Node.ControllingTriggerDotName;
 		if (Note == "")
@@ -2269,14 +2261,14 @@ public partial class GUBP : BuildCommand
 		return Note;
 	}
 
-	private static string GetRunConditionForNode(List<NodeInfo> UncompletedEcDeps, string PreconditionParentPath)
+	private static string GetRunConditionForNode(List<BuildNode> UncompletedEcDeps, string PreconditionParentPath)
 	{
 		string RunCondition = "";
 		if (UncompletedEcDeps.Count > 0)
 		{
 			RunCondition = "\"\\$\" . \"[/javascript if(";
 			int Index = 0;
-			foreach (NodeInfo Dep in UncompletedEcDeps)
+			foreach (BuildNode Dep in UncompletedEcDeps)
 			{
 				RunCondition = RunCondition + "((\'\\$\" . \"[" + GetJobStep(PreconditionParentPath, Dep) + "/outcome]\' == \'success\') || ";
 				RunCondition = RunCondition + "(\'\\$\" . \"[" + GetJobStep(PreconditionParentPath, Dep) + "/outcome]\' == \'warning\'))";
@@ -2292,10 +2284,10 @@ public partial class GUBP : BuildCommand
 		return RunCondition;
 	}
 
-	void ExecuteNodes(List<NodeInfo> OrdereredToDo, bool bFake, bool bFakeEC, bool bSaveSharedTempStorage, string CLString, string StoreName, string FakeFail)
+	void ExecuteNodes(List<BuildNode> OrdereredToDo, bool bFake, bool bFakeEC, bool bSaveSharedTempStorage, string CLString, string StoreName, string FakeFail)
 	{
-        Dictionary<string, NodeInfo> BuildProductToNodeMap = new Dictionary<string, NodeInfo>();
-		foreach (NodeInfo NodeToDo in OrdereredToDo)
+        Dictionary<string, BuildNode> BuildProductToNodeMap = new Dictionary<string, BuildNode>();
+		foreach (BuildNode NodeToDo in OrdereredToDo)
         {
             if (NodeToDo.Node.BuildProducts != null || NodeToDo.Node.AllDependencyBuildProducts != null)
             {
@@ -2304,7 +2296,7 @@ public partial class GUBP : BuildCommand
 
             NodeToDo.Node.AllDependencyBuildProducts = new List<string>();
             NodeToDo.Node.AllDependencies = new List<string>();
-            foreach (NodeInfo Dep in NodeToDo.Dependencies)
+            foreach (BuildNode Dep in NodeToDo.Dependencies)
             {
                 NodeToDo.Node.AddAllDependent(Dep.Name);
 
@@ -2568,34 +2560,34 @@ public partial class GUBP : BuildCommand
 	/// Sorts a list of nodes to display in EC. The default order is based on execution order and agent groups, whereas this function arranges nodes by
 	/// frequency then execution order, while trying to group nodes on parallel paths (eg. Mac/Windows editor nodes) together.
 	/// </summary>
-	static Dictionary<NodeInfo, int> GetDisplayOrder(List<NodeInfo> Nodes)
+	static Dictionary<BuildNode, int> GetDisplayOrder(List<BuildNode> Nodes)
 	{
 		// Split the nodes into separate lists for each frequency
-		SortedDictionary<int, List<NodeInfo>> NodesByFrequency = new SortedDictionary<int,List<NodeInfo>>();
-		foreach(NodeInfo Node in Nodes)
+		SortedDictionary<int, List<BuildNode>> NodesByFrequency = new SortedDictionary<int,List<BuildNode>>();
+		foreach(BuildNode Node in Nodes)
 		{
-			List<NodeInfo> NodesByThisFrequency;
+			List<BuildNode> NodesByThisFrequency;
 			if(!NodesByFrequency.TryGetValue(Node.FrequencyShift, out NodesByThisFrequency))
 			{
-				NodesByThisFrequency = new List<NodeInfo>();
+				NodesByThisFrequency = new List<BuildNode>();
 				NodesByFrequency.Add(Node.FrequencyShift, NodesByThisFrequency);
 			}
 			NodesByThisFrequency.Add(Node);
 		}
 
 		// Build the output list by scanning each frequency in order
-		HashSet<NodeInfo> VisitedNodes = new HashSet<NodeInfo>();
-		Dictionary<NodeInfo, int> SortedNodes = new Dictionary<NodeInfo,int>();
-		foreach(List<NodeInfo> NodesByThisFrequency in NodesByFrequency.Values)
+		HashSet<BuildNode> VisitedNodes = new HashSet<BuildNode>();
+		Dictionary<BuildNode, int> SortedNodes = new Dictionary<BuildNode,int>();
+		foreach(List<BuildNode> NodesByThisFrequency in NodesByFrequency.Values)
 		{
 			// Find a list of nodes in each display group. If the group name matches the node name, put that node at the front of the list.
-			Dictionary<string, List<NodeInfo>> DisplayGroups = new Dictionary<string,List<NodeInfo>>();
-			foreach(NodeInfo Node in NodesByThisFrequency)
+			Dictionary<string, List<BuildNode>> DisplayGroups = new Dictionary<string,List<BuildNode>>();
+			foreach(BuildNode Node in NodesByThisFrequency)
 			{
 				string GroupName = Node.Node.GetDisplayGroupName();
 				if(!DisplayGroups.ContainsKey(GroupName))
 				{
-					DisplayGroups.Add(GroupName, new List<NodeInfo>{ Node });
+					DisplayGroups.Add(GroupName, new List<BuildNode>{ Node });
 				}
 				else if(GroupName == Node.Name)
 				{
@@ -2608,10 +2600,10 @@ public partial class GUBP : BuildCommand
 			}
 
 			// Build a list of ordering dependencies, putting all Mac nodes after Windows nodes with the same names.
-			Dictionary<NodeInfo, List<NodeInfo>> NodeDependencies = new Dictionary<NodeInfo,List<NodeInfo>>(Nodes.ToDictionary(x => x, x => x.AllDirectDependencies.ToList()));
-			foreach(KeyValuePair<string, List<NodeInfo>> DisplayGroup in DisplayGroups)
+			Dictionary<BuildNode, List<BuildNode>> NodeDependencies = new Dictionary<BuildNode,List<BuildNode>>(Nodes.ToDictionary(x => x, x => x.AllDirectDependencies.ToList()));
+			foreach(KeyValuePair<string, List<BuildNode>> DisplayGroup in DisplayGroups)
 			{
-				List<NodeInfo> GroupNodes = DisplayGroup.Value;
+				List<BuildNode> GroupNodes = DisplayGroup.Value;
 				for (int Idx = 1; Idx < GroupNodes.Count; Idx++)
 				{
 					NodeDependencies[GroupNodes[Idx]].Add(GroupNodes[0]);
@@ -2619,10 +2611,10 @@ public partial class GUBP : BuildCommand
 			}
 
 			// Add nodes for each frequency into the master list, trying to match up different groups along the way
-			foreach(NodeInfo FirstNode in NodesByThisFrequency)
+			foreach(BuildNode FirstNode in NodesByThisFrequency)
 			{
-				List<NodeInfo> GroupNodes = DisplayGroups[FirstNode.Node.GetDisplayGroupName()];
-				foreach(NodeInfo GroupNode in GroupNodes)
+				List<BuildNode> GroupNodes = DisplayGroups[FirstNode.Node.GetDisplayGroupName()];
+				foreach(BuildNode GroupNode in GroupNodes)
 				{
 					AddNodeAndDependencies(GroupNode, NodeDependencies, VisitedNodes, SortedNodes);
 				}
@@ -2631,12 +2623,12 @@ public partial class GUBP : BuildCommand
 		return SortedNodes;
 	}
 
-	static void AddNodeAndDependencies(NodeInfo Node, Dictionary<NodeInfo, List<NodeInfo>> NodeDependencies, HashSet<NodeInfo> VisitedNodes, Dictionary<NodeInfo, int> SortedNodes)
+	static void AddNodeAndDependencies(BuildNode Node, Dictionary<BuildNode, List<BuildNode>> NodeDependencies, HashSet<BuildNode> VisitedNodes, Dictionary<BuildNode, int> SortedNodes)
 	{
 		if(!VisitedNodes.Contains(Node))
 		{
 			VisitedNodes.Add(Node);
-			foreach (NodeInfo NodeDependency in NodeDependencies[Node])
+			foreach (BuildNode NodeDependency in NodeDependencies[Node])
 			{
 				AddNodeAndDependencies(NodeDependency, NodeDependencies, VisitedNodes, SortedNodes);
 			}
diff --git a/Engine/Source/Programs/AutomationTool/GUBP/LegacyBranchSetup.cs b/Engine/Source/Programs/AutomationTool/GUBP/LegacyBranchSetup.cs
--- a/Engine/Source/Programs/AutomationTool/GUBP/LegacyBranchSetup.cs
+++ b/Engine/Source/Programs/AutomationTool/GUBP/LegacyBranchSetup.cs
@@ -213,7 +213,7 @@ partial class GUBP
 		}
     }
 
-	private void FindEmailsForNodes(string BranchName, IEnumerable<NodeInfo> Nodes)
+	private void FindEmailsForNodes(string BranchName, IEnumerable<BuildNode> Nodes)
 	{
 		List<GUBPEmailHacker> EmailHackers = new List<GUBPEmailHacker>();
 
@@ -234,13 +234,13 @@ partial class GUBP
 			}
 		}
 
-		foreach (NodeInfo Node in Nodes)
+		foreach (BuildNode Node in Nodes)
 		{
 			Node.RecipientsForFailureEmails = HackEmails(EmailHackers, BranchName, Node, out Node.AddSubmittersToFailureEmails);
 		}
 	}
 
-    private string[] HackEmails(List<GUBPEmailHacker> EmailHackers, string Branch, NodeInfo NodeInfo, out bool bIncludeCausers)
+    private string[] HackEmails(List<GUBPEmailHacker> EmailHackers, string Branch, BuildNode NodeInfo, out bool bIncludeCausers)
     {
 		string OnlyEmail = ParseParamValue("OnlyEmail");
         if (!String.IsNullOrEmpty(OnlyEmail))
@@ -295,7 +295,7 @@ partial class GUBP
         }
     }
 
-	void FindFrequenciesForNodes(GUBP.GUBPBranchConfig BranchConfig, IEnumerable<NodeInfo> Nodes, Dictionary<string, int> FrequencyOverrides)
+	void FindFrequenciesForNodes(GUBP.GUBPBranchConfig BranchConfig, IEnumerable<BuildNode> Nodes, Dictionary<string, int> FrequencyOverrides)
 	{
         List<GUBPFrequencyHacker> FrequencyHackers = new List<GUBPFrequencyHacker>();
 
@@ -316,7 +316,7 @@ partial class GUBP
             }
         }
 
-		foreach(NodeInfo Node in Nodes)
+		foreach(BuildNode Node in Nodes)
 		{
 			Node.FrequencyShift = Node.Node.CISFrequencyQuantumShift(BranchConfig);
 
@@ -344,7 +344,7 @@ partial class GUBP
         return AltHostPlatform;
     }
 
-	void AddNodesForBranch(int CL, List<UnrealTargetPlatform> InitialHostPlatforms, bool bPreflightBuild, string PreflightMangleSuffix, out List<NodeInfo> AllNodes, out List<AggregateInfo> AllAggregates, ref int TimeQuantum)
+	void AddNodesForBranch(int CL, List<UnrealTargetPlatform> InitialHostPlatforms, bool bPreflightBuild, string PreflightMangleSuffix, out List<BuildNode> AllNodes, out List<AggregateNode> AllAggregates, ref int TimeQuantum)
 	{
 		string BranchName;
 		if (P4Enabled)
@@ -1202,11 +1202,11 @@ partial class GUBP
 		Dictionary<string, int> FrequencyOverrides = ApplyFrequencyBarriers(BranchConfig.GUBPNodes, BranchConfig.GUBPAggregates, BranchOptions.FrequencyBarriers);
 
 		// Convert the GUBPNodes and GUBPAggregates maps into lists
-		AllNodes = new List<NodeInfo>();
-		AllNodes.AddRange(BranchConfig.GUBPNodes.Values.Select(x => new NodeInfo{ Name = x.GetFullName(), Node = x, AgentMemoryRequirement = x.AgentMemoryRequirement(BranchConfig) }));
+		AllNodes = new List<BuildNode>();
+		AllNodes.AddRange(BranchConfig.GUBPNodes.Values.Select(x => new BuildNode{ Name = x.GetFullName(), Node = x, AgentMemoryRequirement = x.AgentMemoryRequirement(BranchConfig) }));
 
-		AllAggregates = new List<AggregateInfo>();
-		AllAggregates.AddRange(BranchConfig.GUBPAggregates.Values.Select(x => new AggregateInfo{ Name = x.GetFullName(), Node = x }));
+		AllAggregates = new List<AggregateNode>();
+		AllAggregates.AddRange(BranchConfig.GUBPAggregates.Values.Select(x => new AggregateNode{ Name = x.GetFullName(), Node = x }));
 
 		// Calculate the frequencies for each node
 		FindFrequenciesForNodes(BranchConfig, AllNodes, FrequencyOverrides);
@@ -1222,7 +1222,7 @@ partial class GUBP
 			{
 				NodeNames.Add(GameAggregatePromotableNode.StaticGetFullName(GameProj));
 			}
-			foreach (NodeInfo Node in AllNodes)
+			foreach (BuildNode Node in AllNodes)
 			{
 				if (Node.Node.GameNameIfAnyForTempStorage() == GameProj.GameName)
 				{
diff --git a/Engine/Source/Programs/AutomationTool/GUBP/Nodes/AggregateNode.cs b/Engine/Source/Programs/AutomationTool/GUBP/Nodes/AggregateNode.cs
new file mode 100644
--- /dev/null
+++ b/Engine/Source/Programs/AutomationTool/GUBP/Nodes/AggregateNode.cs
@@ -0,0 +1,16 @@
+﻿using System;
+using System.Collections.Generic;
+using System.Diagnostics;
+using System.Linq;
+using System.Text;
+
+namespace AutomationTool
+{
+	[DebuggerDisplay("{Name}")]
+	class AggregateNode
+	{
+		public string Name;
+		public GUBP.GUBPAggregateNode Node;
+		public BuildNode[] Dependencies;
+	}
+}
diff --git a/Engine/Source/Programs/AutomationTool/GUBP/Nodes/BuildNode.cs b/Engine/Source/Programs/AutomationTool/GUBP/Nodes/BuildNode.cs
--- a/Engine/Source/Programs/AutomationTool/GUBP/Nodes/BuildNode.cs
+++ b/Engine/Source/Programs/AutomationTool/GUBP/Nodes/BuildNode.cs
@@ -7,15 +7,15 @@ using System.Text;
 namespace AutomationTool
 {
 	[DebuggerDisplay("{Name}")]
-	class NodeInfo
+	class BuildNode
 	{
 		public string Name;
 		public GUBP.GUBPNode Node;
-		public NodeInfo[] Dependencies;
-		public NodeInfo[] PseudoDependencies;
-		public NodeInfo[] AllDirectDependencies;
-		public NodeInfo[] AllIndirectDependencies;
-		public NodeInfo[] ControllingTriggers;
+		public BuildNode[] Dependencies;
+		public BuildNode[] PseudoDependencies;
+		public BuildNode[] AllDirectDependencies;
+		public BuildNode[] AllIndirectDependencies;
+		public BuildNode[] ControllingTriggers;
 		public int FrequencyShift;
 		public bool IsComplete;
 		public string[] RecipientsForFailureEmails;
@@ -27,7 +27,7 @@ namespace AutomationTool
 			get { return String.Join(".", ControllingTriggers.Select(x => x.Name)); }
 		}
 
-		public bool DependsOn(NodeInfo Node)
+		public bool DependsOn(BuildNode Node)
 		{
 			return AllIndirectDependencies.Contains(Node);
 		}
