Separate out some internal EC stuff from graph setup.

diff --git a/Engine/Source/Programs/AutomationTool/GUBP/GUBP.cs b/Engine/Source/Programs/AutomationTool/GUBP/GUBP.cs
--- a/Engine/Source/Programs/AutomationTool/GUBP/GUBP.cs
+++ b/Engine/Source/Programs/AutomationTool/GUBP/GUBP.cs
@@ -990,38 +990,30 @@ public partial class GUBP : BuildCommand
         }
         return Result;
     }
-    List<string> GetECPropsForNode(BuildNode NodeToDo, string CLString, bool OnlyLateUpdates)
+    List<string> GetECPropsForNode(BuildNode NodeToDo)
     {
         List<string> ECProps = new List<string>();
 		ECProps.Add("FailEmails/" + NodeToDo.Name + "=" + String.Join(" ", NodeToDo.RecipientsForFailureEmails));
 	
-		if (!OnlyLateUpdates)
+		string AgentReq = NodeToDo.Node.ECAgentString();
+		if(ParseParamValue("AgentOverride") != "" && !NodeToDo.Node.GetFullName().Contains("OnMac"))
 		{
-			string AgentReq = NodeToDo.Node.ECAgentString();
-			if(ParseParamValue("AgentOverride") != "" && !NodeToDo.Node.GetFullName().Contains("OnMac"))
-			{
-				AgentReq = ParseParamValue("AgentOverride");
-			}
-			ECProps.Add(string.Format("AgentRequirementString/{0}={1}", NodeToDo.Name, AgentReq));
-			ECProps.Add(string.Format("RequiredMemory/{0}={1}", NodeToDo.Name, NodeToDo.AgentMemoryRequirement));
-			ECProps.Add(string.Format("Timeouts/{0}={1}", NodeToDo.Name, NodeToDo.Node.TimeoutInMinutes()));
-			ECProps.Add(string.Format("JobStepPath/{0}={1}", NodeToDo.Name, GetJobStepPath(NodeToDo)));
+			AgentReq = ParseParamValue("AgentOverride");
 		}
+		ECProps.Add(string.Format("AgentRequirementString/{0}={1}", NodeToDo.Name, AgentReq));
+		ECProps.Add(string.Format("RequiredMemory/{0}={1}", NodeToDo.Name, NodeToDo.AgentMemoryRequirement));
+		ECProps.Add(string.Format("Timeouts/{0}={1}", NodeToDo.Name, NodeToDo.Node.TimeoutInMinutes()));
+		ECProps.Add(string.Format("JobStepPath/{0}={1}", NodeToDo.Name, GetJobStepPath(NodeToDo)));
 		
         return ECProps;
     }
-
-    void UpdateECProps(BuildNode NodeToDo, string CLString)
+	
+    void UpdateECProps(BuildNode NodeToDo)
     {
         try
         {
             Log("Updating node props for node {0}", NodeToDo.Name);
-            List<string> Props = GetECPropsForNode(NodeToDo, CLString, true);
-            foreach (string Prop in Props)
-            {
-                string[] Parts = Prop.Split("=".ToCharArray());
-                RunECTool(String.Format("setProperty \"/myWorkflow/{0}\" \"{1}\"", Parts[0], Parts[1]), true);
-            }			
+            RunECTool(String.Format("setProperty \"/myWorkflow/FailEmails/{0}\" \"{1}\"", NodeToDo.Name, String.Join(" ", NodeToDo.RecipientsForFailureEmails)), true);
         }
         catch (Exception Ex)
         {
@@ -1786,15 +1778,54 @@ public partial class GUBP : BuildCommand
 	}
 
 	private void DoCommanderSetup(IEnumerable<BuildNode> AllNodes, IEnumerable<AggregateNode> AllAggregates, HashSet<BuildNode> NodesToDo, List<BuildNode> OrdereredToDo, int TimeIndex, int TimeQuantum, bool bSkipTriggers, bool bFake, bool bFakeEC, string CLString, BuildNode ExplicitTrigger, List<BuildNode> UnfinishedTriggers, string FakeFail, bool bPreflightBuild)
+	{
+		List<BuildNode> SortedNodes = TopologicalSort(new HashSet<BuildNode>(AllNodes), null, SubSort: false, DoNotConsiderCompletion: true);
+		Log("******* {0} GUBP Nodes", SortedNodes.Count);
+
+		List<BuildNode> FilteredOrdereredToDo = new List<BuildNode>();
+		using(TelemetryStopwatch StartFilterTimer = new TelemetryStopwatch("FilterNodes"))
+		{
+			// remove nodes that have unfinished triggers
+			foreach (BuildNode NodeToDo in OrdereredToDo)
+			{
+				BuildNode ControllingTrigger = (NodeToDo.ControllingTriggers.Length > 0)? NodeToDo.ControllingTriggers.Last() : null;
+				bool bNoUnfinishedTriggers = !UnfinishedTriggers.Contains(ControllingTrigger);
+
+				if (bNoUnfinishedTriggers)
+				{
+					// if we are triggering, then remove nodes that are not controlled by the trigger or are dependencies of this trigger
+					if (ExplicitTrigger != null)
+					{
+						if (ExplicitTrigger != NodeToDo && !ExplicitTrigger.DependsOn(NodeToDo) && !NodeToDo.DependsOn(ExplicitTrigger))
+						{
+							continue; // this wasn't on the chain related to the trigger we are triggering, so it is not relevant
+						}
+					}
+					if (bPreflightBuild && !bSkipTriggers && NodeToDo.Node.TriggerNode())
+					{
+						// in preflight builds, we are either skipping triggers (and running things downstream) or we just stop at triggers and don't make them available for triggering.
+						continue;
+					}
+					FilteredOrdereredToDo.Add(NodeToDo);
+				}
+			}
+		}
+		using(TelemetryStopwatch PrintNodesTimer = new TelemetryStopwatch("SetupCommanderPrint"))
+		{
+			Log("*********** EC Nodes, in order.");
+			PrintNodes(this, FilteredOrdereredToDo, AllAggregates, UnfinishedTriggers, TimeQuantum);
+		}
+
+		DoCommanderSetupInternal(AllNodes, AllAggregates, NodesToDo, FilteredOrdereredToDo, SortedNodes, TimeIndex, TimeQuantum, bSkipTriggers, bFake, bFakeEC, CLString, ExplicitTrigger, UnfinishedTriggers, FakeFail);
+	}
+
+	private void DoCommanderSetupInternal(IEnumerable<BuildNode> AllNodes, IEnumerable<AggregateNode> AllAggregates, HashSet<BuildNode> NodesToDo, List<BuildNode> OrdereredToDo, List<BuildNode> SortedNodes, int TimeIndex, int TimeQuantum, bool bSkipTriggers, bool bFake, bool bFakeEC, string CLString, BuildNode ExplicitTrigger, List<BuildNode> UnfinishedTriggers, string FakeFail)
 	{
 		Dictionary<string, string> FullNodeDependedOnBy = GetFullNodeDependedOnBy(NodesToDo);
 
 		List<AggregateNode> SeparatePromotables = FindPromotables(AllAggregates);
 		Dictionary<BuildNode, List<AggregateNode>> DependentPromotions = FindDependentPromotables(AllNodes, SeparatePromotables);
 
-		List<BuildNode> SortedNodes = TopologicalSort(new HashSet<BuildNode>(AllNodes), null, SubSort: false, DoNotConsiderCompletion: true);
-		Log("******* {0} GUBP Nodes", SortedNodes.Count);
-
 		Dictionary<BuildNode, int> FullNodeListSortKey = GetDisplayOrder(SortedNodes);
 
 		if (OrdereredToDo.Count == 0)
@@ -1837,40 +1868,6 @@ public partial class GUBP : BuildCommand
 			ECJobProps.Add("IsRoot=1");
 		}
 
-		using(TelemetryStopwatch StartFilterTimer = new TelemetryStopwatch("FilterNodes"))
-		{
-			List<BuildNode> FilteredOrdereredToDo = new List<BuildNode>();
-			// remove nodes that have unfinished triggers
-			foreach (BuildNode NodeToDo in OrdereredToDo)
-			{
-				BuildNode ControllingTrigger = (NodeToDo.ControllingTriggers.Length > 0)? NodeToDo.ControllingTriggers.Last() : null;
-				bool bNoUnfinishedTriggers = !UnfinishedTriggers.Contains(ControllingTrigger);
-
-				if (bNoUnfinishedTriggers)
-				{
-					// if we are triggering, then remove nodes that are not controlled by the trigger or are dependencies of this trigger
-					if (ExplicitTrigger != null)
-					{
-						if (ExplicitTrigger != NodeToDo && !ExplicitTrigger.DependsOn(NodeToDo) && !NodeToDo.DependsOn(ExplicitTrigger))
-						{
-							continue; // this wasn't on the chain related to the trigger we are triggering, so it is not relevant
-						}
-					}
-					if (bPreflightBuild && !bSkipTriggers && NodeToDo.Node.TriggerNode())
-					{
-						// in preflight builds, we are either skipping triggers (and running things downstream) or we just stop at triggers and don't make them available for triggering.
-						continue;
-					}
-					FilteredOrdereredToDo.Add(NodeToDo);
-				}
-			}
-			OrdereredToDo = FilteredOrdereredToDo;
-		}
-		using(TelemetryStopwatch PrintNodesTimer = new TelemetryStopwatch("SetupCommanderPrint"))
-		{
-			Log("*********** EC Nodes, in order.");
-			PrintNodes(this, OrdereredToDo, AllAggregates, UnfinishedTriggers, TimeQuantum);
-		}
 		// here we are just making sure everything before the explicit trigger is completed.
 		if (ExplicitTrigger != null)
 		{
@@ -1960,7 +1957,7 @@ public partial class GUBP : BuildCommand
 			{
 				if (!NodeToDo.IsComplete) // if something is already finished, we don't put it into EC  
 				{
-					List<string> NodeProps = GetECPropsForNode(NodeToDo, CLString, false);
+					List<string> NodeProps = GetECPropsForNode(NodeToDo);
 					ECProps.AddRange(NodeProps);
 
 					bool Sticky = NodeToDo.Node.IsSticky();
@@ -2461,7 +2458,7 @@ public partial class GUBP : BuildCommand
 						}
 						using(TelemetryStopwatch UpdateECPropsStopwatch = new TelemetryStopwatch("UpdateECProps"))
 						{
-							UpdateECProps(NodeToDo, CLString);
+							UpdateECProps(NodeToDo);
 						}
                         
 						if (IsBuildMachine)
@@ -2531,7 +2528,7 @@ public partial class GUBP : BuildCommand
 					}
 					using(TelemetryStopwatch UpdateECPropsStopwatch = new TelemetryStopwatch("UpdateECProps"))
 					{
-						UpdateECProps(NodeToDo, CLString);
+						UpdateECProps(NodeToDo);
 					}
                     
 					if (IsBuildMachine)
