Added support for shared packs to builder

diff --git a/Engine/Source/Programs/AutomationTool/GUBP/LegacyNodes.cs b/Engine/Source/Programs/AutomationTool/GUBP/LegacyNodes.cs
--- a/Engine/Source/Programs/AutomationTool/GUBP/LegacyNodes.cs
+++ b/Engine/Source/Programs/AutomationTool/GUBP/LegacyNodes.cs
@@ -1159,16 +1159,9 @@ partial class GUBP
 
 		public static bool IsFeaturePack(BranchInfo.BranchUProject InGameProj)
 		{
-			// No obvious way to store this in the project options; it's a property of non-code projects too.
-			if(InGameProj.GameName == "StarterContent" || InGameProj.GameName == "MobileStarterContent" || InGameProj.GameName.StartsWith("FP_"))
-			{
-				return true;
-			}
-			if(InGameProj.GameName.StartsWith("TP_"))
-			{
-				return CommandUtils.FileExists(CommandUtils.CombinePaths(CommandUtils.GetDirectoryName(InGameProj.FilePath), "contents.txt"));
-			}
-			return false;
+			bool bHasContents = CommandUtils.FileExists(CommandUtils.CombinePaths(CommandUtils.GetDirectoryName(InGameProj.FilePath), "contents.txt"));
+			bool bHasManifest = CommandUtils.FileExists(CommandUtils.CombinePaths(CommandUtils.GetDirectoryName(InGameProj.FilePath), "manifest.json"));
+			return bHasContents && bHasManifest;
 		}
 
 		public static UnrealTargetPlatform GetDefaultBuildPlatform(List<UnrealTargetPlatform> HostPlatforms)
diff --git a/Engine/Source/Programs/AutomationTool/Scripts/RocketBuild.Automation.cs b/Engine/Source/Programs/AutomationTool/Scripts/RocketBuild.Automation.cs
--- a/Engine/Source/Programs/AutomationTool/Scripts/RocketBuild.Automation.cs
+++ b/Engine/Source/Programs/AutomationTool/Scripts/RocketBuild.Automation.cs
@@ -67,7 +67,23 @@ namespace Rocket
 			"TP_VehicleAdvBP",
 			"StarterContent",
 			"MobileStarterContent",
-
+			"GeometryHigh",		//Shared packs
+			"GeometryStandard",
+			"MannequinHigh",
+			"MannequinStandard",
+			"RollingHigh",
+			"RollingStandard",
+			"2DSideScrollerStandard",
+			"FirstPersonStandard",
+			"FlyingStandard",
+			"PuzzleStandard",
+			"RollingStandard",
+			"RollingHigh",
+			"SideScrollerStandard",
+			"ThirdPersonStandard",
+			"TwinStickStandard",
+			"VehicleStandard",
+			"VehicleAdvStandard",
 		};
 
 		public RocketBuild()
