Move all EC functionality into ElectricCommander.cs. Also tidy up parsing of node parameters (and remove special -NoLinux handling for removing cross compile nodes - will revisit and fix in the nodes themselves if still needed).

diff --git a/Engine/Source/Programs/AutomationTool/GUBP/ElectricCommander.cs b/Engine/Source/Programs/AutomationTool/GUBP/ElectricCommander.cs
--- a/Engine/Source/Programs/AutomationTool/GUBP/ElectricCommander.cs
+++ b/Engine/Source/Programs/AutomationTool/GUBP/ElectricCommander.cs
@@ -1,9 +1,766 @@
 ﻿using System;
 using System.Collections.Generic;
+using System.IO;
 using System.Linq;
 using System.Text;
 using System.Xml;
 
+namespace AutomationTool
+{
+	class ElectricCommander
+	{
+		BuildCommand Command;
+
+		public ElectricCommander(BuildCommand InCommand)
+		{
+			Command = InCommand;
+		}
+
+		public string RunECTool(string Args, bool bQuiet = false)
+		{
+			if (Command.ParseParam("FakeEC"))
+			{
+				CommandUtils.LogWarning("***** Would have ran ectool {0}", Args);
+				return "We didn't actually run ectool";
+			}
+			else
+			{
+				CommandUtils.ERunOptions Opts = CommandUtils.ERunOptions.Default;
+				if (bQuiet)
+				{
+					Opts = (Opts & ~CommandUtils.ERunOptions.AllowSpew) | CommandUtils.ERunOptions.NoLoggingOfRunCommand;
+				}
+				return CommandUtils.RunAndLog("ectool", "--timeout 900 " + Args, Options: Opts);
+			}
+		}
+
+		static void WriteECPerl(List<string> Args)
+		{
+			Args.Add("$batch->submit();");
+			string ECPerlFile = CommandUtils.CombinePaths(CommandUtils.CmdEnv.LogFolder, "jobsteps.pl");
+			CommandUtils.WriteAllLines_NoExceptions(ECPerlFile, Args.ToArray());
+		}
+
+		string GetPropertyFromStep(string PropertyPath)
+		{
+			string Property = "";
+			Property = RunECTool("getProperty \"" + PropertyPath + "\"");
+			Property = Property.TrimEnd('\r', '\n');
+			return Property;
+		}
+
+		static string GetJobStepPath(BuildNode Dep)
+		{
+			if (Dep.Node.AgentSharingGroup == "")
+			{
+				return "jobSteps[" + Dep.Name + "]";
+			}
+			else
+			{
+				return "jobSteps[" + Dep.Node.AgentSharingGroup + "]/jobSteps[" + Dep.Name + "]";
+			}
+		}
+		static string GetJobStep(string ParentPath, BuildNode Dep)
+		{
+			return ParentPath + "/" + GetJobStepPath(Dep);
+		}
+
+		List<string> GetECPropsForNode(BuildNode NodeToDo)
+		{
+			List<string> ECProps = new List<string>();
+			ECProps.Add("FailEmails/" + NodeToDo.Name + "=" + String.Join(" ", NodeToDo.RecipientsForFailureEmails));
+	
+			string AgentReq = NodeToDo.Node.ECAgentString();
+			if(Command.ParseParamValue("AgentOverride") != "" && !NodeToDo.Node.GetFullName().Contains("OnMac"))
+			{
+				AgentReq = Command.ParseParamValue("AgentOverride");
+			}
+			ECProps.Add(string.Format("AgentRequirementString/{0}={1}", NodeToDo.Name, AgentReq));
+			ECProps.Add(string.Format("RequiredMemory/{0}={1}", NodeToDo.Name, NodeToDo.AgentMemoryRequirement));
+			ECProps.Add(string.Format("Timeouts/{0}={1}", NodeToDo.Name, NodeToDo.Node.TimeoutInMinutes()));
+			ECProps.Add(string.Format("JobStepPath/{0}={1}", NodeToDo.Name, GetJobStepPath(NodeToDo)));
+		
+			return ECProps;
+		}
+	
+		public void UpdateECProps(BuildNode NodeToDo)
+		{
+			try
+			{
+				CommandUtils.Log("Updating node props for node {0}", NodeToDo.Name);
+				RunECTool(String.Format("setProperty \"/myWorkflow/FailEmails/{0}\" \"{1}\"", NodeToDo.Name, String.Join(" ", NodeToDo.RecipientsForFailureEmails)), true);
+			}
+			catch (Exception Ex)
+			{
+				CommandUtils.Log(System.Diagnostics.TraceEventType.Warning, "Failed to UpdateECProps.");
+				CommandUtils.Log(System.Diagnostics.TraceEventType.Warning, LogUtils.FormatException(Ex));
+			}
+		}
+		public void UpdateECBuildTime(BuildNode NodeToDo, double BuildDuration)
+		{
+			try
+			{
+				CommandUtils.Log("Updating duration prop for node {0}", NodeToDo.Name);
+				RunECTool(String.Format("setProperty \"/myWorkflow/NodeDuration/{0}\" \"{1}\"", NodeToDo.Name, BuildDuration.ToString()));
+				RunECTool(String.Format("setProperty \"/myJobStep/NodeDuration\" \"{0}\"", BuildDuration.ToString()));
+			}
+			catch (Exception Ex)
+			{
+				CommandUtils.Log(System.Diagnostics.TraceEventType.Warning, "Failed to UpdateECBuildTime.");
+				CommandUtils.Log(System.Diagnostics.TraceEventType.Warning, LogUtils.FormatException(Ex));
+			}
+		}
+
+		public void SaveStatus(BuildNode NodeToDo, string Suffix, string NodeStoreName, bool bSaveSharedTempStorage, string GameNameIfAny, string JobStepIDForFailure = null)
+		{
+			string Contents = "Just a status record: " + Suffix;
+			if (!String.IsNullOrEmpty(JobStepIDForFailure) && CommandUtils.IsBuildMachine)
+			{
+				try
+				{
+					Contents = RunECTool(String.Format("getProperties --jobStepId {0} --recurse 1", JobStepIDForFailure), true);
+				}
+				catch (Exception Ex)
+				{
+					CommandUtils.Log(System.Diagnostics.TraceEventType.Warning, "Failed to get properties for jobstep to save them.");
+					CommandUtils.Log(System.Diagnostics.TraceEventType.Warning, LogUtils.FormatException(Ex));
+				}
+			}
+			string RecordOfSuccess = CommandUtils.CombinePaths(CommandUtils.CmdEnv.LocalRoot, "Engine", "Saved", "Logs", NodeToDo.Name + Suffix +".log");
+			CommandUtils.CreateDirectory(Path.GetDirectoryName(RecordOfSuccess));
+			CommandUtils.WriteAllText(RecordOfSuccess, Contents);		
+			TempStorage.StoreToTempStorage(CommandUtils.CmdEnv, NodeStoreName + Suffix, new List<string> { RecordOfSuccess }, !bSaveSharedTempStorage, GameNameIfAny);		
+		}
+
+		public void UpdateEmailProperties(BuildNode NodeToDo, int LastSucceededCL, string FailedString, string FailCauserEMails, string EMailNote, bool SendSuccessForGreenAfterRed)
+		{
+			RunECTool(String.Format("setProperty \"/myWorkflow/LastGreen/{0}\" \"{1}\"", NodeToDo.Name, LastSucceededCL), true);
+			RunECTool(String.Format("setProperty \"/myWorkflow/LastGreen/{0}\" \"{1}\"", NodeToDo.Name, FailedString), true);
+			RunECTool(String.Format("setProperty \"/myWorkflow/FailCausers/{0}\" \"{1}\"", NodeToDo.Name, FailCauserEMails));
+			RunECTool(String.Format("setProperty \"/myWorkflow/EmailNotes/{0}\" \"{1}\"", NodeToDo.Name, EMailNote));
+			{
+				HashSet<string> Emails = new HashSet<string>(NodeToDo.RecipientsForFailureEmails);
+				if (Command.ParseParam("CIS") && !NodeToDo.Node.SendSuccessEmail() && !NodeToDo.Node.TriggerNode() && NodeToDo.AddSubmittersToFailureEmails)
+				{
+					Emails.UnionWith(FailCauserEMails.Split(new char[] { ' ' }, StringSplitOptions.RemoveEmptyEntries));
+				}
+				RunECTool(String.Format("setProperty \"/myWorkflow/FailEmails/{0}\" \"{1}\"", NodeToDo.Name, String.Join(" ", Emails)));
+			}
+			if (NodeToDo.Node.SendSuccessEmail() || SendSuccessForGreenAfterRed)
+			{
+				RunECTool(String.Format("setProperty \"/myWorkflow/SendSuccessEmail/{0}\" \"{1}\"", NodeToDo.Name, "1"));
+			}
+			else
+			{
+				RunECTool(String.Format("setProperty \"/myWorkflow/SendSuccessEmail/{0}\" \"{1}\"", NodeToDo.Name, "0"));
+			}
+		}
+
+		public void DoCommanderSetup(IEnumerable<BuildNode> AllNodes, IEnumerable<AggregateNode> AllAggregates, HashSet<BuildNode> NodesToDo, List<BuildNode> OrdereredToDo, List<BuildNode> SortedNodes, int TimeIndex, int TimeQuantum, bool bSkipTriggers, bool bFake, bool bFakeEC, string CLString, BuildNode ExplicitTrigger, List<BuildNode> UnfinishedTriggers, string FakeFail)
+		{
+			Dictionary<string, string> FullNodeDependedOnBy = GetFullNodeDependedOnBy(NodesToDo);
+
+			List<AggregateNode> SeparatePromotables = FindPromotables(AllAggregates);
+			Dictionary<BuildNode, List<AggregateNode>> DependentPromotions = FindDependentPromotables(AllNodes, SeparatePromotables);
+
+			Dictionary<BuildNode, int> FullNodeListSortKey = GetDisplayOrder(SortedNodes);
+
+			if (OrdereredToDo.Count == 0)
+			{
+				throw new AutomationException("No nodes to do!");
+			}
+			List<string> ECProps = new List<string>();
+			ECProps.Add(String.Format("TimeIndex={0}", TimeIndex));
+			foreach (BuildNode Node in SortedNodes)
+			{
+				ECProps.Add(string.Format("AllNodes/{0}={1}", Node.Name, GetNodeForAllNodesProperty(Node, TimeQuantum)));
+			}
+			foreach (BuildNode Node in SortedNodes)
+			{
+				ECProps.Add(string.Format("DirectDependencies/{0}={1}", Node.Name, String.Join(" ", Node.AllDirectDependencies.Select(x => x.Name))));
+			}
+			foreach (KeyValuePair<BuildNode, int> NodePair in FullNodeListSortKey)
+			{
+				ECProps.Add(string.Format("SortKey/{0}={1}", NodePair.Key.Name, NodePair.Value));
+			}
+			foreach (KeyValuePair<string, string> NodePair in FullNodeDependedOnBy)
+			{
+				ECProps.Add(string.Format("DependedOnBy/{0}={1}", NodePair.Key, NodePair.Value));
+			}
+			foreach (KeyValuePair<BuildNode, List<AggregateNode>> NodePair in DependentPromotions)
+			{
+				ECProps.Add(string.Format("DependentPromotions/{0}={1}", NodePair.Key.Name, String.Join(" ", NodePair.Value.Select(x => x.Name))));
+			}
+			foreach (AggregateNode Node in SeparatePromotables)
+			{
+				ECProps.Add(string.Format("PossiblePromotables/{0}={1}", Node.Name, ""));
+			}
+			List<string> ECJobProps = new List<string>();
+			if (ExplicitTrigger != null)
+			{
+				ECJobProps.Add("IsRoot=0");
+			}
+			else
+			{
+				ECJobProps.Add("IsRoot=1");
+			}
+
+			// here we are just making sure everything before the explicit trigger is completed.
+			if (ExplicitTrigger != null)
+			{
+				foreach (BuildNode NodeToDo in OrdereredToDo)
+				{
+					if (!NodeToDo.IsComplete && NodeToDo != ExplicitTrigger && !NodeToDo.DependsOn(ExplicitTrigger)) // if something is already finished, we don't put it into EC
+					{
+						throw new AutomationException("We are being asked to process node {0}, however, this is an explicit trigger {1}, so everything before it should already be handled. It seems likely that you waited too long to run the trigger. You will have to do a new build from scratch.", NodeToDo.Name, ExplicitTrigger.Name);
+					}
+				}
+			}
+
+			BuildNode LastSticky = null;
+			bool HitNonSticky = false;
+			bool bHaveECNodes = false;
+			List<string> StepList = new List<string>();
+			StepList.Add("use strict;");
+			StepList.Add("use diagnostics;");
+			StepList.Add("use ElectricCommander();");
+			StepList.Add("my $ec = new ElectricCommander;");
+			StepList.Add("$ec->setTimeout(600);");
+			StepList.Add("my $batch = $ec->newBatch(\"serial\");");
+			// sticky nodes are ones that we run on the main agent. We run then first and they must not be intermixed with parallel jobs
+			foreach (BuildNode NodeToDo in OrdereredToDo)
+			{
+				if (!NodeToDo.IsComplete) // if something is already finished, we don't put it into EC
+				{
+					bHaveECNodes = true;
+					if (NodeToDo.Node.IsSticky())
+					{
+						LastSticky = NodeToDo;
+						if (HitNonSticky && !bSkipTriggers)
+						{
+							throw new AutomationException("Sticky and non-sticky jobs did not sort right.");
+						}
+					}
+					else
+					{
+						HitNonSticky = true;
+					}
+				}
+			}
+
+			List<string> FakeECArgs = new List<string>();
+			using(CommandUtils.TelemetryStopwatch PerlOutputStopwatch = new CommandUtils.TelemetryStopwatch("PerlOutput"))
+			{
+				string ParentPath = Command.ParseParamValue("ParentPath");
+				string BaseArgs = String.Format("$batch->createJobStep({{parentPath => '{0}'", ParentPath);
+
+				bool bHasNoop = false;
+				if (LastSticky == null && bHaveECNodes)
+				{
+					// if we don't have any sticky nodes and we have other nodes, we run a fake noop just to release the resource 
+					string Args = String.Format("{0}, subprocedure => 'GUBP_UAT_Node', parallel => '0', jobStepName => 'Noop', actualParameter => [{{actualParameterName => 'NodeName', value => 'Noop'}}, {{actualParameterName => 'Sticky', value =>'1' }}], releaseMode => 'release'}});", BaseArgs);
+					StepList.Add(Args);
+					bHasNoop = true;
+				}
+
+				Dictionary<string, List<BuildNode>> AgentGroupChains = new Dictionary<string, List<BuildNode>>();
+				List<BuildNode> StickyChain = new List<BuildNode>();
+				foreach (BuildNode NodeToDo in OrdereredToDo)
+				{
+					if (!NodeToDo.IsComplete) // if something is already finished, we don't put it into EC  
+					{
+						string MyAgentGroup = NodeToDo.Node.AgentSharingGroup;
+						if (MyAgentGroup != "")
+						{
+							if (!AgentGroupChains.ContainsKey(MyAgentGroup))
+							{
+								AgentGroupChains.Add(MyAgentGroup, new List<BuildNode> { NodeToDo });
+							}
+							else
+							{
+								AgentGroupChains[MyAgentGroup].Add(NodeToDo);
+							}
+						}
+					}
+					if (NodeToDo.Node.IsSticky())
+					{
+						if (!StickyChain.Contains(NodeToDo))
+						{
+							StickyChain.Add(NodeToDo);
+						}
+					}
+				}
+				foreach (BuildNode NodeToDo in OrdereredToDo)
+				{
+					if (!NodeToDo.IsComplete) // if something is already finished, we don't put it into EC  
+					{
+						List<string> NodeProps = GetECPropsForNode(NodeToDo);
+						ECProps.AddRange(NodeProps);
+
+						bool Sticky = NodeToDo.Node.IsSticky();
+						bool DoParallel = !Sticky;
+						if (NodeToDo.Node.ECProcedure() == "GUBP_UAT_Node_Parallel_AgentShare_Editor")
+						{
+							DoParallel = true;
+						}
+						if (Sticky && NodeToDo.Node.ECAgentString() != "")
+						{
+							throw new AutomationException("Node {1} is sticky but has agent requirements.", NodeToDo.Name);
+						}
+						string Procedure = NodeToDo.Node.ECProcedure();
+						if (NodeToDo.Node.IsSticky() && NodeToDo == LastSticky)
+						{
+							Procedure = Procedure + "_Release";
+						}
+						string Args = String.Format("{0}, subprocedure => '{1}', parallel => '{2}', jobStepName => '{3}', actualParameter => [{{actualParameterName => 'NodeName', value =>'{4}'}}",
+							BaseArgs, Procedure, DoParallel ? 1 : 0, NodeToDo.Name, NodeToDo.Name);
+						string ProcedureParams = NodeToDo.Node.ECProcedureParams();
+						if (!String.IsNullOrEmpty(ProcedureParams))
+						{
+							Args = Args + ProcedureParams;
+						}
+
+						if ((Procedure == "GUBP_UAT_Trigger" || Procedure == "GUBP_Hardcoded_Trigger") && NodeToDo.RecipientsForFailureEmails.Length > 0)
+						{
+							Args = Args + ", {actualParameterName => 'EmailsForTrigger', value => \'" + String.Join(" ", NodeToDo.RecipientsForFailureEmails) + "\'}";
+						}
+						Args = Args + "]";
+
+						List<BuildNode> UncompletedEcDeps = new List<BuildNode>();
+						{
+							foreach (BuildNode Dep in NodeToDo.AllDirectDependencies)
+							{
+								if (!Dep.IsComplete && OrdereredToDo.Contains(Dep)) // if something is already finished, we don't put it into EC
+								{
+									if (OrdereredToDo.IndexOf(Dep) > OrdereredToDo.IndexOf(NodeToDo))
+									{
+										throw new AutomationException("Topological sort error, node {0} has a dependency of {1} which sorted after it.", NodeToDo.Name, Dep.Name);
+									}
+									UncompletedEcDeps.Add(Dep);
+								}
+							}
+						}
+
+						string PreCondition = GetPreConditionForNode(OrdereredToDo, ParentPath, bHasNoop, AgentGroupChains, StickyChain, NodeToDo, UncompletedEcDeps);
+						string RunCondition = GetRunConditionForNode(UncompletedEcDeps, ParentPath);
+
+						string MyAgentGroup = NodeToDo.Node.AgentSharingGroup;
+						bool bDoNestedJobstep = false;
+						bool bDoFirstNestedJobstep = false;
+
+						string NodeParentPath = ParentPath;
+						if (MyAgentGroup != "")
+						{
+							bDoNestedJobstep = true;
+							NodeParentPath = ParentPath + "/jobSteps[" + MyAgentGroup + "]";
+
+							List<BuildNode> MyChain = AgentGroupChains[MyAgentGroup];
+							int MyIndex = MyChain.IndexOf(NodeToDo);
+							if (MyIndex <= 0)
+							{
+								bDoFirstNestedJobstep = bDoNestedJobstep;
+							}
+						}
+
+						if (bDoNestedJobstep)
+						{
+							if (bDoFirstNestedJobstep)
+							{
+								{
+									string NestArgs = String.Format("$batch->createJobStep({{parentPath => '{0}', jobStepName => '{1}', parallel => '1'",
+										ParentPath, MyAgentGroup);
+									if (!String.IsNullOrEmpty(PreCondition))
+									{
+										NestArgs = NestArgs + ", precondition => " + PreCondition;
+									}
+									NestArgs = NestArgs + "});";
+									StepList.Add(NestArgs);
+								}
+								{
+									string NestArgs = String.Format("$batch->createJobStep({{parentPath => '{0}/jobSteps[{1}]', jobStepName => '{2}_GetPool', subprocedure => 'GUBP{3}_AgentShare_GetPool', parallel => '1', actualParameter => [{{actualParameterName => 'AgentSharingGroup', value => '{4}'}}, {{actualParameterName => 'NodeName', value => '{5}'}}]",
+										ParentPath, MyAgentGroup, MyAgentGroup, NodeToDo.Node.ECProcedureInfix(), MyAgentGroup, NodeToDo.Name);
+									if (!String.IsNullOrEmpty(PreCondition))
+									{
+										NestArgs = NestArgs + ", precondition => " + PreCondition;
+									}
+									NestArgs = NestArgs + "});";
+									StepList.Add(NestArgs);
+								}
+								{
+									string NestArgs = String.Format("$batch->createJobStep({{parentPath => '{0}/jobSteps[{1}]', jobStepName => '{2}_GetAgent', subprocedure => 'GUBP{3}_AgentShare_GetAgent', parallel => '1', exclusiveMode => 'call', resourceName => '{4}', actualParameter => [{{actualParameterName => 'AgentSharingGroup', value => '{5}'}}, {{actualParameterName => 'NodeName', value=> '{6}'}}]",
+										ParentPath, MyAgentGroup, MyAgentGroup, NodeToDo.Node.ECProcedureInfix(),
+										String.Format("$[/myJob/jobSteps[{0}]/ResourcePool]", MyAgentGroup),
+										MyAgentGroup, NodeToDo.Name);
+									{
+										NestArgs = NestArgs + ", precondition  => ";
+										NestArgs = NestArgs + "\"\\$\" . \"[/javascript if(";
+										NestArgs = NestArgs + "getProperty('" + ParentPath + "/jobSteps[" + MyAgentGroup + "]/jobSteps[" + MyAgentGroup + "_GetPool]/status') == 'completed'";
+										NestArgs = NestArgs + ") true;]\"";
+									}
+									NestArgs = NestArgs + "});";
+									StepList.Add(NestArgs);
+								}
+								{
+									PreCondition = "\"\\$\" . \"[/javascript if(";
+									PreCondition = PreCondition + "getProperty('" + ParentPath + "/jobSteps[" + MyAgentGroup + "]/jobSteps[" + MyAgentGroup + "_GetAgent]/status') == 'completed'";
+									PreCondition = PreCondition + ") true;]\"";
+								}
+							}
+							Args = Args.Replace(String.Format("parentPath => '{0}'", ParentPath), String.Format("parentPath => '{0}'", NodeParentPath));
+							Args = Args.Replace("UAT_Node_Parallel_AgentShare", "UAT_Node_Parallel_AgentShare3");
+						}
+
+						if (!String.IsNullOrEmpty(PreCondition))
+						{
+							Args = Args + ", precondition => " + PreCondition;
+						}
+						if (!String.IsNullOrEmpty(RunCondition))
+						{
+							Args = Args + ", condition => " + RunCondition;
+						}
+		#if false
+							// this doesn't work because it includes precondition time
+							if (GUBPNodes[NodeToDo].TimeoutInMinutes() > 0)
+							{
+								Args = Args + String.Format(" --timeLimitUnits minutes --timeLimit {0}", GUBPNodes[NodeToDo].TimeoutInMinutes());
+							}
+		#endif
+						if (Sticky && NodeToDo == LastSticky)
+						{
+							Args = Args + ", releaseMode => 'release'";
+						}
+						Args = Args + "});";
+						StepList.Add(Args);
+						if (bFakeEC &&
+							!UnfinishedTriggers.Contains(NodeToDo) &&
+							(NodeToDo.Node.ECProcedure().StartsWith("GUBP_UAT_Node") || NodeToDo.Node.ECProcedure().StartsWith("GUBP_Mac_UAT_Node")) // other things we really can't test
+							) // unfinished triggers are never run directly by EC, rather it does another job setup
+						{
+							string Arg = String.Format("gubp -Node={0} -FakeEC {1} {2} {3} {4} {5}",
+								NodeToDo.Name,
+								bFake ? "-Fake" : "",
+								Command.ParseParam("AllPlatforms") ? "-AllPlatforms" : "",
+								Command.ParseParam("UnfinishedTriggersFirst") ? "-UnfinishedTriggersFirst" : "",
+								Command.ParseParam("UnfinishedTriggersParallel") ? "-UnfinishedTriggersParallel" : "",
+								Command.ParseParam("WithMac") ? "-WithMac" : ""
+								);
+
+							string Node = Command.ParseParamValue("-Node");
+							if (!String.IsNullOrEmpty(Node))
+							{
+								Arg = Arg + " -Node=" + Node;
+							}
+							if (!String.IsNullOrEmpty(FakeFail))
+							{
+								Arg = Arg + " -FakeFail=" + FakeFail;
+							}
+							FakeECArgs.Add(Arg);
+						}
+
+						if (MyAgentGroup != "" && !bDoNestedJobstep)
+						{
+							List<BuildNode> MyChain = AgentGroupChains[MyAgentGroup];
+							int MyIndex = MyChain.IndexOf(NodeToDo);
+							if (MyIndex == MyChain.Count - 1)
+							{
+								string RelPreCondition = "\"\\$\" . \"[/javascript if(";
+								// this runs "parallel", but we a precondition to serialize it
+								RelPreCondition = RelPreCondition + "getProperty('" + ParentPath + "/jobSteps[" + NodeToDo.Name + "]/status') == 'completed'";
+								RelPreCondition = RelPreCondition + ") true;]\"";
+								// we need to release the resource
+								string RelArgs = String.Format("{0}, subprocedure => 'GUBP_Release_AgentShare', parallel => '1', jobStepName => 'Release_{1}', actualParameter => [{{actualParameterName => 'AgentSharingGroup', valued => '{2}'}}], releaseMode => 'release', precondition => '{3}'",
+									BaseArgs, MyAgentGroup, MyAgentGroup, RelPreCondition);
+								StepList.Add(RelArgs);
+							}
+						}
+					}
+				}
+				WriteECPerl(StepList);
+			}
+			bool bHasTests = OrdereredToDo.Any(x => x.Node.IsTest());
+			RunECTool(String.Format("setProperty \"/myWorkflow/HasTests\" \"{0}\"", bHasTests));
+			{
+				ECProps.Add("GUBP_LoadedProps=1");
+				string BranchDefFile = CommandUtils.CombinePaths(CommandUtils.CmdEnv.LogFolder, "BranchDef.properties");
+				CommandUtils.WriteAllLines(BranchDefFile, ECProps.ToArray());
+				RunECTool(String.Format("setProperty \"/myWorkflow/BranchDefFile\" \"{0}\"", BranchDefFile.Replace("\\", "\\\\")));
+			}
+			{
+				ECProps.Add("GUBP_LoadedJobProps=1");
+				string BranchJobDefFile = CommandUtils.CombinePaths(CommandUtils.CmdEnv.LogFolder, "BranchJobDef.properties");
+				CommandUtils.WriteAllLines(BranchJobDefFile, ECProps.ToArray());
+				RunECTool(String.Format("setProperty \"/myJob/BranchJobDefFile\" \"{0}\"", BranchJobDefFile.Replace("\\", "\\\\")));
+			}
+		}
+
+		private string GetPreConditionForNode(List<BuildNode> OrdereredToDo, string PreconditionParentPath, bool bHasNoop, Dictionary<string, List<BuildNode>> AgentGroupChains, List<BuildNode> StickyChain, BuildNode NodeToDo, List<BuildNode> UncompletedEcDeps)
+		{
+			List<BuildNode> PreConditionUncompletedEcDeps = new List<BuildNode>();
+			if(NodeToDo.Node.AgentSharingGroup == "")
+			{
+				PreConditionUncompletedEcDeps = new List<BuildNode>(UncompletedEcDeps);
+			}
+			else 
+			{
+				List<BuildNode> MyChain = AgentGroupChains[NodeToDo.Node.AgentSharingGroup];
+				int MyIndex = MyChain.IndexOf(NodeToDo);
+				if (MyIndex > 0)
+				{
+					PreConditionUncompletedEcDeps.Add(MyChain[MyIndex - 1]);
+				}
+				else
+				{
+					// to avoid idle agents (and also EC doesn't actually reserve our agent!), we promote all dependencies to the first one
+					foreach (BuildNode Chain in MyChain)
+					{
+						foreach (BuildNode Dep in Chain.AllDirectDependencies)
+						{
+							if (!Dep.IsComplete && OrdereredToDo.Contains(Dep)) // if something is already finished, we don't put it into EC
+							{
+								if (OrdereredToDo.IndexOf(Dep) > OrdereredToDo.IndexOf(Chain))
+								{
+									throw new AutomationException("Topological sort error, node {0} has a dependency of {1} which sorted after it.", Chain.Name, Dep.Name);
+								}
+								if (!MyChain.Contains(Dep) && !PreConditionUncompletedEcDeps.Contains(Dep))
+								{
+									PreConditionUncompletedEcDeps.Add(Dep);
+								}
+							}
+						}
+					}
+				}
+			}
+			if (NodeToDo.Node.IsSticky())
+			{
+				List<BuildNode> MyChain = StickyChain;
+				int MyIndex = MyChain.IndexOf(NodeToDo);
+				if (MyIndex > 0)
+				{
+					if (!PreConditionUncompletedEcDeps.Contains(MyChain[MyIndex - 1]) && !MyChain[MyIndex - 1].IsComplete)
+					{
+						PreConditionUncompletedEcDeps.Add(MyChain[MyIndex - 1]);
+					}
+				}
+				else
+				{
+					foreach (BuildNode Dep in NodeToDo.AllDirectDependencies)
+					{
+						if (!Dep.IsComplete && OrdereredToDo.Contains(Dep)) // if something is already finished, we don't put it into EC
+						{
+							if (OrdereredToDo.IndexOf(Dep) > OrdereredToDo.IndexOf(NodeToDo))
+							{
+								throw new AutomationException("Topological sort error, node {0} has a dependency of {1} which sorted after it.", NodeToDo.Name, Dep.Name);
+							}
+							if (!MyChain.Contains(Dep) && !PreConditionUncompletedEcDeps.Contains(Dep))
+							{
+								PreConditionUncompletedEcDeps.Add(Dep);
+							}
+						}
+					}
+				}
+			}
+
+			List<string> JobStepNames = new List<string>();
+			if (bHasNoop && PreConditionUncompletedEcDeps.Count == 0)
+			{
+				JobStepNames.Add(PreconditionParentPath + "/jobSteps[Noop]");
+			}
+			else
+			{
+				JobStepNames.AddRange(PreConditionUncompletedEcDeps.Select(x => GetJobStep(PreconditionParentPath, x)));
+			}
+
+			string PreCondition = "";
+			if (JobStepNames.Count > 0)
+			{
+				PreCondition = String.Format("\"\\$\" . \"[/javascript if({0}) true;]\"", String.Join(" && ", JobStepNames.Select(x => String.Format("getProperty('{0}/status') == 'completed'", x))));
+			}
+			return PreCondition;
+		}
+
+		private static string GetNodeForAllNodesProperty(BuildNode Node, int TimeQuantum)
+		{
+			string Note = Node.ControllingTriggerDotName;
+			if (Note == "")
+			{
+				int Minutes = TimeQuantum << Node.FrequencyShift;
+				if (Minutes == TimeQuantum)
+				{
+					Note = "always";
+				}
+				else if (Minutes < 60)
+				{
+					Note = String.Format("{0}m", Minutes);
+				}
+				else
+				{
+					Note = String.Format("{0}h{1}m", Minutes / 60, Minutes % 60);
+				}
+			}
+			return Note;
+		}
+
+		private static string GetRunConditionForNode(List<BuildNode> UncompletedEcDeps, string PreconditionParentPath)
+		{
+			string RunCondition = "";
+			if (UncompletedEcDeps.Count > 0)
+			{
+				RunCondition = "\"\\$\" . \"[/javascript if(";
+				int Index = 0;
+				foreach (BuildNode Dep in UncompletedEcDeps)
+				{
+					RunCondition = RunCondition + "((\'\\$\" . \"[" + GetJobStep(PreconditionParentPath, Dep) + "/outcome]\' == \'success\') || ";
+					RunCondition = RunCondition + "(\'\\$\" . \"[" + GetJobStep(PreconditionParentPath, Dep) + "/outcome]\' == \'warning\'))";
+
+					Index++;
+					if (Index != UncompletedEcDeps.Count)
+					{
+						RunCondition = RunCondition + " && ";
+					}
+				}
+				RunCondition = RunCondition + ")true; else false;]\"";
+			}
+			return RunCondition;
+		}
+
+		private static List<AggregateNode> FindPromotables(IEnumerable<AggregateNode> NodesToDo)
+		{
+			List<AggregateNode> SeparatePromotables = new List<AggregateNode>();
+			foreach (AggregateNode NodeToDo in NodesToDo)
+			{
+				if (NodeToDo.Node.IsSeparatePromotable())
+				{
+					SeparatePromotables.Add(NodeToDo);
+				}
+			}
+			return SeparatePromotables;
+		}
+
+		private static Dictionary<BuildNode, List<AggregateNode>> FindDependentPromotables(IEnumerable<BuildNode> NodesToDo, IEnumerable<AggregateNode> SeparatePromotions)
+		{
+			Dictionary<BuildNode, List<AggregateNode>> DependentPromotions = NodesToDo.ToDictionary(x => x, x => new List<AggregateNode>());
+			foreach (AggregateNode SeparatePromotion in SeparatePromotions)
+			{
+				BuildNode[] Dependencies = SeparatePromotion.Dependencies.SelectMany(x => x.AllIndirectDependencies).Distinct().ToArray();
+				foreach (BuildNode Dependency in Dependencies)
+				{
+					DependentPromotions[Dependency].Add(SeparatePromotion);
+				}
+			}
+			return DependentPromotions;
+		}
+
+		private static Dictionary<string, string> GetFullNodeDependedOnBy(HashSet<BuildNode> NodesToDo)
+		{
+			//find things that depend on our nodes and setup commander dictionary
+			Dictionary<string, string> FullNodeDependedOnBy = new Dictionary<string, string>();
+			foreach (BuildNode NodeToDo in NodesToDo)
+			{
+				FullNodeDependedOnBy[NodeToDo.Name] = "";
+			}
+			foreach (BuildNode NodeToDo in NodesToDo)
+			{
+				if (!NodeToDo.Node.IsTest())
+				{
+					foreach (BuildNode Dep in NodeToDo.AllDirectDependencies)
+					{
+						string CurrentValue;
+						if (!FullNodeDependedOnBy.TryGetValue(Dep.Name, out CurrentValue) || CurrentValue.Length == 0)
+						{
+							FullNodeDependedOnBy[Dep.Name] = NodeToDo.Name;
+						}
+						else
+						{
+							FullNodeDependedOnBy[Dep.Name] += " " + NodeToDo.Name;
+						}
+					}
+				}
+			}
+			return FullNodeDependedOnBy;
+		}
+
+		/// <summary>
+		/// Sorts a list of nodes to display in EC. The default order is based on execution order and agent groups, whereas this function arranges nodes by
+		/// frequency then execution order, while trying to group nodes on parallel paths (eg. Mac/Windows editor nodes) together.
+		/// </summary>
+		static Dictionary<BuildNode, int> GetDisplayOrder(List<BuildNode> Nodes)
+		{
+			// Split the nodes into separate lists for each frequency
+			SortedDictionary<int, List<BuildNode>> NodesByFrequency = new SortedDictionary<int,List<BuildNode>>();
+			foreach(BuildNode Node in Nodes)
+			{
+				List<BuildNode> NodesByThisFrequency;
+				if(!NodesByFrequency.TryGetValue(Node.FrequencyShift, out NodesByThisFrequency))
+				{
+					NodesByThisFrequency = new List<BuildNode>();
+					NodesByFrequency.Add(Node.FrequencyShift, NodesByThisFrequency);
+				}
+				NodesByThisFrequency.Add(Node);
+			}
+
+			// Build the output list by scanning each frequency in order
+			HashSet<BuildNode> VisitedNodes = new HashSet<BuildNode>();
+			Dictionary<BuildNode, int> SortedNodes = new Dictionary<BuildNode,int>();
+			foreach(List<BuildNode> NodesByThisFrequency in NodesByFrequency.Values)
+			{
+				// Find a list of nodes in each display group. If the group name matches the node name, put that node at the front of the list.
+				Dictionary<string, List<BuildNode>> DisplayGroups = new Dictionary<string,List<BuildNode>>();
+				foreach(BuildNode Node in NodesByThisFrequency)
+				{
+					string GroupName = Node.Node.GetDisplayGroupName();
+					if(!DisplayGroups.ContainsKey(GroupName))
+					{
+						DisplayGroups.Add(GroupName, new List<BuildNode>{ Node });
+					}
+					else if(GroupName == Node.Name)
+					{
+						DisplayGroups[GroupName].Insert(0, Node);
+					}
+					else
+					{
+						DisplayGroups[GroupName].Add(Node);
+					}
+				}
+
+				// Build a list of ordering dependencies, putting all Mac nodes after Windows nodes with the same names.
+				Dictionary<BuildNode, List<BuildNode>> NodeDependencies = new Dictionary<BuildNode,List<BuildNode>>(Nodes.ToDictionary(x => x, x => x.AllDirectDependencies.ToList()));
+				foreach(KeyValuePair<string, List<BuildNode>> DisplayGroup in DisplayGroups)
+				{
+					List<BuildNode> GroupNodes = DisplayGroup.Value;
+					for (int Idx = 1; Idx < GroupNodes.Count; Idx++)
+					{
+						NodeDependencies[GroupNodes[Idx]].Add(GroupNodes[0]);
+					}
+				}
+
+				// Add nodes for each frequency into the master list, trying to match up different groups along the way
+				foreach(BuildNode FirstNode in NodesByThisFrequency)
+				{
+					List<BuildNode> GroupNodes = DisplayGroups[FirstNode.Node.GetDisplayGroupName()];
+					foreach(BuildNode GroupNode in GroupNodes)
+					{
+						AddNodeAndDependencies(GroupNode, NodeDependencies, VisitedNodes, SortedNodes);
+					}
+				}
+			}
+			return SortedNodes;
+		}
+
+		static void AddNodeAndDependencies(BuildNode Node, Dictionary<BuildNode, List<BuildNode>> NodeDependencies, HashSet<BuildNode> VisitedNodes, Dictionary<BuildNode, int> SortedNodes)
+		{
+			if(!VisitedNodes.Contains(Node))
+			{
+				VisitedNodes.Add(Node);
+				foreach (BuildNode NodeDependency in NodeDependencies[Node])
+				{
+					AddNodeAndDependencies(NodeDependency, NodeDependencies, VisitedNodes, SortedNodes);
+				}
+				SortedNodes.Add(Node, SortedNodes.Count);
+			}
+		}
+	}
+}
+
 public class ECJobPropsUtils
 {
     public static HashSet<string> ErrorsFromProps(string Filename)
diff --git a/Engine/Source/Programs/AutomationTool/GUBP/GUBP.cs b/Engine/Source/Programs/AutomationTool/GUBP/GUBP.cs
--- a/Engine/Source/Programs/AutomationTool/GUBP/GUBP.cs
+++ b/Engine/Source/Programs/AutomationTool/GUBP/GUBP.cs
@@ -17,18 +17,18 @@ public partial class GUBP : BuildCommand
 	const string FailedTempStorageSuffix = "_Failed";
 	const string SucceededTempStorageSuffix = "_Succeeded";
 
-    class NodeHistory
-    {
-        public int LastSucceeded = 0;
-        public int LastFailed = 0;
-        public List<int> InProgress = new List<int>();
-        public string InProgressString = "";
+	class NodeHistory
+	{
+		public int LastSucceeded = 0;
+		public int LastFailed = 0;
+		public List<int> InProgress = new List<int>();
+		public string InProgressString = "";
 		public List<int> Failed = new List<int>();
-        public string FailedString = "";
-        public List<int> AllStarted = new List<int>();
-        public List<int> AllSucceeded = new List<int>();
-        public List<int> AllFailed = new List<int>();
-    };
+		public string FailedString = "";
+		public List<int> AllStarted = new List<int>();
+		public List<int> AllSucceeded = new List<int>();
+		public List<int> AllFailed = new List<int>();
+	};
 
 	/// <summary>
 	/// Recursively update the ControllingTriggers array for each of the nodes passed in. 
@@ -163,30 +163,6 @@ public partial class GUBP : BuildCommand
 		}
     }
 
-    string RunECTool(string Args, bool bQuiet = false)
-    {
-        if (ParseParam("FakeEC"))
-        {
-            LogWarning("***** Would have ran ectool {0}", Args);
-            return "We didn't actually run ectool";
-        }
-        else
-        {
-            ERunOptions Opts = ERunOptions.Default;
-            if (bQuiet)
-            {
-                Opts = (Opts & ~ERunOptions.AllowSpew) | ERunOptions.NoLoggingOfRunCommand;
-            }
-            return RunAndLog("ectool", "--timeout 900 " + Args, Options: Opts);
-        }
-    }
-    static void WriteECPerl(List<string> Args)
-    {
-        Args.Add("$batch->submit();");
-        string ECPerlFile = CommandUtils.CombinePaths(CommandUtils.CmdEnv.LogFolder, "jobsteps.pl");
-        WriteAllLines_NoExceptions(ECPerlFile, Args.ToArray());
-    }
-
     static List<P4Connection.ChangeRecord> GetChanges(int LastOutputForChanges, int TopCL, int LastGreen)
     {
         List<P4Connection.ChangeRecord> Result = new List<P4Connection.ChangeRecord>();
@@ -481,33 +457,6 @@ public partial class GUBP : BuildCommand
         Result.Sort();
         return Result;
     }
-    void SaveStatus(BuildNode NodeToDo, string Suffix, string NodeStoreName, bool bSaveSharedTempStorage, string GameNameIfAny, string JobStepIDForFailure = null)
-    {
-        string Contents = "Just a status record: " + Suffix;
-        if (!String.IsNullOrEmpty(JobStepIDForFailure) && IsBuildMachine)
-        {
-            try
-            {
-                Contents = RunECTool(String.Format("getProperties --jobStepId {0} --recurse 1", JobStepIDForFailure), true);
-            }
-            catch (Exception Ex)
-            {
-                Log(System.Diagnostics.TraceEventType.Warning, "Failed to get properties for jobstep to save them.");
-                Log(System.Diagnostics.TraceEventType.Warning, LogUtils.FormatException(Ex));
-            }
-        }
-        string RecordOfSuccess = CombinePaths(CommandUtils.CmdEnv.LocalRoot, "Engine", "Saved", "Logs", NodeToDo.Name + Suffix +".log");
-        CreateDirectory(Path.GetDirectoryName(RecordOfSuccess));
-        WriteAllText(RecordOfSuccess, Contents);		
-		TempStorage.StoreToTempStorage(CmdEnv, NodeStoreName + Suffix, new List<string> { RecordOfSuccess }, !bSaveSharedTempStorage, GameNameIfAny);		
-    }
-	string GetPropertyFromStep(string PropertyPath)
-	{
-		string Property = "";
-		Property = RunECTool("getProperty \"" + PropertyPath + "\"");
-		Property = Property.TrimEnd('\r', '\n');
-		return Property;
-	}
 
     int CountZeros(int Num)
     {
@@ -754,22 +703,6 @@ public partial class GUBP : BuildCommand
         return OrdereredToDo;
     }
 
-    static string GetJobStepPath(BuildNode Dep)
-    {
-		if (Dep.Node.AgentSharingGroup == "")
-		{
-			return "jobSteps[" + Dep.Name + "]";
-		}
-        else
-        {
-            return "jobSteps[" + Dep.Node.AgentSharingGroup + "]/jobSteps[" + Dep.Name + "]";
-        }
-    }
-    static string GetJobStep(string ParentPath, BuildNode Dep)
-    {
-        return ParentPath + "/" + GetJobStepPath(Dep);
-    }
-
     static NodeHistory FindNodeHistory(BuildNode NodeToDo, string CLString, string StoreName)
     {
 		NodeHistory History = null;
@@ -813,7 +746,7 @@ public partial class GUBP : BuildCommand
 
 		return History;
     }
-	void GetFailureEmails(BuildNode NodeToDo, NodeHistory History, string CLString, string StoreName, bool OnlyLateUpdates = false)
+	void GetFailureEmails(ElectricCommander EC, BuildNode NodeToDo, NodeHistory History, string CLString, string StoreName, bool OnlyLateUpdates = false)
 	{
         string FailCauserEMails = "";
         string EMailNote = "";
@@ -821,9 +754,6 @@ public partial class GUBP : BuildCommand
         int NumPeople = 0;
         if (History != null)
         {
-			RunECTool(String.Format("setProperty \"/myWorkflow/LastGreen/{0}\" \"{1}\"", NodeToDo.Name, History.LastSucceeded), true);
-			RunECTool(String.Format("setProperty \"/myWorkflow/LastGreen/{0}\" \"{1}\"", NodeToDo.Name, History.FailedString), true);
-
             if (History.LastSucceeded > 0 && History.LastSucceeded < P4Env.Changelist)
             {
                 int LastNonDuplicateFail = P4Env.Changelist;
@@ -882,29 +812,15 @@ public partial class GUBP : BuildCommand
             {
                 SendSuccessForGreenAfterRed = ParseParam("CIS");
             }
-        }
-        else
-        {
-			RunECTool(String.Format("setProperty \"/myWorkflow/LastGreen/{0}\" \"{1}\"", NodeToDo.Name, "0"));
-			RunECTool(String.Format("setProperty \"/myWorkflow/RedsSince/{0}\" \"{1}\"", NodeToDo.Name, ""));
-        }
-		RunECTool(String.Format("setProperty \"/myWorkflow/FailCausers/{0}\" \"{1}\"", NodeToDo.Name, FailCauserEMails));
-		RunECTool(String.Format("setProperty \"/myWorkflow/EmailNotes/{0}\" \"{1}\"", NodeToDo.Name, EMailNote));
-        {
-			HashSet<string> Emails = new HashSet<string>(NodeToDo.RecipientsForFailureEmails);
-            if (ParseParam("CIS") && !NodeToDo.Node.SendSuccessEmail() && !NodeToDo.Node.TriggerNode() && NodeToDo.AddSubmittersToFailureEmails)
-            {
-				Emails.UnionWith(FailCauserEMails.Split(new char[] { ' ' }, StringSplitOptions.RemoveEmptyEntries));
-			}
-			RunECTool(String.Format("setProperty \"/myWorkflow/FailEmails/{0}\" \"{1}\"", NodeToDo.Name, String.Join(" ", Emails)));            
-        }
-		if (NodeToDo.Node.SendSuccessEmail() || SendSuccessForGreenAfterRed)
+		}
+
+		if (History == null)
 		{
-			RunECTool(String.Format("setProperty \"/myWorkflow/SendSuccessEmail/{0}\" \"{1}\"", NodeToDo.Name, "1"));
+			EC.UpdateEmailProperties(NodeToDo, 0, "", FailCauserEMails, EMailNote, SendSuccessForGreenAfterRed);
 		}
 		else
 		{
-			RunECTool(String.Format("setProperty \"/myWorkflow/SendSuccessEmail/{0}\" \"{1}\"", NodeToDo.Name, "0"));
+			EC.UpdateEmailProperties(NodeToDo, History.LastSucceeded, History.FailedString, FailCauserEMails, EMailNote, SendSuccessForGreenAfterRed);
 		}
 	}
 
@@ -990,51 +906,6 @@ public partial class GUBP : BuildCommand
         }
         return Result;
     }
-    List<string> GetECPropsForNode(BuildNode NodeToDo)
-    {
-        List<string> ECProps = new List<string>();
-		ECProps.Add("FailEmails/" + NodeToDo.Name + "=" + String.Join(" ", NodeToDo.RecipientsForFailureEmails));
-	
-		string AgentReq = NodeToDo.Node.ECAgentString();
-		if(ParseParamValue("AgentOverride") != "" && !NodeToDo.Node.GetFullName().Contains("OnMac"))
-		{
-			AgentReq = ParseParamValue("AgentOverride");
-		}
-		ECProps.Add(string.Format("AgentRequirementString/{0}={1}", NodeToDo.Name, AgentReq));
-		ECProps.Add(string.Format("RequiredMemory/{0}={1}", NodeToDo.Name, NodeToDo.AgentMemoryRequirement));
-		ECProps.Add(string.Format("Timeouts/{0}={1}", NodeToDo.Name, NodeToDo.Node.TimeoutInMinutes()));
-		ECProps.Add(string.Format("JobStepPath/{0}={1}", NodeToDo.Name, GetJobStepPath(NodeToDo)));
-		
-        return ECProps;
-    }
-	
-    void UpdateECProps(BuildNode NodeToDo)
-    {
-        try
-        {
-            Log("Updating node props for node {0}", NodeToDo.Name);
-            RunECTool(String.Format("setProperty \"/myWorkflow/FailEmails/{0}\" \"{1}\"", NodeToDo.Name, String.Join(" ", NodeToDo.RecipientsForFailureEmails)), true);
-        }
-        catch (Exception Ex)
-        {
-            Log(System.Diagnostics.TraceEventType.Warning, "Failed to UpdateECProps.");
-            Log(System.Diagnostics.TraceEventType.Warning, LogUtils.FormatException(Ex));
-        }
-    }
-	void UpdateECBuildTime(BuildNode NodeToDo, double BuildDuration)
-	{
-		try
-		{
-			Log("Updating duration prop for node {0}", NodeToDo.Name);
-			RunECTool(String.Format("setProperty \"/myWorkflow/NodeDuration/{0}\" \"{1}\"", NodeToDo.Name, BuildDuration.ToString()));
-			RunECTool(String.Format("setProperty \"/myJobStep/NodeDuration\" \"{0}\"", BuildDuration.ToString()));
-		}
-		catch (Exception Ex)
-		{
-			Log(System.Diagnostics.TraceEventType.Warning, "Failed to UpdateECBuildTime.");
-			Log(System.Diagnostics.TraceEventType.Warning, LogUtils.FormatException(Ex));
-		}
-	}
 
     [Help("Runs one, several or all of the GUBP nodes")]
     [Help(typeof(UE4Build))]
@@ -1198,8 +1069,9 @@ public partial class GUBP : BuildCommand
 			Log("Setting TimeIndex to {0}", TimeIndex);
 		}
 
-		bool WithoutLinux = ParseParam("NoLinux");
-		HashSet<BuildNode> NodesToDo = ParseNodesToDo(AllNodes, AllAggregates, WithoutLinux, TimeIndex, bPreflightBuild, CommanderSetup);
+		HashSet<BuildNode> NodesToDo = ParseNodesToDo(AllNodes, AllAggregates);
+		CullNodesForTimeIndex(NodesToDo, TimeIndex);
+		CullNodesForPreflight(NodesToDo, bPreflightBuild);
 
 		BuildNode ExplicitTrigger = null;
 		if (CommanderSetup)
@@ -1247,10 +1119,12 @@ public partial class GUBP : BuildCommand
         //check sorting
 		CheckSortOrder(OrderedToDo);
 
+		ElectricCommander EC = new ElectricCommander(this);
+
         string FakeFail = ParseParamValue("FakeFail");
         if(CommanderSetup)
         {
-			DoCommanderSetup(AllNodes, AllAggregates, NodesToDo, OrderedToDo, TimeIndex, TimeQuantum, bSkipTriggers, bFake, bFakeEC, CLString, ExplicitTrigger, UnfinishedTriggers, FakeFail, bPreflightBuild);
+			DoCommanderSetup(EC, AllNodes, AllAggregates, NodesToDo, OrderedToDo, TimeIndex, TimeQuantum, bSkipTriggers, bFake, bFakeEC, CLString, ExplicitTrigger, UnfinishedTriggers, FakeFail, bPreflightBuild);
         }
 		else if(ParseParam("SaveGraph"))
 		{
@@ -1258,38 +1132,11 @@ public partial class GUBP : BuildCommand
 		}
 		else if(!ParseParam("ListOnly"))
 		{
-			ExecuteNodes(OrderedToDo, bFake, bFakeEC, bSaveSharedTempStorage, CLString, StoreName, FakeFail);
+			ExecuteNodes(EC, OrderedToDo, bFake, bFakeEC, bSaveSharedTempStorage, CLString, StoreName, FakeFail);
 		}
         PrintRunTime();
 	}
 
-	private static List<AggregateNode> FindPromotables(IEnumerable<AggregateNode> NodesToDo)
-	{
-		List<AggregateNode> SeparatePromotables = new List<AggregateNode>();
-		foreach (AggregateNode NodeToDo in NodesToDo)
-		{
-			if (NodeToDo.Node.IsSeparatePromotable())
-			{
-				SeparatePromotables.Add(NodeToDo);
-			}
-		}
-		return SeparatePromotables;
-	}
-
-	private static Dictionary<BuildNode, List<AggregateNode>> FindDependentPromotables(IEnumerable<BuildNode> NodesToDo, IEnumerable<AggregateNode> SeparatePromotions)
-	{
-		Dictionary<BuildNode, List<AggregateNode>> DependentPromotions = NodesToDo.ToDictionary(x => x, x => new List<AggregateNode>());
-		foreach (AggregateNode SeparatePromotion in SeparatePromotions)
-		{
-			BuildNode[] Dependencies = SeparatePromotion.Dependencies.SelectMany(x => x.AllIndirectDependencies).Distinct().ToArray();
-			foreach (BuildNode Dependency in Dependencies)
-			{
-				DependentPromotions[Dependency].Add(SeparatePromotion);
-			}
-		}
-		return DependentPromotions;
-	}
-
 	/// <summary>
 	/// Resolves the names of each node and aggregates' dependencies, and links them together into the build graph.
 	/// </summary>
@@ -1445,8 +1292,15 @@ public partial class GUBP : BuildCommand
 		return false;
 	}
 
-	private HashSet<BuildNode> ParseNodesToDo(IEnumerable<BuildNode> PotentialNodes, IEnumerable<AggregateNode> PotentialAggregates, bool WithoutLinux, int TimeIndex, bool bPreflightBuild, bool CommanderSetup)
+	/// <summary>
+	/// Determine which nodes to build from the command line, or return everything if there's nothing specified explicitly.
+	/// </summary>
+	/// <param name="PotentialNodes">The nodes in the graph</param>
+	/// <param name="PotentialAggregates">The aggregates in the graph</param>
+	/// <returns>Set of all nodes to execute, recursively including their dependencies.</returns>
+	private HashSet<BuildNode> ParseNodesToDo(IEnumerable<BuildNode> PotentialNodes, IEnumerable<AggregateNode> PotentialAggregates)
 	{
+		// Parse all the node names that we want
 		List<string> NamesToDo = new List<string>();
 
 		string NodeParam = ParseParamValue("Node", null);
@@ -1461,6 +1315,7 @@ public partial class GUBP : BuildCommand
 			NamesToDo.AddRange(GameParam.Split('+').Select(x => FullGameAggregateNode.StaticGetFullName(x)));
 		}
 
+		// Resolve all those node names to a list of build nodes to execute
 		HashSet<BuildNode> NodesToDo = new HashSet<BuildNode>();
 		foreach (string NameToDo in NamesToDo)
 		{
@@ -1487,133 +1342,63 @@ public partial class GUBP : BuildCommand
 			}
 		}
 
+		// Add everything if there was nothing specified explicitly
 		if (NodesToDo.Count == 0)
 		{
-			LogVerbose("No nodes specified, adding all nodes");
 			NodesToDo.UnionWith(PotentialNodes);
 		}
-		else if (TimeIndex != 0)
-		{
-			LogVerbose("Check to make sure we didn't ask for nodes that will be culled by time index");
-			foreach (BuildNode NodeToDo in NodesToDo)
-			{
-				if (TimeIndex % (1 << NodeToDo.FrequencyShift) != 0)
-				{
-					throw new AutomationException("You asked specifically for node {0}, but it is culled by the time quantum: TimeIndex = {1}, DependentCISFrequencyQuantumShift = {2}.", NodeToDo.Name, TimeIndex, NodeToDo.FrequencyShift);
-				}
-			}
-		}
 
-		LogVerbose("Desired Nodes");
-		foreach (BuildNode NodeToDo in NodesToDo)
+		// Recursively find the complete set of all nodes we want to execute
+		HashSet<BuildNode> RecursiveNodesToDo = new HashSet<BuildNode>(NodesToDo);
+		foreach(BuildNode NodeToDo in NodesToDo)
 		{
-			LogVerbose("  {0}", NodeToDo.Name);
+			RecursiveNodesToDo.UnionWith(NodeToDo.AllIndirectDependencies);
 		}
 
-		// find things that our nodes depend on
-		bool bDoneWithDependencies = false;
-		while (!bDoneWithDependencies)
-		{
-			bDoneWithDependencies = true;
-			HashSet<BuildNode> Fringe = new HashSet<BuildNode>();
-			foreach (BuildNode NodeToDo in NodesToDo)
-			{
-				foreach (BuildNode Dep in NodeToDo.Dependencies)
-				{
-					if (!NodesToDo.Contains(Dep))
-					{
-						Fringe.Add(Dep);
-						bDoneWithDependencies = false;
-					}
-				}
-			}
-			NodesToDo.UnionWith(Fringe);
-		}
+		return RecursiveNodesToDo;
+	}
 
+	/// <summary>
+	/// Culls the list of nodes to do depending on the current time index.
+	/// </summary>
+	/// <param name="NodesToDo">The list of nodes to do</param>
+	/// <param name="TimeIndex">The current time index. All nodes are run for TimeIndex=0, otherwise they are culled based on their FrequencyShift parameter.</param>
+	private void CullNodesForTimeIndex(HashSet<BuildNode> NodesToDo, int TimeIndex)
+	{
 		if (TimeIndex != 0)
 		{
-			LogVerbose("Culling based on time index");
-			HashSet<BuildNode> NewNodesToDo = new HashSet<BuildNode>();
-			foreach (BuildNode NodeToDo in NodesToDo)
-			{
-				if (TimeIndex % (1 << NodeToDo.FrequencyShift) == 0)
-				{
-					LogVerbose("  Keeping {0}", NodeToDo.Name);
-					NewNodesToDo.Add(NodeToDo);
-				}
-				else
-				{
-					LogVerbose("  Rejecting {0}", NodeToDo.Name);
-				}
-			}
-			NodesToDo = NewNodesToDo;
-		}
-		//Remove Plat if specified
-		if (WithoutLinux)
-		{
-			HashSet<BuildNode> NewNodesToDo = new HashSet<BuildNode>();
-			foreach (BuildNode NodeToDo in NodesToDo)
-			{
-				if (!NodeToDo.Node.GetFullName().Contains("Linux"))
-				{
-					NewNodesToDo.Add(NodeToDo);
-				}
-				else
-				{
-					LogVerbose(" Rejecting {0} because -NoLinux was requested", NodeToDo.Name);
-				}
-			}
-			NodesToDo = NewNodesToDo;
-		}
-		if (bPreflightBuild)
-		{
-			LogVerbose("Culling triggers and downstream for preflight builds ");
-			HashSet<BuildNode> NewNodesToDo = new HashSet<BuildNode>();
+			List<BuildNode> NodesToCull = new List<BuildNode>();
 			foreach (BuildNode NodeToDo in NodesToDo)
 			{
-				string TriggerDot = NodeToDo.ControllingTriggerDotName;
-				if (TriggerDot == "" && !NodeToDo.Node.TriggerNode())
-				{
-					LogVerbose("  Keeping {0}", NodeToDo.Name);
-					NewNodesToDo.Add(NodeToDo);
-				}
-				else
+				if ((TimeIndex % (1 << NodeToDo.FrequencyShift)) != 0)
 				{
-					LogVerbose("  Rejecting {0}", NodeToDo.Name);
+					NodesToCull.Add(NodeToDo);
 				}
 			}
-			NodesToDo = NewNodesToDo;
+			NodesToDo.ExceptWith(NodesToCull);
 		}
-		return NodesToDo;
 	}
 
-	private static Dictionary<string, string> GetFullNodeDependedOnBy(HashSet<BuildNode> NodesToDo)
+	/// <summary>
+	/// Culls everything downstream of a trigger if we're running a preflight build.
+	/// </summary>
+	/// <param name="NodesToDo">The current list of nodes to do</param>
+	/// <param name="IsPreflightBuild">Whether this is a preflight build or not</param>
+	private void CullNodesForPreflight(HashSet<BuildNode> NodesToDo, bool IsPreflightBuild)
 	{
-		//find things that depend on our nodes and setup commander dictionary
-		Dictionary<string, string> FullNodeDependedOnBy = new Dictionary<string, string>();
-		foreach (BuildNode NodeToDo in NodesToDo)
-		{
-			FullNodeDependedOnBy[NodeToDo.Name] = "";
-		}
-		foreach (BuildNode NodeToDo in NodesToDo)
+		if (IsPreflightBuild)
 		{
-			if (!NodeToDo.Node.IsTest())
+			HashSet<BuildNode> NodesToCull = new HashSet<BuildNode>();
+			foreach(BuildNode NodeToDo in NodesToDo)
 			{
-				foreach (BuildNode Dep in NodeToDo.AllDirectDependencies)
+				if(NodeToDo.Node.TriggerNode() || NodeToDo.ControllingTriggers.Length > 0)
 				{
-					string CurrentValue;
-					if (!FullNodeDependedOnBy.TryGetValue(Dep.Name, out CurrentValue) || CurrentValue.Length == 0)
-					{
-						FullNodeDependedOnBy[Dep.Name] = NodeToDo.Name;
-					}
-					else
-					{
-						FullNodeDependedOnBy[Dep.Name] += " " + NodeToDo.Name;
-					}
+					LogVerbose(" Culling {0} due to being downstream of trigger in preflight", NodeToDo.Name);
+					NodesToCull.Add(NodeToDo);
 				}
 			}
+			NodesToDo.ExceptWith(NodesToCull);
 		}
-		return FullNodeDependedOnBy;
 	}
 
 	private int UpdateCISCounter(int TimeQuantum)
@@ -1749,6 +1534,10 @@ public partial class GUBP : BuildCommand
 		return UnfinishedTriggers;
 	}
 
+	/// <summary>
+	/// Validates that the given nodes are sorted correctly, so that all dependencies are met first
+	/// </summary>
+	/// <param name="OrdereredToDo">The sorted list of nodes</param>
 	private void CheckSortOrder(List<BuildNode> OrdereredToDo)
 	{
 		foreach (BuildNode NodeToDo in OrdereredToDo)
@@ -1777,7 +1566,7 @@ public partial class GUBP : BuildCommand
 		}
 	}
 
-	private void DoCommanderSetup(IEnumerable<BuildNode> AllNodes, IEnumerable<AggregateNode> AllAggregates, HashSet<BuildNode> NodesToDo, List<BuildNode> OrdereredToDo, int TimeIndex, int TimeQuantum, bool bSkipTriggers, bool bFake, bool bFakeEC, string CLString, BuildNode ExplicitTrigger, List<BuildNode> UnfinishedTriggers, string FakeFail, bool bPreflightBuild)
+	private void DoCommanderSetup(ElectricCommander EC, IEnumerable<BuildNode> AllNodes, IEnumerable<AggregateNode> AllAggregates, HashSet<BuildNode> NodesToDo, List<BuildNode> OrdereredToDo, int TimeIndex, int TimeQuantum, bool bSkipTriggers, bool bFake, bool bFakeEC, string CLString, BuildNode ExplicitTrigger, List<BuildNode> UnfinishedTriggers, string FakeFail, bool bPreflightBuild)
 	{
 		List<BuildNode> SortedNodes = TopologicalSort(new HashSet<BuildNode>(AllNodes), null, SubSort: false, DoNotConsiderCompletion: true);
 		Log("******* {0} GUBP Nodes", SortedNodes.Count);
@@ -1816,472 +1605,10 @@ public partial class GUBP : BuildCommand
 			PrintNodes(this, FilteredOrdereredToDo, AllAggregates, UnfinishedTriggers, TimeQuantum);
 		}
 
-		DoCommanderSetupInternal(AllNodes, AllAggregates, NodesToDo, FilteredOrdereredToDo, SortedNodes, TimeIndex, TimeQuantum, bSkipTriggers, bFake, bFakeEC, CLString, ExplicitTrigger, UnfinishedTriggers, FakeFail);
-	}
-
-	private void DoCommanderSetupInternal(IEnumerable<BuildNode> AllNodes, IEnumerable<AggregateNode> AllAggregates, HashSet<BuildNode> NodesToDo, List<BuildNode> OrdereredToDo, List<BuildNode> SortedNodes, int TimeIndex, int TimeQuantum, bool bSkipTriggers, bool bFake, bool bFakeEC, string CLString, BuildNode ExplicitTrigger, List<BuildNode> UnfinishedTriggers, string FakeFail)
-	{
-		Dictionary<string, string> FullNodeDependedOnBy = GetFullNodeDependedOnBy(NodesToDo);
-
-		List<AggregateNode> SeparatePromotables = FindPromotables(AllAggregates);
-		Dictionary<BuildNode, List<AggregateNode>> DependentPromotions = FindDependentPromotables(AllNodes, SeparatePromotables);
-
-		Dictionary<BuildNode, int> FullNodeListSortKey = GetDisplayOrder(SortedNodes);
-
-		if (OrdereredToDo.Count == 0)
-		{
-			throw new AutomationException("No nodes to do!");
-		}
-		List<string> ECProps = new List<string>();
-		ECProps.Add(String.Format("TimeIndex={0}", TimeIndex));
-		foreach (BuildNode Node in SortedNodes)
-		{
-			ECProps.Add(string.Format("AllNodes/{0}={1}", Node.Name, GetNodeForAllNodesProperty(Node, TimeQuantum)));
-		}
-		foreach (BuildNode Node in SortedNodes)
-		{
-			ECProps.Add(string.Format("DirectDependencies/{0}={1}", Node.Name, String.Join(" ", Node.AllDirectDependencies.Select(x => x.Name))));
-		}
-		foreach (KeyValuePair<BuildNode, int> NodePair in FullNodeListSortKey)
-		{
-			ECProps.Add(string.Format("SortKey/{0}={1}", NodePair.Key.Name, NodePair.Value));
-		}
-		foreach (KeyValuePair<string, string> NodePair in FullNodeDependedOnBy)
-		{
-			ECProps.Add(string.Format("DependedOnBy/{0}={1}", NodePair.Key, NodePair.Value));
-		}
-		foreach (KeyValuePair<BuildNode, List<AggregateNode>> NodePair in DependentPromotions)
-		{
-			ECProps.Add(string.Format("DependentPromotions/{0}={1}", NodePair.Key.Name, String.Join(" ", NodePair.Value.Select(x => x.Name))));
-		}
-		foreach (AggregateNode Node in SeparatePromotables)
-		{
-			ECProps.Add(string.Format("PossiblePromotables/{0}={1}", Node.Name, ""));
-		}
-		List<string> ECJobProps = new List<string>();
-		if (ExplicitTrigger != null)
-		{
-			ECJobProps.Add("IsRoot=0");
-		}
-		else
-		{
-			ECJobProps.Add("IsRoot=1");
-		}
-
-		// here we are just making sure everything before the explicit trigger is completed.
-		if (ExplicitTrigger != null)
-		{
-			foreach (BuildNode NodeToDo in OrdereredToDo)
-			{
-				if (!NodeToDo.IsComplete && NodeToDo != ExplicitTrigger && !NodeToDo.DependsOn(ExplicitTrigger)) // if something is already finished, we don't put it into EC
-				{
-					throw new AutomationException("We are being asked to process node {0}, however, this is an explicit trigger {1}, so everything before it should already be handled. It seems likely that you waited too long to run the trigger. You will have to do a new build from scratch.", NodeToDo.Name, ExplicitTrigger.Name);
-				}
-			}
-		}
-
-		BuildNode LastSticky = null;
-		bool HitNonSticky = false;
-		bool bHaveECNodes = false;
-		List<string> StepList = new List<string>();
-		StepList.Add("use strict;");
-		StepList.Add("use diagnostics;");
-		StepList.Add("use ElectricCommander();");
-		StepList.Add("my $ec = new ElectricCommander;");
-		StepList.Add("$ec->setTimeout(600);");
-		StepList.Add("my $batch = $ec->newBatch(\"serial\");");
-		// sticky nodes are ones that we run on the main agent. We run then first and they must not be intermixed with parallel jobs
-		foreach (BuildNode NodeToDo in OrdereredToDo)
-		{
-			if (!NodeToDo.IsComplete) // if something is already finished, we don't put it into EC
-			{
-				bHaveECNodes = true;
-				if (NodeToDo.Node.IsSticky())
-				{
-					LastSticky = NodeToDo;
-					if (HitNonSticky && !bSkipTriggers)
-					{
-						throw new AutomationException("Sticky and non-sticky jobs did not sort right.");
-					}
-				}
-				else
-				{
-					HitNonSticky = true;
-				}
-			}
-		}
-
-		List<string> FakeECArgs = new List<string>();
-		using(TelemetryStopwatch PerlOutputStopwatch = new TelemetryStopwatch("PerlOutput"))
-		{
-			string ParentPath = ParseParamValue("ParentPath");
-			string BaseArgs = String.Format("$batch->createJobStep({{parentPath => '{0}'", ParentPath);
-
-			bool bHasNoop = false;
-			if (LastSticky == null && bHaveECNodes)
-			{
-				// if we don't have any sticky nodes and we have other nodes, we run a fake noop just to release the resource 
-				string Args = String.Format("{0}, subprocedure => 'GUBP_UAT_Node', parallel => '0', jobStepName => 'Noop', actualParameter => [{{actualParameterName => 'NodeName', value => 'Noop'}}, {{actualParameterName => 'Sticky', value =>'1' }}], releaseMode => 'release'}});", BaseArgs);
-				StepList.Add(Args);
-				bHasNoop = true;
-			}
-
-			Dictionary<string, List<BuildNode>> AgentGroupChains = new Dictionary<string, List<BuildNode>>();
-			List<BuildNode> StickyChain = new List<BuildNode>();
-			foreach (BuildNode NodeToDo in OrdereredToDo)
-			{
-				if (!NodeToDo.IsComplete) // if something is already finished, we don't put it into EC  
-				{
-					string MyAgentGroup = NodeToDo.Node.AgentSharingGroup;
-					if (MyAgentGroup != "")
-					{
-						if (!AgentGroupChains.ContainsKey(MyAgentGroup))
-						{
-							AgentGroupChains.Add(MyAgentGroup, new List<BuildNode> { NodeToDo });
-						}
-						else
-						{
-							AgentGroupChains[MyAgentGroup].Add(NodeToDo);
-						}
-					}
-				}
-				if (NodeToDo.Node.IsSticky())
-				{
-					if (!StickyChain.Contains(NodeToDo))
-					{
-						StickyChain.Add(NodeToDo);
-					}
-				}
-			}
-			foreach (BuildNode NodeToDo in OrdereredToDo)
-			{
-				if (!NodeToDo.IsComplete) // if something is already finished, we don't put it into EC  
-				{
-					List<string> NodeProps = GetECPropsForNode(NodeToDo);
-					ECProps.AddRange(NodeProps);
-
-					bool Sticky = NodeToDo.Node.IsSticky();
-					bool DoParallel = !Sticky;
-					if (NodeToDo.Node.ECProcedure() == "GUBP_UAT_Node_Parallel_AgentShare_Editor")
-					{
-						DoParallel = true;
-					}
-					if (Sticky && NodeToDo.Node.ECAgentString() != "")
-					{
-						throw new AutomationException("Node {1} is sticky but has agent requirements.", NodeToDo.Name);
-					}
-					string Procedure = NodeToDo.Node.ECProcedure();
-					if (NodeToDo.Node.IsSticky() && NodeToDo == LastSticky)
-					{
-						Procedure = Procedure + "_Release";
-					}
-					string Args = String.Format("{0}, subprocedure => '{1}', parallel => '{2}', jobStepName => '{3}', actualParameter => [{{actualParameterName => 'NodeName', value =>'{4}'}}",
-						BaseArgs, Procedure, DoParallel ? 1 : 0, NodeToDo.Name, NodeToDo.Name);
-					string ProcedureParams = NodeToDo.Node.ECProcedureParams();
-					if (!String.IsNullOrEmpty(ProcedureParams))
-					{
-						Args = Args + ProcedureParams;
-					}
-
-					if ((Procedure == "GUBP_UAT_Trigger" || Procedure == "GUBP_Hardcoded_Trigger") && NodeToDo.RecipientsForFailureEmails.Length > 0)
-					{
-						Args = Args + ", {actualParameterName => 'EmailsForTrigger', value => \'" + String.Join(" ", NodeToDo.RecipientsForFailureEmails) + "\'}";
-					}
-					Args = Args + "]";
-
-					List<BuildNode> UncompletedEcDeps = new List<BuildNode>();
-					{
-						foreach (BuildNode Dep in NodeToDo.AllDirectDependencies)
-						{
-							if (!Dep.IsComplete && OrdereredToDo.Contains(Dep)) // if something is already finished, we don't put it into EC
-							{
-								if (OrdereredToDo.IndexOf(Dep) > OrdereredToDo.IndexOf(NodeToDo))
-								{
-									throw new AutomationException("Topological sort error, node {0} has a dependency of {1} which sorted after it.", NodeToDo.Name, Dep.Name);
-								}
-								UncompletedEcDeps.Add(Dep);
-							}
-						}
-					}
-
-					string PreCondition = GetPreConditionForNode(OrdereredToDo, ParentPath, bHasNoop, AgentGroupChains, StickyChain, NodeToDo, UncompletedEcDeps);
-					string RunCondition = GetRunConditionForNode(UncompletedEcDeps, ParentPath);
-
-					string MyAgentGroup = NodeToDo.Node.AgentSharingGroup;
-					bool bDoNestedJobstep = false;
-					bool bDoFirstNestedJobstep = false;
-
-					string NodeParentPath = ParentPath;
-					if (MyAgentGroup != "")
-					{
-						bDoNestedJobstep = true;
-						NodeParentPath = ParentPath + "/jobSteps[" + MyAgentGroup + "]";
-
-						List<BuildNode> MyChain = AgentGroupChains[MyAgentGroup];
-						int MyIndex = MyChain.IndexOf(NodeToDo);
-						if (MyIndex <= 0)
-						{
-							bDoFirstNestedJobstep = bDoNestedJobstep;
-						}
-					}
-
-					if (bDoNestedJobstep)
-					{
-						if (bDoFirstNestedJobstep)
-						{
-							{
-								string NestArgs = String.Format("$batch->createJobStep({{parentPath => '{0}', jobStepName => '{1}', parallel => '1'",
-									ParentPath, MyAgentGroup);
-								if (!String.IsNullOrEmpty(PreCondition))
-								{
-									NestArgs = NestArgs + ", precondition => " + PreCondition;
-								}
-								NestArgs = NestArgs + "});";
-								StepList.Add(NestArgs);
-							}
-							{
-								string NestArgs = String.Format("$batch->createJobStep({{parentPath => '{0}/jobSteps[{1}]', jobStepName => '{2}_GetPool', subprocedure => 'GUBP{3}_AgentShare_GetPool', parallel => '1', actualParameter => [{{actualParameterName => 'AgentSharingGroup', value => '{4}'}}, {{actualParameterName => 'NodeName', value => '{5}'}}]",
-									ParentPath, MyAgentGroup, MyAgentGroup, NodeToDo.Node.ECProcedureInfix(), MyAgentGroup, NodeToDo.Name);
-								if (!String.IsNullOrEmpty(PreCondition))
-								{
-									NestArgs = NestArgs + ", precondition => " + PreCondition;
-								}
-								NestArgs = NestArgs + "});";
-								StepList.Add(NestArgs);
-							}
-							{
-								string NestArgs = String.Format("$batch->createJobStep({{parentPath => '{0}/jobSteps[{1}]', jobStepName => '{2}_GetAgent', subprocedure => 'GUBP{3}_AgentShare_GetAgent', parallel => '1', exclusiveMode => 'call', resourceName => '{4}', actualParameter => [{{actualParameterName => 'AgentSharingGroup', value => '{5}'}}, {{actualParameterName => 'NodeName', value=> '{6}'}}]",
-									ParentPath, MyAgentGroup, MyAgentGroup, NodeToDo.Node.ECProcedureInfix(),
-									String.Format("$[/myJob/jobSteps[{0}]/ResourcePool]", MyAgentGroup),
-									MyAgentGroup, NodeToDo.Name);
-								{
-									NestArgs = NestArgs + ", precondition  => ";
-									NestArgs = NestArgs + "\"\\$\" . \"[/javascript if(";
-									NestArgs = NestArgs + "getProperty('" + ParentPath + "/jobSteps[" + MyAgentGroup + "]/jobSteps[" + MyAgentGroup + "_GetPool]/status') == 'completed'";
-									NestArgs = NestArgs + ") true;]\"";
-								}
-								NestArgs = NestArgs + "});";
-								StepList.Add(NestArgs);
-							}
-							{
-								PreCondition = "\"\\$\" . \"[/javascript if(";
-								PreCondition = PreCondition + "getProperty('" + ParentPath + "/jobSteps[" + MyAgentGroup + "]/jobSteps[" + MyAgentGroup + "_GetAgent]/status') == 'completed'";
-								PreCondition = PreCondition + ") true;]\"";
-							}
-						}
-						Args = Args.Replace(String.Format("parentPath => '{0}'", ParentPath), String.Format("parentPath => '{0}'", NodeParentPath));
-						Args = Args.Replace("UAT_Node_Parallel_AgentShare", "UAT_Node_Parallel_AgentShare3");
-					}
-
-					if (!String.IsNullOrEmpty(PreCondition))
-					{
-						Args = Args + ", precondition => " + PreCondition;
-					}
-					if (!String.IsNullOrEmpty(RunCondition))
-					{
-						Args = Args + ", condition => " + RunCondition;
-					}
-	#if false
-						// this doesn't work because it includes precondition time
-						if (GUBPNodes[NodeToDo].TimeoutInMinutes() > 0)
-						{
-							Args = Args + String.Format(" --timeLimitUnits minutes --timeLimit {0}", GUBPNodes[NodeToDo].TimeoutInMinutes());
-						}
-	#endif
-					if (Sticky && NodeToDo == LastSticky)
-					{
-						Args = Args + ", releaseMode => 'release'";
-					}
-					Args = Args + "});";
-					StepList.Add(Args);
-					if (bFakeEC &&
-						!UnfinishedTriggers.Contains(NodeToDo) &&
-						(NodeToDo.Node.ECProcedure().StartsWith("GUBP_UAT_Node") || NodeToDo.Node.ECProcedure().StartsWith("GUBP_Mac_UAT_Node")) // other things we really can't test
-						) // unfinished triggers are never run directly by EC, rather it does another job setup
-					{
-						string Arg = String.Format("gubp -Node={0} -FakeEC {1} {2} {3} {4} {5}",
-							NodeToDo.Name,
-							bFake ? "-Fake" : "",
-							ParseParam("AllPlatforms") ? "-AllPlatforms" : "",
-							ParseParam("UnfinishedTriggersFirst") ? "-UnfinishedTriggersFirst" : "",
-							ParseParam("UnfinishedTriggersParallel") ? "-UnfinishedTriggersParallel" : "",
-							ParseParam("WithMac") ? "-WithMac" : ""
-							);
-
-						string Node = ParseParamValue("-Node");
-						if (!String.IsNullOrEmpty(Node))
-						{
-							Arg = Arg + " -Node=" + Node;
-						}
-						if (!String.IsNullOrEmpty(FakeFail))
-						{
-							Arg = Arg + " -FakeFail=" + FakeFail;
-						}
-						FakeECArgs.Add(Arg);
-					}
-
-					if (MyAgentGroup != "" && !bDoNestedJobstep)
-					{
-						List<BuildNode> MyChain = AgentGroupChains[MyAgentGroup];
-						int MyIndex = MyChain.IndexOf(NodeToDo);
-						if (MyIndex == MyChain.Count - 1)
-						{
-							string RelPreCondition = "\"\\$\" . \"[/javascript if(";
-							// this runs "parallel", but we a precondition to serialize it
-							RelPreCondition = RelPreCondition + "getProperty('" + ParentPath + "/jobSteps[" + NodeToDo.Name + "]/status') == 'completed'";
-							RelPreCondition = RelPreCondition + ") true;]\"";
-							// we need to release the resource
-							string RelArgs = String.Format("{0}, subprocedure => 'GUBP_Release_AgentShare', parallel => '1', jobStepName => 'Release_{1}', actualParameter => [{{actualParameterName => 'AgentSharingGroup', valued => '{2}'}}], releaseMode => 'release', precondition => '{3}'",
-								BaseArgs, MyAgentGroup, MyAgentGroup, RelPreCondition);
-							StepList.Add(RelArgs);
-						}
-					}
-				}
-			}
-			WriteECPerl(StepList);
-		}
-		bool bHasTests = OrdereredToDo.Any(x => x.Node.IsTest());
-		RunECTool(String.Format("setProperty \"/myWorkflow/HasTests\" \"{0}\"", bHasTests));
-		{
-			ECProps.Add("GUBP_LoadedProps=1");
-			string BranchDefFile = CommandUtils.CombinePaths(CommandUtils.CmdEnv.LogFolder, "BranchDef.properties");
-			CommandUtils.WriteAllLines(BranchDefFile, ECProps.ToArray());
-			RunECTool(String.Format("setProperty \"/myWorkflow/BranchDefFile\" \"{0}\"", BranchDefFile.Replace("\\", "\\\\")));
-		}
-		{
-			ECProps.Add("GUBP_LoadedJobProps=1");
-			string BranchJobDefFile = CommandUtils.CombinePaths(CommandUtils.CmdEnv.LogFolder, "BranchJobDef.properties");
-			CommandUtils.WriteAllLines(BranchJobDefFile, ECProps.ToArray());
-			RunECTool(String.Format("setProperty \"/myJob/BranchJobDefFile\" \"{0}\"", BranchJobDefFile.Replace("\\", "\\\\")));
-		}
-	}
-
-	private string GetPreConditionForNode(List<BuildNode> OrdereredToDo, string PreconditionParentPath, bool bHasNoop, Dictionary<string, List<BuildNode>> AgentGroupChains, List<BuildNode> StickyChain, BuildNode NodeToDo, List<BuildNode> UncompletedEcDeps)
-	{
-		List<BuildNode> PreConditionUncompletedEcDeps = new List<BuildNode>();
-		if(NodeToDo.Node.AgentSharingGroup == "")
-		{
-			PreConditionUncompletedEcDeps = new List<BuildNode>(UncompletedEcDeps);
-		}
-		else 
-		{
-			List<BuildNode> MyChain = AgentGroupChains[NodeToDo.Node.AgentSharingGroup];
-			int MyIndex = MyChain.IndexOf(NodeToDo);
-			if (MyIndex > 0)
-			{
-				PreConditionUncompletedEcDeps.Add(MyChain[MyIndex - 1]);
-			}
-			else
-			{
-				// to avoid idle agents (and also EC doesn't actually reserve our agent!), we promote all dependencies to the first one
-				foreach (BuildNode Chain in MyChain)
-				{
-					foreach (BuildNode Dep in Chain.AllDirectDependencies)
-					{
-						if (!Dep.IsComplete && OrdereredToDo.Contains(Dep)) // if something is already finished, we don't put it into EC
-						{
-							if (OrdereredToDo.IndexOf(Dep) > OrdereredToDo.IndexOf(Chain))
-							{
-								throw new AutomationException("Topological sort error, node {0} has a dependency of {1} which sorted after it.", Chain.Name, Dep.Name);
-							}
-							if (!MyChain.Contains(Dep) && !PreConditionUncompletedEcDeps.Contains(Dep))
-							{
-								PreConditionUncompletedEcDeps.Add(Dep);
-							}
-						}
-					}
-				}
-			}
-		}
-		if (NodeToDo.Node.IsSticky())
-		{
-			List<BuildNode> MyChain = StickyChain;
-			int MyIndex = MyChain.IndexOf(NodeToDo);
-			if (MyIndex > 0)
-			{
-				if (!PreConditionUncompletedEcDeps.Contains(MyChain[MyIndex - 1]) && !MyChain[MyIndex - 1].IsComplete)
-				{
-					PreConditionUncompletedEcDeps.Add(MyChain[MyIndex - 1]);
-				}
-			}
-			else
-			{
-				foreach (BuildNode Dep in NodeToDo.AllDirectDependencies)
-				{
-					if (!Dep.IsComplete && OrdereredToDo.Contains(Dep)) // if something is already finished, we don't put it into EC
-					{
-						if (OrdereredToDo.IndexOf(Dep) > OrdereredToDo.IndexOf(NodeToDo))
-						{
-							throw new AutomationException("Topological sort error, node {0} has a dependency of {1} which sorted after it.", NodeToDo.Name, Dep.Name);
-						}
-						if (!MyChain.Contains(Dep) && !PreConditionUncompletedEcDeps.Contains(Dep))
-						{
-							PreConditionUncompletedEcDeps.Add(Dep);
-						}
-					}
-				}
-			}
-		}
-
-		List<string> JobStepNames = new List<string>();
-		if (bHasNoop && PreConditionUncompletedEcDeps.Count == 0)
-		{
-			JobStepNames.Add(PreconditionParentPath + "/jobSteps[Noop]");
-		}
-		else
-		{
-			JobStepNames.AddRange(PreConditionUncompletedEcDeps.Select(x => GetJobStep(PreconditionParentPath, x)));
-		}
-
-		string PreCondition = "";
-		if (JobStepNames.Count > 0)
-		{
-			PreCondition = String.Format("\"\\$\" . \"[/javascript if({0}) true;]\"", String.Join(" && ", JobStepNames.Select(x => String.Format("getProperty('{0}/status') == 'completed'", x))));
-		}
-		return PreCondition;
-	}
-
-	private static string GetNodeForAllNodesProperty(BuildNode Node, int TimeQuantum)
-	{
-		string Note = Node.ControllingTriggerDotName;
-		if (Note == "")
-		{
-			if (Node.FrequencyShift == 0)
-			{
-				Note = "always";
-			}
-			else
-			{
-				Note = String.Format(" ({0})", GetTimeIntervalString(TimeQuantum << Node.FrequencyShift));
-			}
-		}
-		return Note;
-	}
-
-	private static string GetRunConditionForNode(List<BuildNode> UncompletedEcDeps, string PreconditionParentPath)
-	{
-		string RunCondition = "";
-		if (UncompletedEcDeps.Count > 0)
-		{
-			RunCondition = "\"\\$\" . \"[/javascript if(";
-			int Index = 0;
-			foreach (BuildNode Dep in UncompletedEcDeps)
-			{
-				RunCondition = RunCondition + "((\'\\$\" . \"[" + GetJobStep(PreconditionParentPath, Dep) + "/outcome]\' == \'success\') || ";
-				RunCondition = RunCondition + "(\'\\$\" . \"[" + GetJobStep(PreconditionParentPath, Dep) + "/outcome]\' == \'warning\'))";
-
-				Index++;
-				if (Index != UncompletedEcDeps.Count)
-				{
-					RunCondition = RunCondition + " && ";
-				}
-			}
-			RunCondition = RunCondition + ")true; else false;]\"";
-		}
-		return RunCondition;
+		EC.DoCommanderSetup(AllNodes, AllAggregates, NodesToDo, FilteredOrdereredToDo, SortedNodes, TimeIndex, TimeQuantum, bSkipTriggers, bFake, bFakeEC, CLString, ExplicitTrigger, UnfinishedTriggers, FakeFail);
 	}
 
-	void ExecuteNodes(List<BuildNode> OrdereredToDo, bool bFake, bool bFakeEC, bool bSaveSharedTempStorage, string CLString, string StoreName, string FakeFail)
+	void ExecuteNodes(ElectricCommander EC, List<BuildNode> OrdereredToDo, bool bFake, bool bFakeEC, bool bSaveSharedTempStorage, string CLString, string StoreName, string FakeFail)
 	{
         Dictionary<string, BuildNode> BuildProductToNodeMap = new Dictionary<string, BuildNode>();
 		foreach (BuildNode NodeToDo in OrdereredToDo)
@@ -2380,7 +1707,7 @@ public partial class GUBP : BuildCommand
             {
                 if (SaveSuccessRecords) 
                 {
-                    SaveStatus(NodeToDo, StartedTempStorageSuffix, NodeStoreName, bSaveSharedTempStorage, GameNameIfAny);
+                    EC.SaveStatus(NodeToDo, StartedTempStorageSuffix, NodeStoreName, bSaveSharedTempStorage, GameNameIfAny);
                 }
 				double BuildDuration = 0.0;
                 try
@@ -2415,7 +1742,7 @@ public partial class GUBP : BuildCommand
 						Log("Took {0} seconds to store build products", StoreDuration);
 						if (IsBuildMachine)
 						{
-							RunECTool(String.Format("setProperty \"/myJobStep/StoreDuration\" \"{0}\"", StoreDuration.ToString()));
+							EC.RunECTool(String.Format("setProperty \"/myJobStep/StoreDuration\" \"{0}\"", StoreDuration.ToString()));
 						}
 					}
                     if (ParseParam("StompCheck"))
@@ -2454,21 +1781,21 @@ public partial class GUBP : BuildCommand
 						}
 						using(TelemetryStopwatch SaveNodeStatusStopwatch = new TelemetryStopwatch("SaveNodeStatus"))
 						{
-							SaveStatus(NodeToDo, FailedTempStorageSuffix, NodeStoreName, bSaveSharedTempStorage, GameNameIfAny, ParseParamValue("MyJobStepId"));
+							EC.SaveStatus(NodeToDo, FailedTempStorageSuffix, NodeStoreName, bSaveSharedTempStorage, GameNameIfAny, ParseParamValue("MyJobStepId"));
 						}
 						using(TelemetryStopwatch UpdateECPropsStopwatch = new TelemetryStopwatch("UpdateECProps"))
 						{
-							UpdateECProps(NodeToDo);
+							EC.UpdateECProps(NodeToDo);
 						}
                         
 						if (IsBuildMachine)
 						{
 							using(TelemetryStopwatch GetFailEmailsStopwatch = new TelemetryStopwatch("GetFailEmails"))
 							{
-								GetFailureEmails(NodeToDo, History, CLString, StoreName);
+								GetFailureEmails(EC, NodeToDo, History, CLString, StoreName);
 							}
 						}
-						UpdateECBuildTime(NodeToDo, BuildDuration);
+						EC.UpdateECBuildTime(NodeToDo, BuildDuration);
                     }
 
                     Log("{0}", ExceptionToString(Ex));
@@ -2524,21 +1851,21 @@ public partial class GUBP : BuildCommand
 					}
 					using(TelemetryStopwatch SaveNodeStatusStopwatch = new TelemetryStopwatch("SaveNodeStatus"))
 					{
-						SaveStatus(NodeToDo, SucceededTempStorageSuffix, NodeStoreName, bSaveSharedTempStorage, GameNameIfAny);
+						EC.SaveStatus(NodeToDo, SucceededTempStorageSuffix, NodeStoreName, bSaveSharedTempStorage, GameNameIfAny);
 					}
 					using(TelemetryStopwatch UpdateECPropsStopwatch = new TelemetryStopwatch("UpdateECProps"))
 					{
-						UpdateECProps(NodeToDo);
+						EC.UpdateECProps(NodeToDo);
 					}
                     
 					if (IsBuildMachine)
 					{
 						using(TelemetryStopwatch GetFailEmailsStopwatch = new TelemetryStopwatch("GetFailEmails"))
 						{
-							GetFailureEmails(NodeToDo, History, CLString, StoreName);
+							GetFailureEmails(EC, NodeToDo, History, CLString, StoreName);
 						}
 					}
-					UpdateECBuildTime(NodeToDo, BuildDuration);
+					EC.UpdateECBuildTime(NodeToDo, BuildDuration);
                 }
             }
             foreach (string Product in NodeToDo.Node.BuildProducts)
@@ -2551,85 +1878,5 @@ public partial class GUBP : BuildCommand
             }
         }        
         PrintRunTime();
-    }
-    
-	/// <summary>
-	/// Sorts a list of nodes to display in EC. The default order is based on execution order and agent groups, whereas this function arranges nodes by
-	/// frequency then execution order, while trying to group nodes on parallel paths (eg. Mac/Windows editor nodes) together.
-	/// </summary>
-	static Dictionary<BuildNode, int> GetDisplayOrder(List<BuildNode> Nodes)
-	{
-		// Split the nodes into separate lists for each frequency
-		SortedDictionary<int, List<BuildNode>> NodesByFrequency = new SortedDictionary<int,List<BuildNode>>();
-		foreach(BuildNode Node in Nodes)
-		{
-			List<BuildNode> NodesByThisFrequency;
-			if(!NodesByFrequency.TryGetValue(Node.FrequencyShift, out NodesByThisFrequency))
-			{
-				NodesByThisFrequency = new List<BuildNode>();
-				NodesByFrequency.Add(Node.FrequencyShift, NodesByThisFrequency);
-			}
-			NodesByThisFrequency.Add(Node);
-		}
-
-		// Build the output list by scanning each frequency in order
-		HashSet<BuildNode> VisitedNodes = new HashSet<BuildNode>();
-		Dictionary<BuildNode, int> SortedNodes = new Dictionary<BuildNode,int>();
-		foreach(List<BuildNode> NodesByThisFrequency in NodesByFrequency.Values)
-		{
-			// Find a list of nodes in each display group. If the group name matches the node name, put that node at the front of the list.
-			Dictionary<string, List<BuildNode>> DisplayGroups = new Dictionary<string,List<BuildNode>>();
-			foreach(BuildNode Node in NodesByThisFrequency)
-			{
-				string GroupName = Node.Node.GetDisplayGroupName();
-				if(!DisplayGroups.ContainsKey(GroupName))
-				{
-					DisplayGroups.Add(GroupName, new List<BuildNode>{ Node });
-				}
-				else if(GroupName == Node.Name)
-				{
-					DisplayGroups[GroupName].Insert(0, Node);
-				}
-				else
-				{
-					DisplayGroups[GroupName].Add(Node);
-				}
-			}
-
-			// Build a list of ordering dependencies, putting all Mac nodes after Windows nodes with the same names.
-			Dictionary<BuildNode, List<BuildNode>> NodeDependencies = new Dictionary<BuildNode,List<BuildNode>>(Nodes.ToDictionary(x => x, x => x.AllDirectDependencies.ToList()));
-			foreach(KeyValuePair<string, List<BuildNode>> DisplayGroup in DisplayGroups)
-			{
-				List<BuildNode> GroupNodes = DisplayGroup.Value;
-				for (int Idx = 1; Idx < GroupNodes.Count; Idx++)
-				{
-					NodeDependencies[GroupNodes[Idx]].Add(GroupNodes[0]);
-				}
-			}
-
-			// Add nodes for each frequency into the master list, trying to match up different groups along the way
-			foreach(BuildNode FirstNode in NodesByThisFrequency)
-			{
-				List<BuildNode> GroupNodes = DisplayGroups[FirstNode.Node.GetDisplayGroupName()];
-				foreach(BuildNode GroupNode in GroupNodes)
-				{
-					AddNodeAndDependencies(GroupNode, NodeDependencies, VisitedNodes, SortedNodes);
-				}
-			}
-		}
-		return SortedNodes;
-	}
-
-	static void AddNodeAndDependencies(BuildNode Node, Dictionary<BuildNode, List<BuildNode>> NodeDependencies, HashSet<BuildNode> VisitedNodes, Dictionary<BuildNode, int> SortedNodes)
-	{
-		if(!VisitedNodes.Contains(Node))
-		{
-			VisitedNodes.Add(Node);
-			foreach (BuildNode NodeDependency in NodeDependencies[Node])
-			{
-				AddNodeAndDependencies(NodeDependency, NodeDependencies, VisitedNodes, SortedNodes);
-			}
-			SortedNodes.Add(Node, SortedNodes.Count);
-		}
-	}
+    }   
 }
