For for UAT RunSingleInstance not correctly honoring uebp_UATMutexNoWait. It would run the command and then try and take the mutex and run it again, which would fail when it couldn't take the mutex (because uebp_UATMutexNoWait is only set when another UAT is expected to already be running).

diff --git a/Engine/Source/Programs/AutomationTool/AutomationUtils/Utils.cs b/Engine/Source/Programs/AutomationTool/AutomationUtils/Utils.cs
--- a/Engine/Source/Programs/AutomationTool/AutomationUtils/Utils.cs
+++ b/Engine/Source/Programs/AutomationTool/AutomationUtils/Utils.cs
@@ -687,6 +687,7 @@ namespace AutomationTool
 			if (Environment.GetEnvironmentVariable("uebp_UATMutexNoWait") == "1")
 			{
 				Main(Param);
+                return;
 			}
 			var bCreatedMutex = false;
 			var LocationHash = InternalUtils.ExecutingAssemblyLocation.GetHashCode();
