Move the GUBP entry point to the top of the file.

diff --git a/Engine/Source/Programs/AutomationTool/GUBP/GUBP.cs b/Engine/Source/Programs/AutomationTool/GUBP/GUBP.cs
--- a/Engine/Source/Programs/AutomationTool/GUBP/GUBP.cs
+++ b/Engine/Source/Programs/AutomationTool/GUBP/GUBP.cs
@@ -11,6 +11,41 @@
 using System.Linq;
 using System.Diagnostics;
 
+[Help("Runs one, several or all of the GUBP nodes")]
+[Help(typeof(UE4Build))]
+[Help("NoMac", "Toggle to exclude the Mac host platform, default is Win64+Mac+Linux")]
+[Help("NoLinux", "Toggle to exclude the Linux (PC, 64-bit) host platform, default is Win64+Mac+Linux")]
+[Help("NoPC", "Toggle to exclude the PC host platform, default is Win64+Mac+Linux")]
+[Help("CleanLocal", "delete the local temp storage before we start")]
+[Help("Store=", "Sets the name of the temp storage block, normally, this is built for you.")]
+[Help("StoreSuffix=", "Tacked onto a store name constructed from CL, branch, etc")]
+[Help("TimeIndex=", "An integer used to determine subsets to run based on DependentCISFrequencyQuantumShift")]
+[Help("UserTimeIndex=", "An integer used to determine subsets to run based on DependentCISFrequencyQuantumShift, this one overrides TimeIndex")]
+[Help("PreflightUID=", "A unique integer tag from EC used as part of the tempstorage, builds and label names to distinguish multiple attempts.")]
+[Help("Node=", "Nodes to process, -node=Node1+Node2+Node3, if no nodes or games are specified, defaults to all nodes.")]
+[Help("TriggerNode=", "Trigger Nodes to process, -triggernode=Node.")]
+[Help("Game=", "Games to process, -game=Game1+Game2+Game3, if no games or nodes are specified, defaults to all nodes.")]
+[Help("ListOnly", "List Nodes in this branch")]
+[Help("SaveGraph", "Save graph as an xml file")]
+[Help("CommanderJobSetupOnly", "Set up the EC branch info via ectool and quit")]
+[Help("FakeEC", "don't run ectool, rather just do it locally, emulating what EC would have done.")]
+[Help("Fake", "Don't actually build anything, just store a record of success as the build product for each node.")]
+[Help("AllPlatforms", "Regardless of what is installed on this machine, set up the graph for all platforms; true by default on build machines.")]
+[Help("SkipTriggers", "ignore all triggers")]
+[Help("CL", "force the CL to something, disregarding the P4 value.")]
+[Help("History", "Like ListOnly, except gives you a full history. Must have -P4 for this to work.")]
+[Help("Changes", "Like history, but also shows the P4 changes. Must have -P4 for this to work.")]
+[Help("AllChanges", "Like changes except includes changes before the last green. Must have -P4 for this to work.")]
+[Help("EmailOnly", "Only emails the folks given in the argument.")]
+[Help("AddEmails", "Add these space delimited emails too all email lists.")]
+[Help("ShowDependencies", "Show node dependencies.")]
+[Help("ShowECDependencies", "Show EC node dependencies instead.")]
+[Help("ShowECProc", "Show EC proc names.")]
+[Help("ECProject", "From EC, the name of the project, used to get a version number.")]
+[Help("CIS", "This is a CIS run, assign TimeIndex based on the history.")]
+[Help("ForceIncrementalCompile", "make sure all compiles are incremental")]
+[Help("AutomatedTesting", "Allow automated testing, currently disabled.")]
+[Help("StompCheck", "Look for stomped build products.")]
 public partial class GUBP : BuildCommand
 {
 	const string StartedTempStorageSuffix = "_Started";
@@ -31,5 +66,202 @@
 	};
 
 	/// <summary>
+	/// Main entry point for GUBP
+	/// </summary>
+    public override void ExecuteBuild()
+    {
+        Log("************************* GUBP");
+
+		bool bPreflightBuild = false;
+		int PreflightShelveCL = 0;
+        string PreflightShelveCLString = GetEnvVar("uebp_PreflightShelveCL");
+        if ((!String.IsNullOrEmpty(PreflightShelveCLString) && IsBuildMachine) || ParseParam("PreflightTest"))
+        {
+            Log("**** Preflight shelve {0}", PreflightShelveCLString);
+            if (!String.IsNullOrEmpty(PreflightShelveCLString))
+            {
+                PreflightShelveCL = int.Parse(PreflightShelveCLString);
+                if (PreflightShelveCL < 2000000)
+                {
+                    throw new AutomationException(String.Format( "{0} does not look like a CL", PreflightShelveCL));
+                }
+            }
+            bPreflightBuild = true;
+        }
+        
+		List<UnrealTargetPlatform> HostPlatforms = new List<UnrealTargetPlatform>();
+        if (!ParseParam("NoPC"))
+        {
+            HostPlatforms.Add(UnrealTargetPlatform.Win64);
+        }
+        if (!ParseParam("NoMac"))
+        {
+			HostPlatforms.Add(UnrealTargetPlatform.Mac);
+        }
+		if(!ParseParam("NoLinux") && UnrealBuildTool.BuildHostPlatform.Current.Platform == UnrealTargetPlatform.Linux)
+		{
+			HostPlatforms.Add(UnrealTargetPlatform.Linux);
+		}
+
+        string StoreName = ParseParamValue("Store");
+        string StoreSuffix = ParseParamValue("StoreSuffix", "");
+
+		string PreflightMangleSuffix = "";
+        if (bPreflightBuild)
+        {
+            int PreflightUID = ParseParamInt("PreflightUID", 0);
+            PreflightMangleSuffix = String.Format("-PF-{0}-{1}", PreflightShelveCL, PreflightUID);
+            StoreSuffix = StoreSuffix + PreflightMangleSuffix;
+        }
+        int CL = ParseParamInt("CL", 0);
+        bool bCleanLocalTempStorage = ParseParam("CleanLocal");
+        bool bSkipTriggers = ParseParam("SkipTriggers");
+        bool bFake = ParseParam("fake");
+        bool bFakeEC = ParseParam("FakeEC");
+
+        bool bSaveSharedTempStorage = false;
+
+        bool LocalOnly = true;
+        string CLString = "";
+        if (String.IsNullOrEmpty(StoreName))
+        {
+            if (P4Enabled)
+            {
+                if (CL == 0)
+                {
+                    CL = P4Env.Changelist;
+                }
+                CLString = String.Format("{0}", CL);
+                StoreName = P4Env.BuildRootEscaped + "-" + CLString;
+                bSaveSharedTempStorage = CommandUtils.IsBuildMachine;
+                LocalOnly = false;
+            }
+            else
+            {
+                StoreName = "TempLocal";
+                bSaveSharedTempStorage = false;
+            }
+        }
+        StoreName = StoreName + StoreSuffix;
+        if (bFakeEC)
+        {
+            LocalOnly = true;
+        } 
+        if (bSaveSharedTempStorage)
+        {
+            if (!TempStorage.HaveSharedTempStorage(true))
+            {
+                throw new AutomationException("Request to save to temp storage, but {0} is unavailable.", TempStorage.UE4TempStorageDirectory());
+            }
+        }
+        else if (!LocalOnly && !TempStorage.HaveSharedTempStorage(false))
+        {
+            LogWarning("Looks like we want to use shared temp storage, but since we don't have it, we won't use it.");
+            LocalOnly = true;
+        }
+
+        bool CommanderSetup = ParseParam("CommanderJobSetupOnly");
+
+		string ExplicitTriggerName = "";
+		if (CommanderSetup)
+		{
+			ExplicitTriggerName = ParseParamValue("TriggerNode", "");
+		}
+
+		List<BuildNode> AllNodes;
+		List<AggregateNode> AllAggregates;
+		int TimeQuantum = 20;
+		AddNodesForBranch(CL, HostPlatforms, bPreflightBuild, PreflightMangleSuffix, out AllNodes, out AllAggregates, ref TimeQuantum);
+
+		LinkGraph(AllAggregates, AllNodes);
+		FindControllingTriggers(AllNodes);
+		FindCompletionState(AllNodes, StoreName, LocalOnly);
+		ComputeDependentFrequencies(AllNodes);
+
+        if (bCleanLocalTempStorage)  // shared temp storage can never be wiped
+        {
+            TempStorage.DeleteLocalTempStorageManifests(CmdEnv);
+        }
+
+		int TimeIndex = ParseParamInt("TimeIndex", 0);
+		if (TimeIndex == 0)
+		{
+			TimeIndex = ParseParamInt("UserTimeIndex", 0);
+		}
+		if (ParseParam("CIS") && ExplicitTriggerName == "" && CommanderSetup) // explicit triggers will already have a time index assigned
+		{
+			TimeIndex = UpdateCISCounter(TimeQuantum);
+			Log("Setting TimeIndex to {0}", TimeIndex);
+		}
+
+		HashSet<BuildNode> NodesToDo = ParseNodesToDo(AllNodes, AllAggregates);
+		CullNodesForTimeIndex(NodesToDo, TimeIndex);
+		CullNodesForPreflight(NodesToDo, bPreflightBuild);
+
+		BuildNode ExplicitTrigger = null;
+		if (CommanderSetup)
+        {
+            if (!String.IsNullOrEmpty(ExplicitTriggerName))
+            {
+                foreach (BuildNode Node in AllNodes)
+                {
+                    if (Node.Name.Equals(ExplicitTriggerName, StringComparison.InvariantCultureIgnoreCase))
+                    {
+                        if (Node.Node.TriggerNode())
+                        {
+                            Node.Node.SetAsExplicitTrigger();
+							ExplicitTrigger = Node;
+                            break;
+                        }
+                    }
+                }
+                if (ExplicitTrigger == null)
+                {
+                    throw new AutomationException("Could not find trigger node named {0}", ExplicitTriggerName);
+                }
+            }
+            else
+            {
+                if (bSkipTriggers)
+                {
+                    foreach (BuildNode Node in AllNodes)
+                    {
+                        if (Node.Node.TriggerNode())
+                        {
+                            Node.Node.SetAsExplicitTrigger();
+                        }
+                    }
+                }
+            }
+        }
+
+        List<BuildNode> OrderedToDo = TopologicalSort(NodesToDo, ExplicitTrigger, false, false);
+
+		List<BuildNode> UnfinishedTriggers = FindUnfinishedTriggers(bSkipTriggers, ExplicitTrigger, OrderedToDo);
+
+		PrintNodes(this, OrderedToDo, AllAggregates, UnfinishedTriggers, TimeQuantum);
+
+        //check sorting
+		CheckSortOrder(OrderedToDo);
+
+		ElectricCommander EC = new ElectricCommander(this);
+
+        string FakeFail = ParseParamValue("FakeFail");
+        if(CommanderSetup)
+        {
+			DoCommanderSetup(EC, AllNodes, AllAggregates, NodesToDo, OrderedToDo, TimeIndex, TimeQuantum, bSkipTriggers, bFake, bFakeEC, CLString, ExplicitTrigger, UnfinishedTriggers, FakeFail, bPreflightBuild);
+        }
+		else if(ParseParam("SaveGraph"))
+		{
+			SaveGraphVisualization(OrderedToDo);
+		}
+		else if(!ParseParam("ListOnly"))
+		{
+			ExecuteNodes(EC, OrderedToDo, bFake, bFakeEC, bSaveSharedTempStorage, CLString, StoreName, FakeFail);
+		}
+        PrintRunTime();
+	}
+
+	/// <summary>
 	/// Recursively update the ControllingTriggers array for each of the nodes passed in. 
 	/// </summary>
@@ -907,236 +1139,6 @@
         return Result;
     }
 
-    [Help("Runs one, several or all of the GUBP nodes")]
-    [Help(typeof(UE4Build))]
-    [Help("NoMac", "Toggle to exclude the Mac host platform, default is Win64+Mac+Linux")]
-	[Help("NoLinux", "Toggle to exclude the Linux (PC, 64-bit) host platform, default is Win64+Mac+Linux")]
-	[Help("NoPC", "Toggle to exclude the PC host platform, default is Win64+Mac+Linux")]
-    [Help("CleanLocal", "delete the local temp storage before we start")]
-    [Help("Store=", "Sets the name of the temp storage block, normally, this is built for you.")]
-    [Help("StoreSuffix=", "Tacked onto a store name constructed from CL, branch, etc")]
-    [Help("TimeIndex=", "An integer used to determine subsets to run based on DependentCISFrequencyQuantumShift")]
-    [Help("UserTimeIndex=", "An integer used to determine subsets to run based on DependentCISFrequencyQuantumShift, this one overrides TimeIndex")]
-    [Help("PreflightUID=", "A unique integer tag from EC used as part of the tempstorage, builds and label names to distinguish multiple attempts.")]
-    [Help("Node=", "Nodes to process, -node=Node1+Node2+Node3, if no nodes or games are specified, defaults to all nodes.")]
-    [Help("TriggerNode=", "Trigger Nodes to process, -triggernode=Node.")]
-    [Help("Game=", "Games to process, -game=Game1+Game2+Game3, if no games or nodes are specified, defaults to all nodes.")]
-    [Help("ListOnly", "List Nodes in this branch")]
-    [Help("SaveGraph", "Save graph as an xml file")]
-    [Help("CommanderJobSetupOnly", "Set up the EC branch info via ectool and quit")]
-    [Help("FakeEC", "don't run ectool, rather just do it locally, emulating what EC would have done.")]
-    [Help("Fake", "Don't actually build anything, just store a record of success as the build product for each node.")]
-    [Help("AllPlatforms", "Regardless of what is installed on this machine, set up the graph for all platforms; true by default on build machines.")]
-    [Help("SkipTriggers", "ignore all triggers")]
-    [Help("CL", "force the CL to something, disregarding the P4 value.")]
-    [Help("History", "Like ListOnly, except gives you a full history. Must have -P4 for this to work.")]
-    [Help("Changes", "Like history, but also shows the P4 changes. Must have -P4 for this to work.")]
-    [Help("AllChanges", "Like changes except includes changes before the last green. Must have -P4 for this to work.")]
-    [Help("EmailOnly", "Only emails the folks given in the argument.")]
-    [Help("AddEmails", "Add these space delimited emails too all email lists.")]
-    [Help("ShowDependencies", "Show node dependencies.")]
-    [Help("ShowECDependencies", "Show EC node dependencies instead.")]
-    [Help("ShowECProc", "Show EC proc names.")]
-    [Help("ECProject", "From EC, the name of the project, used to get a version number.")]
-    [Help("CIS", "This is a CIS run, assign TimeIndex based on the history.")]
-    [Help("ForceIncrementalCompile", "make sure all compiles are incremental")]
-    [Help("AutomatedTesting", "Allow automated testing, currently disabled.")]
-    [Help("StompCheck", "Look for stomped build products.")]
-
-    public override void ExecuteBuild()
-    {
-        Log("************************* GUBP");
-
-		bool bPreflightBuild = false;
-		int PreflightShelveCL = 0;
-        string PreflightShelveCLString = GetEnvVar("uebp_PreflightShelveCL");
-        if ((!String.IsNullOrEmpty(PreflightShelveCLString) && IsBuildMachine) || ParseParam("PreflightTest"))
-        {
-            Log("**** Preflight shelve {0}", PreflightShelveCLString);
-            if (!String.IsNullOrEmpty(PreflightShelveCLString))
-            {
-                PreflightShelveCL = int.Parse(PreflightShelveCLString);
-                if (PreflightShelveCL < 2000000)
-                {
-                    throw new AutomationException(String.Format( "{0} does not look like a CL", PreflightShelveCL));
-                }
-            }
-            bPreflightBuild = true;
-        }
-        
-		List<UnrealTargetPlatform> HostPlatforms = new List<UnrealTargetPlatform>();
-        if (!ParseParam("NoPC"))
-        {
-            HostPlatforms.Add(UnrealTargetPlatform.Win64);
-        }
-        if (!ParseParam("NoMac"))
-        {
-			HostPlatforms.Add(UnrealTargetPlatform.Mac);
-        }
-		if(!ParseParam("NoLinux") && UnrealBuildTool.BuildHostPlatform.Current.Platform == UnrealTargetPlatform.Linux)
-		{
-			HostPlatforms.Add(UnrealTargetPlatform.Linux);
-		}
-
-        string StoreName = ParseParamValue("Store");
-        string StoreSuffix = ParseParamValue("StoreSuffix", "");
-
-		string PreflightMangleSuffix = "";
-        if (bPreflightBuild)
-        {
-            int PreflightUID = ParseParamInt("PreflightUID", 0);
-            PreflightMangleSuffix = String.Format("-PF-{0}-{1}", PreflightShelveCL, PreflightUID);
-            StoreSuffix = StoreSuffix + PreflightMangleSuffix;
-        }
-        int CL = ParseParamInt("CL", 0);
-        bool bCleanLocalTempStorage = ParseParam("CleanLocal");
-        bool bSkipTriggers = ParseParam("SkipTriggers");
-        bool bFake = ParseParam("fake");
-        bool bFakeEC = ParseParam("FakeEC");
-
-        bool bSaveSharedTempStorage = false;
-
-        bool LocalOnly = true;
-        string CLString = "";
-        if (String.IsNullOrEmpty(StoreName))
-        {
-            if (P4Enabled)
-            {
-                if (CL == 0)
-                {
-                    CL = P4Env.Changelist;
-                }
-                CLString = String.Format("{0}", CL);
-                StoreName = P4Env.BuildRootEscaped + "-" + CLString;
-                bSaveSharedTempStorage = CommandUtils.IsBuildMachine;
-                LocalOnly = false;
-            }
-            else
-            {
-                StoreName = "TempLocal";
-                bSaveSharedTempStorage = false;
-            }
-        }
-        StoreName = StoreName + StoreSuffix;
-        if (bFakeEC)
-        {
-            LocalOnly = true;
-        } 
-        if (bSaveSharedTempStorage)
-        {
-            if (!TempStorage.HaveSharedTempStorage(true))
-            {
-                throw new AutomationException("Request to save to temp storage, but {0} is unavailable.", TempStorage.UE4TempStorageDirectory());
-            }
-        }
-        else if (!LocalOnly && !TempStorage.HaveSharedTempStorage(false))
-        {
-            LogWarning("Looks like we want to use shared temp storage, but since we don't have it, we won't use it.");
-            LocalOnly = true;
-        }
-
-        bool CommanderSetup = ParseParam("CommanderJobSetupOnly");
-
-		string ExplicitTriggerName = "";
-		if (CommanderSetup)
-		{
-			ExplicitTriggerName = ParseParamValue("TriggerNode", "");
-		}
-
-		List<BuildNode> AllNodes;
-		List<AggregateNode> AllAggregates;
-		int TimeQuantum = 20;
-		AddNodesForBranch(CL, HostPlatforms, bPreflightBuild, PreflightMangleSuffix, out AllNodes, out AllAggregates, ref TimeQuantum);
-
-		LinkGraph(AllAggregates, AllNodes);
-		FindControllingTriggers(AllNodes);
-		FindCompletionState(AllNodes, StoreName, LocalOnly);
-		ComputeDependentFrequencies(AllNodes);
-
-        if (bCleanLocalTempStorage)  // shared temp storage can never be wiped
-        {
-            TempStorage.DeleteLocalTempStorageManifests(CmdEnv);
-        }
-
-		int TimeIndex = ParseParamInt("TimeIndex", 0);
-		if (TimeIndex == 0)
-		{
-			TimeIndex = ParseParamInt("UserTimeIndex", 0);
-		}
-		if (ParseParam("CIS") && ExplicitTriggerName == "" && CommanderSetup) // explicit triggers will already have a time index assigned
-		{
-			TimeIndex = UpdateCISCounter(TimeQuantum);
-			Log("Setting TimeIndex to {0}", TimeIndex);
-		}
-
-		HashSet<BuildNode> NodesToDo = ParseNodesToDo(AllNodes, AllAggregates);
-		CullNodesForTimeIndex(NodesToDo, TimeIndex);
-		CullNodesForPreflight(NodesToDo, bPreflightBuild);
-
-		BuildNode ExplicitTrigger = null;
-		if (CommanderSetup)
-        {
-            if (!String.IsNullOrEmpty(ExplicitTriggerName))
-            {
-                foreach (BuildNode Node in AllNodes)
-                {
-                    if (Node.Name.Equals(ExplicitTriggerName, StringComparison.InvariantCultureIgnoreCase))
-                    {
-                        if (Node.Node.TriggerNode())
-                        {
-                            Node.Node.SetAsExplicitTrigger();
-							ExplicitTrigger = Node;
-                            break;
-                        }
-                    }
-                }
-                if (ExplicitTrigger == null)
-                {
-                    throw new AutomationException("Could not find trigger node named {0}", ExplicitTriggerName);
-                }
-            }
-            else
-            {
-                if (bSkipTriggers)
-                {
-                    foreach (BuildNode Node in AllNodes)
-                    {
-                        if (Node.Node.TriggerNode())
-                        {
-                            Node.Node.SetAsExplicitTrigger();
-                        }
-                    }
-                }
-            }
-        }
-
-        List<BuildNode> OrderedToDo = TopologicalSort(NodesToDo, ExplicitTrigger, false, false);
-
-		List<BuildNode> UnfinishedTriggers = FindUnfinishedTriggers(bSkipTriggers, ExplicitTrigger, OrderedToDo);
-
-		PrintNodes(this, OrderedToDo, AllAggregates, UnfinishedTriggers, TimeQuantum);
-
-        //check sorting
-		CheckSortOrder(OrderedToDo);
-
-		ElectricCommander EC = new ElectricCommander(this);
-
-        string FakeFail = ParseParamValue("FakeFail");
-        if(CommanderSetup)
-        {
-			DoCommanderSetup(EC, AllNodes, AllAggregates, NodesToDo, OrderedToDo, TimeIndex, TimeQuantum, bSkipTriggers, bFake, bFakeEC, CLString, ExplicitTrigger, UnfinishedTriggers, FakeFail, bPreflightBuild);
-        }
-		else if(ParseParam("SaveGraph"))
-		{
-			SaveGraphVisualization(OrderedToDo);
-		}
-		else if(!ParseParam("ListOnly"))
-		{
-			ExecuteNodes(EC, OrderedToDo, bFake, bFakeEC, bSaveSharedTempStorage, CLString, StoreName, FakeFail);
-		}
-        PrintRunTime();
-	}
-
 	/// <summary>
 	/// Resolves the names of each node and aggregates' dependencies, and links them together into the build graph.
 	/// </summary>
