Add a separate LegacyNode class for wrapping around old GUBPNode objects.

diff --git a/Engine/Source/Programs/AutomationTool/GUBP/ElectricCommander.cs b/Engine/Source/Programs/AutomationTool/GUBP/ElectricCommander.cs
--- a/Engine/Source/Programs/AutomationTool/GUBP/ElectricCommander.cs
+++ b/Engine/Source/Programs/AutomationTool/GUBP/ElectricCommander.cs
@@ -51,13 +51,13 @@ namespace AutomationTool
 
 		static string GetJobStepPath(BuildNode Dep)
 		{
-			if (Dep.Node.AgentSharingGroup == "")
+			if (Dep.AgentSharingGroup == "")
 			{
 				return "jobSteps[" + Dep.Name + "]";
 			}
 			else
 			{
-				return "jobSteps[" + Dep.Node.AgentSharingGroup + "]/jobSteps[" + Dep.Name + "]";
+				return "jobSteps[" + Dep.AgentSharingGroup + "]/jobSteps[" + Dep.Name + "]";
 			}
 		}
 		static string GetJobStep(string ParentPath, BuildNode Dep)
@@ -71,13 +71,13 @@ namespace AutomationTool
 			ECProps.Add("FailEmails/" + NodeToDo.Name + "=" + String.Join(" ", NodeToDo.RecipientsForFailureEmails));
 	
 			string AgentReq = NodeToDo.Node.ECAgentString();
-			if(Command.ParseParamValue("AgentOverride") != "" && !NodeToDo.Node.GetFullName().Contains("OnMac"))
+			if(Command.ParseParamValue("AgentOverride") != "" && !NodeToDo.Name.Contains("OnMac"))
 			{
 				AgentReq = Command.ParseParamValue("AgentOverride");
 			}
 			ECProps.Add(string.Format("AgentRequirementString/{0}={1}", NodeToDo.Name, AgentReq));
 			ECProps.Add(string.Format("RequiredMemory/{0}={1}", NodeToDo.Name, NodeToDo.AgentMemoryRequirement));
-			ECProps.Add(string.Format("Timeouts/{0}={1}", NodeToDo.Name, NodeToDo.Node.TimeoutInMinutes()));
+			ECProps.Add(string.Format("Timeouts/{0}={1}", NodeToDo.Name, NodeToDo.TimeoutInMinutes));
 			ECProps.Add(string.Format("JobStepPath/{0}={1}", NodeToDo.Name, GetJobStepPath(NodeToDo)));
 		
 			return ECProps;
@@ -140,13 +140,13 @@ namespace AutomationTool
 			RunECTool(String.Format("setProperty \"/myWorkflow/EmailNotes/{0}\" \"{1}\"", NodeToDo.Name, EMailNote));
 			{
 				HashSet<string> Emails = new HashSet<string>(NodeToDo.RecipientsForFailureEmails);
-				if (Command.ParseParam("CIS") && !NodeToDo.Node.SendSuccessEmail() && !NodeToDo.Node.TriggerNode() && NodeToDo.AddSubmittersToFailureEmails)
+				if (Command.ParseParam("CIS") && !NodeToDo.SendSuccessEmail && !NodeToDo.Node.TriggerNode() && NodeToDo.AddSubmittersToFailureEmails)
 				{
 					Emails.UnionWith(FailCauserEMails.Split(new char[] { ' ' }, StringSplitOptions.RemoveEmptyEntries));
 				}
 				RunECTool(String.Format("setProperty \"/myWorkflow/FailEmails/{0}\" \"{1}\"", NodeToDo.Name, String.Join(" ", Emails)));
 			}
-			if (NodeToDo.Node.SendSuccessEmail() || SendSuccessForGreenAfterRed)
+			if (NodeToDo.SendSuccessEmail || SendSuccessForGreenAfterRed)
 			{
 				RunECTool(String.Format("setProperty \"/myWorkflow/SendSuccessEmail/{0}\" \"{1}\"", NodeToDo.Name, "1"));
 			}
@@ -269,7 +269,7 @@ namespace AutomationTool
 				{
 					if (!NodeToDo.IsComplete) // if something is already finished, we don't put it into EC  
 					{
-						string MyAgentGroup = NodeToDo.Node.AgentSharingGroup;
+						string MyAgentGroup = NodeToDo.AgentSharingGroup;
 						if (MyAgentGroup != "")
 						{
 							if (!AgentGroupChains.ContainsKey(MyAgentGroup))
@@ -344,7 +344,7 @@ namespace AutomationTool
 						string PreCondition = GetPreConditionForNode(OrdereredToDo, ParentPath, bHasNoop, AgentGroupChains, StickyChain, NodeToDo, UncompletedEcDeps);
 						string RunCondition = GetRunConditionForNode(UncompletedEcDeps, ParentPath);
 
-						string MyAgentGroup = NodeToDo.Node.AgentSharingGroup;
+						string MyAgentGroup = NodeToDo.AgentSharingGroup;
 						bool bDoNestedJobstep = false;
 						bool bDoFirstNestedJobstep = false;
 
@@ -496,13 +496,13 @@ namespace AutomationTool
 		private string GetPreConditionForNode(List<BuildNode> OrdereredToDo, string PreconditionParentPath, bool bHasNoop, Dictionary<string, List<BuildNode>> AgentGroupChains, List<BuildNode> StickyChain, BuildNode NodeToDo, List<BuildNode> UncompletedEcDeps)
 		{
 			List<BuildNode> PreConditionUncompletedEcDeps = new List<BuildNode>();
-			if(NodeToDo.Node.AgentSharingGroup == "")
+			if(NodeToDo.AgentSharingGroup == "")
 			{
 				PreConditionUncompletedEcDeps = new List<BuildNode>(UncompletedEcDeps);
 			}
 			else 
 			{
-				List<BuildNode> MyChain = AgentGroupChains[NodeToDo.Node.AgentSharingGroup];
+				List<BuildNode> MyChain = AgentGroupChains[NodeToDo.AgentSharingGroup];
 				int MyIndex = MyChain.IndexOf(NodeToDo);
 				if (MyIndex > 0)
 				{
diff --git a/Engine/Source/Programs/AutomationTool/GUBP/GUBP.Automation.csproj b/Engine/Source/Programs/AutomationTool/GUBP/GUBP.Automation.csproj
--- a/Engine/Source/Programs/AutomationTool/GUBP/GUBP.Automation.csproj
+++ b/Engine/Source/Programs/AutomationTool/GUBP/GUBP.Automation.csproj
@@ -64,6 +64,7 @@
     <Compile Include="LegacyNodes.cs" />
     <Compile Include="Nodes\AggregateNode.cs" />
     <Compile Include="Nodes\BuildNode.cs" />
+    <Compile Include="Nodes\LegacyNode.cs" />
     <Compile Include="TempStorage.cs">
       <SubType>Code</SubType>
     </Compile>
diff --git a/Engine/Source/Programs/AutomationTool/GUBP/GUBP.cs b/Engine/Source/Programs/AutomationTool/GUBP/GUBP.cs
--- a/Engine/Source/Programs/AutomationTool/GUBP/GUBP.cs
+++ b/Engine/Source/Programs/AutomationTool/GUBP/GUBP.cs
@@ -541,11 +541,11 @@ public partial class GUBP : BuildCommand
                 }
             }
 
-			if (NodeToDo.Node.AgentSharingGroup != LastAgentGroup && NodeToDo.Node.AgentSharingGroup != "")
+			if (NodeToDo.AgentSharingGroup != LastAgentGroup && NodeToDo.AgentSharingGroup != "")
             {
-                Log("    Agent Group: {0}", NodeToDo.Node.AgentSharingGroup);
+                Log("    Agent Group: {0}", NodeToDo.AgentSharingGroup);
             }
-            LastAgentGroup = NodeToDo.Node.AgentSharingGroup;
+            LastAgentGroup = NodeToDo.AgentSharingGroup;
 
 			StringBuilder Builder = new StringBuilder("      ");
 			if(LastAgentGroup != "")
@@ -721,7 +721,7 @@ public partial class GUBP : BuildCommand
             Dictionary<string, List<BuildNode>> AgentGroupChains = new Dictionary<string, List<BuildNode>>();
             foreach (BuildNode NodeToDo in NodesToDo)
             {
-                string MyAgentGroup = NodeToDo.Node.AgentSharingGroup;
+                string MyAgentGroup = NodeToDo.AgentSharingGroup;
                 if (MyAgentGroup != "")
                 {
                     if (!AgentGroupChains.ContainsKey(MyAgentGroup))
@@ -756,21 +756,21 @@ public partial class GUBP : BuildCommand
                 bool bReady = true;
 				bool bPseudoReady = true;
 				bool bCompReady = true;
-                if (!SubSort && NodeToDo.Node.AgentSharingGroup != "")
+                if (!SubSort && NodeToDo.AgentSharingGroup != "")
                 {
-                    if (ExaminedAgentGroups.Contains(NodeToDo.Node.AgentSharingGroup))
+                    if (ExaminedAgentGroups.Contains(NodeToDo.AgentSharingGroup))
                     {
-                        bReady = !NonReadyAgentGroups.Contains(NodeToDo.Node.AgentSharingGroup);
-                        bPseudoReady = !NonPeudoReadyAgentGroups.Contains(NodeToDo.Node.AgentSharingGroup); //this might not be accurate if bReady==false
+                        bReady = !NonReadyAgentGroups.Contains(NodeToDo.AgentSharingGroup);
+                        bPseudoReady = !NonPeudoReadyAgentGroups.Contains(NodeToDo.AgentSharingGroup); //this might not be accurate if bReady==false
                     }
                     else
                     {
-                        ExaminedAgentGroups.Add(NodeToDo.Node.AgentSharingGroup);
-                        foreach (BuildNode ChainNode in SortedAgentGroupChains[NodeToDo.Node.AgentSharingGroup])
+                        ExaminedAgentGroups.Add(NodeToDo.AgentSharingGroup);
+                        foreach (BuildNode ChainNode in SortedAgentGroupChains[NodeToDo.AgentSharingGroup])
                         {
                             foreach (BuildNode Dep in ChainNode.Dependencies)
                             {
-                                if (!SortedAgentGroupChains[NodeToDo.Node.AgentSharingGroup].Contains(Dep) && NodesToDo.Contains(Dep))
+                                if (!SortedAgentGroupChains[NodeToDo.AgentSharingGroup].Contains(Dep) && NodesToDo.Contains(Dep))
                                 {
                                     bReady = false;
                                     break;
@@ -778,15 +778,15 @@ public partial class GUBP : BuildCommand
                             }
                             if (!bReady)
                             {
-                                NonReadyAgentGroups.Add(NodeToDo.Node.AgentSharingGroup);
+                                NonReadyAgentGroups.Add(NodeToDo.AgentSharingGroup);
                                 break;
                             }
                             foreach (BuildNode Dep in ChainNode.PseudoDependencies)
                             {
-                                if (!SortedAgentGroupChains[NodeToDo.Node.AgentSharingGroup].Contains(Dep) && NodesToDo.Contains(Dep))
+                                if (!SortedAgentGroupChains[NodeToDo.AgentSharingGroup].Contains(Dep) && NodesToDo.Contains(Dep))
                                 {
                                     bPseudoReady = false;
-                                    NonPeudoReadyAgentGroups.Add(NodeToDo.Node.AgentSharingGroup);
+                                    NonPeudoReadyAgentGroups.Add(NodeToDo.AgentSharingGroup);
                                     break;
                                 }
                             }
@@ -873,9 +873,9 @@ public partial class GUBP : BuildCommand
             }
             if (bProgressMade)
             {
-                if (!SubSort && BestNode.Node.AgentSharingGroup != "")
+                if (!SubSort && BestNode.AgentSharingGroup != "")
                 {
-                    foreach (BuildNode ChainNode in SortedAgentGroupChains[BestNode.Node.AgentSharingGroup])
+                    foreach (BuildNode ChainNode in SortedAgentGroupChains[BestNode.AgentSharingGroup])
                     {
                         OrdereredToDo.Add(ChainNode);
                         NodesToDo.Remove(ChainNode);
@@ -894,15 +894,15 @@ public partial class GUBP : BuildCommand
                 foreach (BuildNode NodeToDo in NodesToDo)
                 {
                     string Deps = "";
-                    if (!SubSort && NodeToDo.Node.AgentSharingGroup != "")
+                    if (!SubSort && NodeToDo.AgentSharingGroup != "")
                     {
-                        foreach (BuildNode ChainNode in SortedAgentGroupChains[NodeToDo.Node.AgentSharingGroup])
+                        foreach (BuildNode ChainNode in SortedAgentGroupChains[NodeToDo.AgentSharingGroup])
                         {
                             foreach (BuildNode Dep in ChainNode.Dependencies)
                             {
-                                if (!SortedAgentGroupChains[NodeToDo.Node.AgentSharingGroup].Contains(Dep) && NodesToDo.Contains(Dep))
+                                if (!SortedAgentGroupChains[NodeToDo.AgentSharingGroup].Contains(Dep) && NodesToDo.Contains(Dep))
                                 {
-                                    Deps = Deps + Dep.Name + "[" + ChainNode.Name + "->" + NodeToDo.Node.AgentSharingGroup + "]" + " ";
+                                    Deps = Deps + Dep.Name + "[" + ChainNode.Name + "->" + NodeToDo.AgentSharingGroup + "]" + " ";
                                 }
                             }
                         }
@@ -942,7 +942,7 @@ public partial class GUBP : BuildCommand
         if (!NodeToDo.Node.TriggerNode() && CLString != "")
         {
             string GameNameIfAny = NodeToDo.Node.GameNameIfAnyForTempStorage();
-            string NodeStoreWildCard = StoreName.Replace(CLString, "*") + "-" + NodeToDo.Node.GetFullName();
+            string NodeStoreWildCard = StoreName.Replace(CLString, "*") + "-" + NodeToDo.Name;
             History = new NodeHistory();
 
             History.AllStarted = ConvertCLToIntList(TempStorage.FindTempStorageManifests(CmdEnv, NodeStoreWildCard + StartedTempStorageSuffix, false, true, GameNameIfAny));
@@ -1084,7 +1084,7 @@ public partial class GUBP : BuildCommand
         int Result = P4Env.Changelist;
 
         string GameNameIfAny = NodeToDo.Node.GameNameIfAnyForTempStorage();
-        string NodeStore = StoreName + "-" + NodeToDo.Node.GetFullName() + FailedTempStorageSuffix;
+        string NodeStore = StoreName + "-" + NodeToDo.Name + FailedTempStorageSuffix;
 
         List<int> BackwardsFails = new List<int>(History.AllFailed);
         BackwardsFails.Add(P4Env.Changelist);
@@ -1324,7 +1324,7 @@ public partial class GUBP : BuildCommand
 			int FoundNames = 0;
 			foreach (BuildNode PotentialNode in PotentialNodes)
 			{
-				if (String.Compare(PotentialNode.Name, NameToDo, StringComparison.InvariantCultureIgnoreCase) == 0 || String.Compare(PotentialNode.Node.AgentSharingGroup, NameToDo, StringComparison.InvariantCultureIgnoreCase) == 0)
+				if (String.Compare(PotentialNode.Name, NameToDo, StringComparison.InvariantCultureIgnoreCase) == 0 || String.Compare(PotentialNode.AgentSharingGroup, NameToDo, StringComparison.InvariantCultureIgnoreCase) == 0)
 				{
 					NodesToDo.Add(PotentialNode);
 					FoundNames++;
@@ -1657,7 +1657,7 @@ public partial class GUBP : BuildCommand
                 }
             }
 
-            string NodeStoreName = StoreName + "-" + NodeToDo.Node.GetFullName();
+            string NodeStoreName = StoreName + "-" + NodeToDo.Name;
             
             string GameNameIfAny = NodeToDo.Node.GameNameIfAnyForTempStorage();
             string StorageRootIfAny = NodeToDo.Node.RootIfAnyForTempStorage();
@@ -1671,18 +1671,18 @@ public partial class GUBP : BuildCommand
             bool SaveSuccessRecords = (IsBuildMachine || bFakeEC) && // no real reason to make these locally except for fakeEC tests
                 (!NodeToDo.Node.TriggerNode() || NodeToDo.Node.IsSticky()); // trigger nodes are run twice, one to start the new workflow and once when it is actually triggered, we will save reconds for the latter
 
-            Log("***** Running GUBP Node {0} -> {1} : {2}", NodeToDo.Node.GetFullName(), GameNameIfAny, NodeStoreName);
+            Log("***** Running GUBP Node {0} -> {1} : {2}", NodeToDo.Name, GameNameIfAny, NodeStoreName);
             if (NodeToDo.IsComplete)
             {
                 if (NodeToDo.Name == VersionFilesNode.StaticGetFullName() && !IsBuildMachine)
                 {
-                    Log("***** NOT ****** Retrieving GUBP Node {0} from {1}; it is the version files.", NodeToDo.Node.GetFullName(), NodeStoreName);
+                    Log("***** NOT ****** Retrieving GUBP Node {0} from {1}; it is the version files.", NodeToDo.Name, NodeStoreName);
                     NodeToDo.Node.BuildProducts = new List<string>();
 
                 }
                 else
                 {
-                    Log("***** Retrieving GUBP Node {0} from {1}", NodeToDo.Node.GetFullName(), NodeStoreName);
+                    Log("***** Retrieving GUBP Node {0} from {1}", NodeToDo.Name, NodeStoreName);
                     bool WasLocal;
 					try
 					{
diff --git a/Engine/Source/Programs/AutomationTool/GUBP/LegacyBranchSetup.cs b/Engine/Source/Programs/AutomationTool/GUBP/LegacyBranchSetup.cs
--- a/Engine/Source/Programs/AutomationTool/GUBP/LegacyBranchSetup.cs
+++ b/Engine/Source/Programs/AutomationTool/GUBP/LegacyBranchSetup.cs
@@ -1203,7 +1203,7 @@ partial class GUBP
 
 		// Convert the GUBPNodes and GUBPAggregates maps into lists
 		AllNodes = new List<BuildNode>();
-		AllNodes.AddRange(BranchConfig.GUBPNodes.Values.Select(x => new BuildNode{ Name = x.GetFullName(), Node = x, AgentMemoryRequirement = x.AgentMemoryRequirement(BranchConfig) }));
+		AllNodes.AddRange(BranchConfig.GUBPNodes.Values.Select(x => new LegacyNode(BranchConfig, x)));
 
 		AllAggregates = new List<AggregateNode>();
 		AllAggregates.AddRange(BranchConfig.GUBPAggregates.Values.Select(x => new AggregateNode{ Name = x.GetFullName(), Node = x }));
diff --git a/Engine/Source/Programs/AutomationTool/GUBP/Nodes/BuildNode.cs b/Engine/Source/Programs/AutomationTool/GUBP/Nodes/BuildNode.cs
--- a/Engine/Source/Programs/AutomationTool/GUBP/Nodes/BuildNode.cs
+++ b/Engine/Source/Programs/AutomationTool/GUBP/Nodes/BuildNode.cs
@@ -7,20 +7,24 @@ using System.Text;
 namespace AutomationTool
 {
 	[DebuggerDisplay("{Name}")]
-	class BuildNode
+	abstract class BuildNode
 	{
 		public string Name;
+		public string AgentSharingGroup;
+		public int FrequencyShift;
+		public int AgentMemoryRequirement;
+		public int TimeoutInMinutes;
+
 		public GUBP.GUBPNode Node;
 		public BuildNode[] Dependencies;
 		public BuildNode[] PseudoDependencies;
 		public BuildNode[] AllDirectDependencies;
 		public BuildNode[] AllIndirectDependencies;
 		public BuildNode[] ControllingTriggers;
-		public int FrequencyShift;
 		public bool IsComplete;
 		public string[] RecipientsForFailureEmails;
 		public bool AddSubmittersToFailureEmails;
-		public int AgentMemoryRequirement;
+		public bool SendSuccessEmail;
 
 		public string ControllingTriggerDotName
 		{
diff --git a/Engine/Source/Programs/AutomationTool/GUBP/Nodes/LegacyNode.cs b/Engine/Source/Programs/AutomationTool/GUBP/Nodes/LegacyNode.cs
new file mode 100644
--- /dev/null
+++ b/Engine/Source/Programs/AutomationTool/GUBP/Nodes/LegacyNode.cs
@@ -0,0 +1,21 @@
+﻿using System;
+using System.Collections.Generic;
+using System.Diagnostics;
+using System.Linq;
+using System.Text;
+
+namespace AutomationTool
+{
+	class LegacyNode : BuildNode
+	{
+		public LegacyNode(GUBP.GUBPBranchConfig InBranchConfig, GUBP.GUBPNode InNode)
+		{
+			Name = InNode.GetFullName();
+			Node = InNode;
+			AgentSharingGroup = Node.AgentSharingGroup;
+			AgentMemoryRequirement = Node.AgentMemoryRequirement(InBranchConfig);
+			TimeoutInMinutes = Node.TimeoutInMinutes();
+			SendSuccessEmail = Node.SendSuccessEmail();
+		}
+	}
+}
