Add a class for trigger nodes, and use explicit typing for triggers wherever we can.

diff --git a/Engine/Source/Programs/AutomationTool/GUBP/ElectricCommander.cs b/Engine/Source/Programs/AutomationTool/GUBP/ElectricCommander.cs
--- a/Engine/Source/Programs/AutomationTool/GUBP/ElectricCommander.cs
+++ b/Engine/Source/Programs/AutomationTool/GUBP/ElectricCommander.cs
@@ -140,7 +140,7 @@ namespace AutomationTool
 			RunECTool(String.Format("setProperty \"/myWorkflow/EmailNotes/{0}\" \"{1}\"", NodeToDo.Name, EMailNote));
 			{
 				HashSet<string> Emails = new HashSet<string>(NodeToDo.RecipientsForFailureEmails);
-				if (Command.ParseParam("CIS") && !NodeToDo.SendSuccessEmail && !NodeToDo.Node.TriggerNode() && NodeToDo.AddSubmittersToFailureEmails)
+				if (Command.ParseParam("CIS") && !NodeToDo.SendSuccessEmail && NodeToDo.AddSubmittersToFailureEmails)
 				{
 					Emails.UnionWith(FailCauserEMails.Split(new char[] { ' ' }, StringSplitOptions.RemoveEmptyEntries));
 				}
@@ -156,7 +156,7 @@ namespace AutomationTool
 			}
 		}
 
-		public void DoCommanderSetup(IEnumerable<BuildNode> AllNodes, IEnumerable<AggregateNode> AllAggregates, HashSet<BuildNode> NodesToDo, List<BuildNode> OrdereredToDo, List<BuildNode> SortedNodes, int TimeIndex, int TimeQuantum, bool bSkipTriggers, bool bFake, bool bFakeEC, string CLString, BuildNode ExplicitTrigger, List<BuildNode> UnfinishedTriggers, string FakeFail)
+		public void DoCommanderSetup(IEnumerable<BuildNode> AllNodes, IEnumerable<AggregateNode> AllAggregates, HashSet<BuildNode> NodesToDo, List<BuildNode> OrdereredToDo, List<BuildNode> SortedNodes, int TimeIndex, int TimeQuantum, bool bSkipTriggers, bool bFake, bool bFakeEC, string CLString, TriggerNode ExplicitTrigger, List<TriggerNode> UnfinishedTriggers, string FakeFail)
 		{
 			Dictionary<string, string> FullNodeDependedOnBy = GetFullNodeDependedOnBy(NodesToDo);
 
diff --git a/Engine/Source/Programs/AutomationTool/GUBP/GUBP.Automation.csproj b/Engine/Source/Programs/AutomationTool/GUBP/GUBP.Automation.csproj
--- a/Engine/Source/Programs/AutomationTool/GUBP/GUBP.Automation.csproj
+++ b/Engine/Source/Programs/AutomationTool/GUBP/GUBP.Automation.csproj
@@ -65,6 +65,7 @@
     <Compile Include="Nodes\AggregateNode.cs" />
     <Compile Include="Nodes\BuildNode.cs" />
     <Compile Include="Nodes\LegacyNode.cs" />
+    <Compile Include="Nodes\TriggerNode.cs" />
     <Compile Include="TempStorage.cs">
       <SubType>Code</SubType>
     </Compile>
diff --git a/Engine/Source/Programs/AutomationTool/GUBP/GUBP.cs b/Engine/Source/Programs/AutomationTool/GUBP/GUBP.cs
--- a/Engine/Source/Programs/AutomationTool/GUBP/GUBP.cs
+++ b/Engine/Source/Programs/AutomationTool/GUBP/GUBP.cs
@@ -198,21 +198,18 @@ public partial class GUBP : BuildCommand
 		CullNodesForTimeIndex(NodesToDo, TimeIndex);
 		CullNodesForPreflight(NodesToDo, bPreflightBuild);
 
-		BuildNode ExplicitTrigger = null;
+		TriggerNode ExplicitTrigger = null;
 		if (CommanderSetup)
         {
             if (!String.IsNullOrEmpty(ExplicitTriggerName))
             {
-                foreach (BuildNode Node in AllNodes)
+                foreach (TriggerNode Node in AllNodes.OfType<TriggerNode>())
                 {
                     if (Node.Name.Equals(ExplicitTriggerName, StringComparison.InvariantCultureIgnoreCase))
                     {
-                        if (Node.Node.TriggerNode())
-                        {
-                            Node.Node.SetAsExplicitTrigger();
-							ExplicitTrigger = Node;
-                            break;
-                        }
+                        Node.Node.SetAsExplicitTrigger();
+						ExplicitTrigger = Node;
+                        break;
                     }
                 }
                 if (ExplicitTrigger == null)
@@ -224,12 +221,9 @@ public partial class GUBP : BuildCommand
             {
                 if (bSkipTriggers)
                 {
-                    foreach (BuildNode Node in AllNodes)
+                    foreach (TriggerNode Node in AllNodes.OfType<TriggerNode>())
                     {
-                        if (Node.Node.TriggerNode())
-                        {
-                            Node.Node.SetAsExplicitTrigger();
-                        }
+                        Node.Node.SetAsExplicitTrigger();
                     }
                 }
             }
@@ -237,7 +231,7 @@ public partial class GUBP : BuildCommand
 
         List<BuildNode> OrderedToDo = TopologicalSort(NodesToDo, ExplicitTrigger, false, false);
 
-		List<BuildNode> UnfinishedTriggers = FindUnfinishedTriggers(bSkipTriggers, ExplicitTrigger, OrderedToDo);
+		List<TriggerNode> UnfinishedTriggers = FindUnfinishedTriggers(bSkipTriggers, ExplicitTrigger, OrderedToDo);
 
 		PrintNodes(this, OrderedToDo, AllAggregates, UnfinishedTriggers, TimeQuantum);
 
@@ -282,7 +276,7 @@ public partial class GUBP : BuildCommand
     {
 		if(NodeToDo.ControllingTriggers == null)
 		{
-			NodeToDo.ControllingTriggers = new BuildNode[0];
+			NodeToDo.ControllingTriggers = new TriggerNode[0];
 
 			// Find all the dependencies of this node
 			List<BuildNode> AllDependencies = new List<BuildNode>();
@@ -290,17 +284,13 @@ public partial class GUBP : BuildCommand
 			AllDependencies.AddRange(NodeToDo.PseudoDependencies);
 
 			// Find the immediate trigger controlling this one
-			List<BuildNode> PreviousTriggers = new List<BuildNode>();
+			List<TriggerNode> PreviousTriggers = new List<TriggerNode>();
 			foreach (BuildNode Dependency in AllDependencies)
 			{
 				FindControllingTriggers(Dependency);
 
-				BuildNode PreviousTrigger = null;
-				if(Dependency.Node.TriggerNode())
-				{
-					PreviousTrigger = Dependency;
-				}
-				else if(Dependency.ControllingTriggers.Length > 0)
+				TriggerNode PreviousTrigger = Dependency as TriggerNode;
+				if(PreviousTrigger == null && Dependency.ControllingTriggers.Length > 0)
 				{
 					PreviousTrigger = Dependency.ControllingTriggers.Last();
 				}
@@ -323,7 +313,7 @@ public partial class GUBP : BuildCommand
 			// Update the list of controlling triggers
 			if (PreviousTriggers.Count == 1)
 			{
-				List<BuildNode> ControllingTriggers = new List<BuildNode>();
+				List<TriggerNode> ControllingTriggers = new List<TriggerNode>();
 				ControllingTriggers.AddRange(PreviousTriggers[0].ControllingTriggers);
 				ControllingTriggers.Add(PreviousTriggers[0]);
 				NodeToDo.ControllingTriggers = ControllingTriggers.ToArray();
@@ -493,7 +483,7 @@ public partial class GUBP : BuildCommand
 
     }
 
-    void PrintNodes(GUBP bp, List<BuildNode> Nodes, IEnumerable<AggregateNode> Aggregates, List<BuildNode> UnfinishedTriggers, int TimeQuantum)
+    void PrintNodes(GUBP bp, List<BuildNode> Nodes, IEnumerable<AggregateNode> Aggregates, List<TriggerNode> UnfinishedTriggers, int TimeQuantum)
     {
 		AggregateNode[] MatchingAggregates = Aggregates.Where(x => x.Dependencies.All(y => Nodes.Contains(y))).ToArray();
 		if (MatchingAggregates.Length > 0)
@@ -557,7 +547,7 @@ public partial class GUBP : BuildCommand
 			{
 				Builder.Append(" - (Completed)");
 			}
-			if(NodeToDo.Node.TriggerNode())
+			if(NodeToDo is TriggerNode)
 			{
 				Builder.Append(" - (TriggerNode)");
 			}
@@ -834,8 +824,8 @@ public partial class GUBP : BuildCommand
                             }
                             else if (BestPseudoReady == bPseudoReady)
                             {
-                                bool IamLateTrigger = !DoNotConsiderCompletion && NodeToDo.Node.TriggerNode() && NodeToDo != ExplicitTrigger && !NodeToDo.IsComplete;
-                                bool BestIsLateTrigger = !DoNotConsiderCompletion && BestNode.Node.TriggerNode() && BestNode != ExplicitTrigger && !BestNode.IsComplete;
+                                bool IamLateTrigger = !DoNotConsiderCompletion && (NodeToDo is TriggerNode) && NodeToDo != ExplicitTrigger && !NodeToDo.IsComplete;
+                                bool BestIsLateTrigger = !DoNotConsiderCompletion && (BestNode is TriggerNode) && BestNode != ExplicitTrigger && !BestNode.IsComplete;
                                 if (BestIsLateTrigger && !IamLateTrigger)
                                 {
                                     bReady = false;
@@ -939,7 +929,7 @@ public partial class GUBP : BuildCommand
     {
 		NodeHistory History = null;
 
-        if (!NodeToDo.Node.TriggerNode() && CLString != "")
+        if (!(NodeToDo is TriggerNode) && CLString != "")
         {
             string GameNameIfAny = NodeToDo.Node.GameNameIfAnyForTempStorage();
             string NodeStoreWildCard = StoreName.Replace(CLString, "*") + "-" + NodeToDo.Name;
@@ -1393,7 +1383,7 @@ public partial class GUBP : BuildCommand
 			HashSet<BuildNode> NodesToCull = new HashSet<BuildNode>();
 			foreach(BuildNode NodeToDo in NodesToDo)
 			{
-				if(NodeToDo.Node.TriggerNode() || NodeToDo.ControllingTriggers.Length > 0)
+				if((NodeToDo is TriggerNode) || NodeToDo.ControllingTriggers.Length > 0)
 				{
 					LogVerbose(" Culling {0} due to being downstream of trigger in preflight", NodeToDo.Name);
 					NodesToCull.Add(NodeToDo);
@@ -1516,20 +1506,17 @@ public partial class GUBP : BuildCommand
 		throw new AutomationException("Failed to update the CIS counter after 20 tries.");
 	}
 
-	private List<BuildNode> FindUnfinishedTriggers(bool bSkipTriggers, BuildNode ExplicitTrigger, List<BuildNode> OrdereredToDo)
+	private List<TriggerNode> FindUnfinishedTriggers(bool bSkipTriggers, BuildNode ExplicitTrigger, List<BuildNode> OrdereredToDo)
 	{
 		// find all unfinished triggers, excepting the one we are triggering right now
-		List<BuildNode> UnfinishedTriggers = new List<BuildNode>();
+		List<TriggerNode> UnfinishedTriggers = new List<TriggerNode>();
 		if (!bSkipTriggers)
 		{
-			foreach (BuildNode NodeToDo in OrdereredToDo)
+			foreach (TriggerNode NodeToDo in OrdereredToDo.OfType<TriggerNode>())
 			{
-				if (NodeToDo.Node.TriggerNode() && !NodeToDo.IsComplete)
+				if (!NodeToDo.IsComplete && ExplicitTrigger != NodeToDo)
 				{
-					if (ExplicitTrigger != NodeToDo)
-					{
-						UnfinishedTriggers.Add(NodeToDo);
-					}
+					UnfinishedTriggers.Add(NodeToDo);
 				}
 			}
 		}
@@ -1544,7 +1531,7 @@ public partial class GUBP : BuildCommand
 	{
 		foreach (BuildNode NodeToDo in OrdereredToDo)
 		{
-			if (NodeToDo.Node.TriggerNode() && (NodeToDo.Node.IsSticky() || NodeToDo.IsComplete)) // these sticky triggers are ok, everything is already completed anyway
+			if ((NodeToDo is TriggerNode) && (NodeToDo.Node.IsSticky() || NodeToDo.IsComplete)) // these sticky triggers are ok, everything is already completed anyway
 			{
 				continue;
 			}
@@ -1568,7 +1555,7 @@ public partial class GUBP : BuildCommand
 		}
 	}
 
-	private void DoCommanderSetup(ElectricCommander EC, IEnumerable<BuildNode> AllNodes, IEnumerable<AggregateNode> AllAggregates, HashSet<BuildNode> NodesToDo, List<BuildNode> OrdereredToDo, int TimeIndex, int TimeQuantum, bool bSkipTriggers, bool bFake, bool bFakeEC, string CLString, BuildNode ExplicitTrigger, List<BuildNode> UnfinishedTriggers, string FakeFail, bool bPreflightBuild)
+	private void DoCommanderSetup(ElectricCommander EC, IEnumerable<BuildNode> AllNodes, IEnumerable<AggregateNode> AllAggregates, HashSet<BuildNode> NodesToDo, List<BuildNode> OrdereredToDo, int TimeIndex, int TimeQuantum, bool bSkipTriggers, bool bFake, bool bFakeEC, string CLString, TriggerNode ExplicitTrigger, List<TriggerNode> UnfinishedTriggers, string FakeFail, bool bPreflightBuild)
 	{
 		List<BuildNode> SortedNodes = TopologicalSort(new HashSet<BuildNode>(AllNodes), null, SubSort: false, DoNotConsiderCompletion: true);
 		Log("******* {0} GUBP Nodes", SortedNodes.Count);
@@ -1592,7 +1579,7 @@ public partial class GUBP : BuildCommand
 							continue; // this wasn't on the chain related to the trigger we are triggering, so it is not relevant
 						}
 					}
-					if (bPreflightBuild && !bSkipTriggers && NodeToDo.Node.TriggerNode())
+					if (bPreflightBuild && !bSkipTriggers && (NodeToDo is TriggerNode))
 					{
 						// in preflight builds, we are either skipping triggers (and running things downstream) or we just stop at triggers and don't make them available for triggering.
 						continue;
@@ -1669,7 +1656,7 @@ public partial class GUBP : BuildCommand
 
             // this is kinda complicated
             bool SaveSuccessRecords = (IsBuildMachine || bFakeEC) && // no real reason to make these locally except for fakeEC tests
-                (!NodeToDo.Node.TriggerNode() || NodeToDo.Node.IsSticky()); // trigger nodes are run twice, one to start the new workflow and once when it is actually triggered, we will save reconds for the latter
+                (!(NodeToDo is TriggerNode) || NodeToDo.Node.IsSticky()); // trigger nodes are run twice, one to start the new workflow and once when it is actually triggered, we will save reconds for the latter
 
             Log("***** Running GUBP Node {0} -> {1} : {2}", NodeToDo.Name, GameNameIfAny, NodeStoreName);
             if (NodeToDo.IsComplete)
diff --git a/Engine/Source/Programs/AutomationTool/GUBP/LegacyBranchSetup.cs b/Engine/Source/Programs/AutomationTool/GUBP/LegacyBranchSetup.cs
--- a/Engine/Source/Programs/AutomationTool/GUBP/LegacyBranchSetup.cs
+++ b/Engine/Source/Programs/AutomationTool/GUBP/LegacyBranchSetup.cs
@@ -1203,7 +1203,7 @@ partial class GUBP
 
 		// Convert the GUBPNodes and GUBPAggregates maps into lists
 		AllNodes = new List<BuildNode>();
-		AllNodes.AddRange(BranchConfig.GUBPNodes.Values.Select(x => new LegacyNode(BranchConfig, x)));
+		AllNodes.AddRange(BranchConfig.GUBPNodes.Values.Select(x => x.GetBuildNode()));
 
 		AllAggregates = new List<AggregateNode>();
 		AllAggregates.AddRange(BranchConfig.GUBPAggregates.Values.Select(x => new AggregateNode{ Name = x.GetFullName(), Node = x }));
diff --git a/Engine/Source/Programs/AutomationTool/GUBP/LegacyNodes.cs b/Engine/Source/Programs/AutomationTool/GUBP/LegacyNodes.cs
--- a/Engine/Source/Programs/AutomationTool/GUBP/LegacyNodes.cs
+++ b/Engine/Source/Programs/AutomationTool/GUBP/LegacyNodes.cs
@@ -37,6 +37,10 @@ partial class GUBP
         {
             return "";
         }
+		public virtual BuildNode GetBuildNode()
+		{
+			return new LegacyNode(this);
+		}
         public virtual void DoBuild(GUBP bp)
         {
             BuildProducts = new List<string>();
@@ -66,7 +70,7 @@ partial class GUBP
 		{
 			return "";
 		}
-        public virtual int AgentMemoryRequirement(GUBP.GUBPBranchConfig BranchConfig)
+        public virtual int AgentMemoryRequirement()
         {
             return 0;
         }
@@ -89,10 +93,6 @@ partial class GUBP
         {
             return 100.0f;
         }
-        public virtual bool TriggerNode()
-        {
-            return false;
-        }
         public virtual void SetAsExplicitTrigger()
         {
 
@@ -424,11 +424,11 @@ partial class GUBP
         {
             return true;
         }
-		public override int AgentMemoryRequirement(GUBP.GUBPBranchConfig BranchConfig)
+		public override int AgentMemoryRequirement()
         {
 			if (bHasLauncherParam)
             {
-                return base.AgentMemoryRequirement(BranchConfig);
+                return base.AgentMemoryRequirement();
             }
             return 32;
         }
@@ -1247,6 +1247,7 @@ partial class GUBP
         UnrealTargetPlatform TargetPlatform;
 		bool WithXp;
 		bool Precompiled; // If true, just builds targets which generate static libraries for the -UsePrecompiled option to UBT. If false, just build those that don't.
+		bool EnhanceAgentRequirements;
 
         public GamePlatformMonolithicsNode(GUBPBranchConfig InBranchConfig, UnrealTargetPlatform InHostPlatform, List<UnrealTargetPlatform> InActivePlatforms, BranchInfo.BranchUProject InGameProj, UnrealTargetPlatform InTargetPlatform, bool InWithXp = false, bool InPrecompiled = false)
             : base(InBranchConfig, InHostPlatform)
@@ -1256,6 +1257,7 @@ partial class GUBP
 			ActivePlatforms = InActivePlatforms;
 			WithXp = InWithXp;
 			Precompiled = InPrecompiled;
+			EnhanceAgentRequirements = BranchConfig.BranchOptions.EnhanceAgentRequirements.Contains(StaticGetFullName(HostPlatform, GameProj, TargetPlatform, WithXp, Precompiled));
 
             if (TargetPlatform == UnrealTargetPlatform.PS4 || TargetPlatform == UnrealTargetPlatform.XboxOne)
             {
@@ -1352,13 +1354,13 @@ partial class GUBP
         {
             return true;
         }
-		public override int AgentMemoryRequirement(GUBP.GUBPBranchConfig BranchConfig)
+		public override int AgentMemoryRequirement()
         {
-			if (BranchConfig.BranchOptions.EnhanceAgentRequirements.Contains(StaticGetFullName(HostPlatform, GameProj, TargetPlatform, WithXp, Precompiled)))
+			if (EnhanceAgentRequirements)
             {
                 return 64;
             }
-            return base.AgentMemoryRequirement(BranchConfig);
+            return base.AgentMemoryRequirement();
         }
         public override int CISFrequencyQuantumShift(GUBP.GUBPBranchConfig BranchConfig)
         {            
@@ -1811,9 +1813,9 @@ partial class GUBP
             BuildProducts = new List<string>();
             SaveRecordOfSuccessAndAddToBuildProducts();
         }
-        public override bool TriggerNode()
+        public override BuildNode GetBuildNode()
         {
-            return true;
+			return new TriggerNode(this);
         }
         public override void SetAsExplicitTrigger()
         {
@@ -2270,7 +2272,7 @@ partial class GUBP
         {
             return 10.0f;
         }
-		public override int AgentMemoryRequirement(GUBP.GUBPBranchConfig BranchConfig)
+		public override int AgentMemoryRequirement()
         {
             return bIsMassive ? 32 : 0;
         }
@@ -2894,9 +2896,9 @@ partial class GUBP
 			}		
 			return Result;
 		}
-		public override int AgentMemoryRequirement(GUBP.GUBPBranchConfig BranchConfig)
+		public override int AgentMemoryRequirement()
 		{
-			int Result = base.AgentMemoryRequirement(BranchConfig);
+			int Result = base.AgentMemoryRequirement();
 			if (HostPlatform == UnrealTargetPlatform.Mac)
 			{
 				Result = 32;
@@ -2945,9 +2947,9 @@ partial class GUBP
             }
             return Result;
         }
-		public override int AgentMemoryRequirement(GUBP.GUBPBranchConfig BranchConfig)
+		public override int AgentMemoryRequirement()
         {
-            int Result = base.AgentMemoryRequirement(BranchConfig);
+            int Result = base.AgentMemoryRequirement();
             if(HostPlatform == UnrealTargetPlatform.Mac)
             {
                 Result = 32;
diff --git a/Engine/Source/Programs/AutomationTool/GUBP/Nodes/AggregateNode.cs b/Engine/Source/Programs/AutomationTool/GUBP/Nodes/AggregateNode.cs
--- a/Engine/Source/Programs/AutomationTool/GUBP/Nodes/AggregateNode.cs
+++ b/Engine/Source/Programs/AutomationTool/GUBP/Nodes/AggregateNode.cs
@@ -7,7 +7,7 @@ using System.Text;
 namespace AutomationTool
 {
 	[DebuggerDisplay("{Name}")]
-	class AggregateNode
+	public class AggregateNode
 	{
 		public string Name;
 		public GUBP.GUBPAggregateNode Node;
diff --git a/Engine/Source/Programs/AutomationTool/GUBP/Nodes/BuildNode.cs b/Engine/Source/Programs/AutomationTool/GUBP/Nodes/BuildNode.cs
--- a/Engine/Source/Programs/AutomationTool/GUBP/Nodes/BuildNode.cs
+++ b/Engine/Source/Programs/AutomationTool/GUBP/Nodes/BuildNode.cs
@@ -7,7 +7,7 @@ using System.Text;
 namespace AutomationTool
 {
 	[DebuggerDisplay("{Name}")]
-	abstract class BuildNode
+	public abstract class BuildNode
 	{
 		public string Name;
 		public string AgentSharingGroup;
@@ -20,7 +20,7 @@ namespace AutomationTool
 		public BuildNode[] PseudoDependencies;
 		public BuildNode[] AllDirectDependencies;
 		public BuildNode[] AllIndirectDependencies;
-		public BuildNode[] ControllingTriggers;
+		public TriggerNode[] ControllingTriggers;
 		public bool IsComplete;
 		public string[] RecipientsForFailureEmails;
 		public bool AddSubmittersToFailureEmails;
diff --git a/Engine/Source/Programs/AutomationTool/GUBP/Nodes/LegacyNode.cs b/Engine/Source/Programs/AutomationTool/GUBP/Nodes/LegacyNode.cs
--- a/Engine/Source/Programs/AutomationTool/GUBP/Nodes/LegacyNode.cs
+++ b/Engine/Source/Programs/AutomationTool/GUBP/Nodes/LegacyNode.cs
@@ -6,14 +6,14 @@ using System.Text;
 
 namespace AutomationTool
 {
-	class LegacyNode : BuildNode
+	public class LegacyNode : BuildNode
 	{
-		public LegacyNode(GUBP.GUBPBranchConfig InBranchConfig, GUBP.GUBPNode InNode)
+		public LegacyNode(GUBP.GUBPNode InNode)
 		{
 			Name = InNode.GetFullName();
 			Node = InNode;
 			AgentSharingGroup = Node.AgentSharingGroup;
-			AgentMemoryRequirement = Node.AgentMemoryRequirement(InBranchConfig);
+			AgentMemoryRequirement = Node.AgentMemoryRequirement();
 			TimeoutInMinutes = Node.TimeoutInMinutes();
 			SendSuccessEmail = Node.SendSuccessEmail();
 		}
diff --git a/Engine/Source/Programs/AutomationTool/GUBP/Nodes/TriggerNode.cs b/Engine/Source/Programs/AutomationTool/GUBP/Nodes/TriggerNode.cs
new file mode 100644
--- /dev/null
+++ b/Engine/Source/Programs/AutomationTool/GUBP/Nodes/TriggerNode.cs
@@ -0,0 +1,16 @@
+﻿using System;
+using System.Collections.Generic;
+using System.Diagnostics;
+using System.Linq;
+using System.Text;
+
+namespace AutomationTool
+{
+	public class TriggerNode : LegacyNode
+	{
+		public TriggerNode(GUBP.WaitForUserInput InNode) : base(InNode)
+		{
+			AddSubmittersToFailureEmails = false;
+		}
+	}
+}
