Add a few more comments.

diff --git a/Engine/Source/Programs/AutomationTool/GUBP/GUBP.cs b/Engine/Source/Programs/AutomationTool/GUBP/GUBP.cs
--- a/Engine/Source/Programs/AutomationTool/GUBP/GUBP.cs
+++ b/Engine/Source/Programs/AutomationTool/GUBP/GUBP.cs
@@ -593,76 +593,71 @@ public partial class GUBP : BuildCommand
         }
     }
 
-    static void SaveGraphVisualization(List<BuildNode> Nodes)
-    {
-        List<GraphNode> GraphNodes = new List<GraphNode>();
-
-        Dictionary<BuildNode, GraphNode> NodeToGraphNodeMap = new Dictionary<BuildNode, GraphNode>();
-
-        for (int NodeIndex = 0; NodeIndex < Nodes.Count; ++NodeIndex)
-        {
-            BuildNode Node = Nodes[NodeIndex];
+	/// <summary>
+	/// Exports the build graph as a GEXF file, for visualization in an external tool (eg. Gephi).
+	/// </summary>
+	/// <param name="Nodes">The nodes in the graph</param>
+	static void SaveGraphVisualization(IEnumerable<BuildNode> Nodes)
+	{
+		// Create a graph node for each GUBP node in the graph
+		List<GraphNode> GraphNodes = new List<GraphNode>();
+		Dictionary<BuildNode, GraphNode> NodeToGraphNodeMap = new Dictionary<BuildNode, GraphNode>();
+		foreach(BuildNode Node in Nodes)
+		{
+			GraphNode GraphNode = new GraphNode();
+			GraphNode.Id = NodeToGraphNodeMap.Count;
+			GraphNode.Label = Node.Name;
 
-            GraphNode GraphNode = new GraphNode()
-            {
-                Id = GraphNodes.Count,
-                Label = Node.Name
-            };
-            GraphNodes.Add(GraphNode);
-            NodeToGraphNodeMap.Add(Node, GraphNode);
-        }
+			NodeToGraphNodeMap.Add(Node, GraphNode);
+			GraphNodes.Add(GraphNode);
+		}
 
-        // Connect everything together
-        List<GraphEdge> GraphEdges = new List<GraphEdge>();
+		// Connect everything together
+		List<GraphEdge> GraphEdges = new List<GraphEdge>();
+		foreach(KeyValuePair<BuildNode, GraphNode> NodeToGraphNodePair in NodeToGraphNodeMap)
+		{
+			foreach (BuildNode Dependency in NodeToGraphNodePair.Key.Dependencies)
+			{
+				GraphNode PrerequisiteFileGraphNode;
+				if (NodeToGraphNodeMap.TryGetValue(Dependency, out PrerequisiteFileGraphNode))
+				{
+					// Connect a file our action is dependent on, to our action itself
+					GraphEdge NewGraphEdge = new GraphEdge();
 
-        for (int NodeIndex = 0; NodeIndex < Nodes.Count; ++NodeIndex)
-        {
-            BuildNode Node = Nodes[NodeIndex];
-            GraphNode NodeGraphNode = NodeToGraphNodeMap[Node];
+					NewGraphEdge.Id = GraphEdges.Count;
+					NewGraphEdge.Source = PrerequisiteFileGraphNode;
+					NewGraphEdge.Target = NodeToGraphNodePair.Value;
+					NewGraphEdge.Color = new GraphColor() { R = 0.0f, G = 0.0f, B = 0.0f, A = 0.75f };
 
-            foreach (BuildNode Dep in Node.Dependencies)
-            {
-                GraphNode PrerequisiteFileGraphNode;
-                if (NodeToGraphNodeMap.TryGetValue(Dep, out PrerequisiteFileGraphNode))
-                {
-                    // Connect a file our action is dependent on, to our action itself
-                    GraphEdge NewGraphEdge = new GraphEdge()
-                    {
-                        Id = GraphEdges.Count,
-                        Source = PrerequisiteFileGraphNode,
-                        Target = NodeGraphNode,
-                        Color = new GraphColor() { R = 0.0f, G = 0.0f, B = 0.0f, A = 0.75f }
-                    };
+					GraphEdges.Add(NewGraphEdge);
+				}
 
-                    GraphEdges.Add(NewGraphEdge);
-                }
+			}
+			foreach (BuildNode Dependency in NodeToGraphNodePair.Key.PseudoDependencies)
+			{
+				GraphNode PrerequisiteFileGraphNode;
+				if (NodeToGraphNodeMap.TryGetValue(Dependency, out PrerequisiteFileGraphNode))
+				{
+					// Connect a file our action is dependent on, to our action itself
+					GraphEdge NewGraphEdge = new GraphEdge();
 
-            }
-            foreach (BuildNode Dep in Node.PseudoDependencies)
-            {
-                GraphNode PrerequisiteFileGraphNode;
-                if (NodeToGraphNodeMap.TryGetValue(Dep, out PrerequisiteFileGraphNode))
-                {
-                    // Connect a file our action is dependent on, to our action itself
-                    GraphEdge NewGraphEdge = new GraphEdge()
-                    {
-                        Id = GraphEdges.Count,
-                        Source = PrerequisiteFileGraphNode,
-                        Target = NodeGraphNode,
-                        Color = new GraphColor() { R = 0.0f, G = 0.0f, B = 0.0f, A = 0.25f }
-                    };
+					NewGraphEdge.Id = GraphEdges.Count;
+					NewGraphEdge.Source = PrerequisiteFileGraphNode;
+					NewGraphEdge.Target = NodeToGraphNodePair.Value;
+					NewGraphEdge.Color = new GraphColor() { R = 0.0f, G = 0.0f, B = 0.0f, A = 0.25f };
 
-                    GraphEdges.Add(NewGraphEdge);
-                }
+					GraphEdges.Add(NewGraphEdge);
+				}
 
-            }
-        }
+			}
+		}
 
-        string Filename = CommandUtils.CombinePaths(CommandUtils.CmdEnv.LogFolder, "GubpGraph.gexf");
-        Log("Writing graph to {0}", Filename);
-        GraphVisualization.WriteGraphFile(Filename, "GUBP Nodes", GraphNodes, GraphEdges);
-        Log("Wrote graph to {0}", Filename);
-     }
+		// Export the graph definition
+		string Filename = CommandUtils.CombinePaths(CommandUtils.CmdEnv.LogFolder, "GubpGraph.gexf");
+		Log("Writing graph to {0}", Filename);
+		GraphVisualization.WriteGraphFile(Filename, "GUBP Nodes", NodeToGraphNodeMap.Values.ToList(), GraphEdges);
+		Log("Wrote graph to {0}", Filename);
+	 }
 
     static List<int> ConvertCLToIntList(List<string> Strings)
     {
@@ -1566,24 +1561,20 @@ public partial class GUBP : BuildCommand
 			// remove nodes that have unfinished triggers
 			foreach (BuildNode NodeToDo in OrdereredToDo)
 			{
-				BuildNode ControllingTrigger = (NodeToDo.ControllingTriggers.Length > 0)? NodeToDo.ControllingTriggers.Last() : null;
-				bool bNoUnfinishedTriggers = !UnfinishedTriggers.Contains(ControllingTrigger);
-
-				if (bNoUnfinishedTriggers)
+				if (NodeToDo.ControllingTriggers.Length == 0 || !UnfinishedTriggers.Contains(NodeToDo.ControllingTriggers.Last()))
 				{
 					// if we are triggering, then remove nodes that are not controlled by the trigger or are dependencies of this trigger
-					if (ExplicitTrigger != null)
+					if (ExplicitTrigger != null && ExplicitTrigger != NodeToDo && !ExplicitTrigger.DependsOn(NodeToDo) && !NodeToDo.DependsOn(ExplicitTrigger))
 					{
-						if (ExplicitTrigger != NodeToDo && !ExplicitTrigger.DependsOn(NodeToDo) && !NodeToDo.DependsOn(ExplicitTrigger))
-						{
-							continue; // this wasn't on the chain related to the trigger we are triggering, so it is not relevant
-						}
+						continue; // this wasn't on the chain related to the trigger we are triggering, so it is not relevant
 					}
+
+					// in preflight builds, we are either skipping triggers (and running things downstream) or we just stop at triggers and don't make them available for triggering.
 					if (bPreflightBuild && !bSkipTriggers && (NodeToDo is TriggerNode))
 					{
-						// in preflight builds, we are either skipping triggers (and running things downstream) or we just stop at triggers and don't make them available for triggering.
 						continue;
 					}
+
 					FilteredOrdereredToDo.Add(NodeToDo);
 				}
 			}
