emscripten's file_packager.py requires EM_CONFIG is set before it's run.

diff --git a/Engine/Source/Programs/AutomationTool/AutomationUtils/CommandUtils.cs b/Engine/Source/Programs/AutomationTool/AutomationUtils/CommandUtils.cs
--- a/Engine/Source/Programs/AutomationTool/AutomationUtils/CommandUtils.cs
+++ b/Engine/Source/Programs/AutomationTool/AutomationUtils/CommandUtils.cs
@@ -2190,6 +2190,26 @@ namespace AutomationTool
 		}
 	}
 
+	/// <summary>
+	/// Use with "using" syntax to temporarily set and environment variable  in a convenient, exception-safe way
+	/// </summary>
+	public class ScopedEnvVar : IDisposable
+	{
+		private string StoredEnvVar = null;
+
+		public ScopedEnvVar(string EnvVar, string Value)
+		{
+			StoredEnvVar = EnvVar;
+			CommandUtils.SetEnvVar(StoredEnvVar, Value);
+		}
+
+		public void Dispose()
+		{
+			CommandUtils.SetEnvVar(StoredEnvVar, "");
+			GC.SuppressFinalize(this);
+		}
+	}
+
     /// <summary>
     /// Helper class to associate a file and its contents
     /// </summary>
diff --git a/Engine/Source/Programs/AutomationTool/HTML5/HTML5Platform.Automation.cs b/Engine/Source/Programs/AutomationTool/HTML5/HTML5Platform.Automation.cs
--- a/Engine/Source/Programs/AutomationTool/HTML5/HTML5Platform.Automation.cs
+++ b/Engine/Source/Programs/AutomationTool/HTML5/HTML5Platform.Automation.cs
@@ -72,10 +72,13 @@ public class HTML5Platform : Platform
 			string PythonPath = HTML5SDKInfo.Python();
 			string PackagerPath = HTML5SDKInfo.EmscriptenPackager();
 
-			using (new PushedDirectory(Path.Combine(Params.BaseStageDirectory, "HTML5")))
+			using (new ScopedEnvVar("EM_CONFIG", HTML5SDKInfo.DOT_EMSCRIPTEN))
 			{
-				string CmdLine = string.Format("\"{0}\" \"{1}\" --preload . --js-output=\"{1}.js\"", PackagerPath, FinalDataLocation);
-				RunAndLog(CmdEnv, PythonPath, CmdLine);
+				using (new PushedDirectory(Path.Combine(Params.BaseStageDirectory, "HTML5")))
+				{
+					string CmdLine = string.Format("\"{0}\" \"{1}\" --preload . --js-output=\"{1}.js\"", PackagerPath, FinalDataLocation);
+					RunAndLog(CmdEnv, PythonPath, CmdLine);
+				}
 			}
 		}
 
diff --git a/Engine/Source/Programs/AutomationTool/HTML5/HTML5Platform.PakFiles.Automation.cs b/Engine/Source/Programs/AutomationTool/HTML5/HTML5Platform.PakFiles.Automation.cs
--- a/Engine/Source/Programs/AutomationTool/HTML5/HTML5Platform.PakFiles.Automation.cs
+++ b/Engine/Source/Programs/AutomationTool/HTML5/HTML5Platform.PakFiles.Automation.cs
@@ -68,11 +68,14 @@ public class HTMLPakAutomation
 		string PythonPath = HTML5SDKInfo.Python();
 		string PackagerPath = HTML5SDKInfo.EmscriptenPackager();
 
-		// we need to operate in the root
-		using (new PushedDirectory(Path.Combine(PackagePath)))
+		using (new ScopedEnvVar("EM_CONFIG", HTML5SDKInfo.DOT_EMSCRIPTEN))
 		{
-			string CmdLine = string.Format("\"{0}\" \"{1}\" --preload {2}\\ --js-output=\"{1}.js\"", PackagerPath, FinalDataLocation, SC.ShortProjectName);
-			CommandUtils.RunAndLog(CommandUtils.CmdEnv, PythonPath, CmdLine);
+			// we need to operate in the root
+			using (new PushedDirectory(Path.Combine(PackagePath)))
+			{
+				string CmdLine = string.Format("\"{0}\" \"{1}\" --preload {2}\\ --js-output=\"{1}.js\"", PackagerPath, FinalDataLocation, SC.ShortProjectName);
+				CommandUtils.RunAndLog(CommandUtils.CmdEnv, PythonPath, CmdLine);
+			}
 		}
 	}
 
