Clean up the change history logging, and add a new command to debug it (-ShowHistory=<NodeName>)

diff --git a/Engine/Source/Programs/AutomationTool/GUBP/GUBP.cs b/Engine/Source/Programs/AutomationTool/GUBP/GUBP.cs
--- a/Engine/Source/Programs/AutomationTool/GUBP/GUBP.cs
+++ b/Engine/Source/Programs/AutomationTool/GUBP/GUBP.cs
@@ -240,20 +240,40 @@ public partial class GUBP : BuildCommand
 
 		ElectricCommander EC = new ElectricCommander(this);
 
-        string FakeFail = ParseParamValue("FakeFail");
-        if(CommanderSetup)
-        {
-			DoCommanderSetup(EC, AllNodes, AllAggregates, NodesToDo, OrderedToDo, TimeIndex, TimeQuantum, bSkipTriggers, bFake, bFakeEC, CLString, ExplicitTrigger, UnfinishedTriggers, FakeFail, bPreflightBuild);
-        }
-		else if(ParseParam("SaveGraph"))
+		string ShowHistoryParam = ParseParamValue("ShowHistory", null);
+		if(ShowHistoryParam != null)
 		{
-			SaveGraphVisualization(OrderedToDo);
+			BuildNode Node = AllNodes.FirstOrDefault(x => x.Name.Equals(ShowHistoryParam, StringComparison.InvariantCultureIgnoreCase));
+			if(Node == null)
+			{
+				throw new AutomationException("Couldn't find node {0}", ShowHistoryParam);
+			}
+
+			NodeHistory History = FindNodeHistory(Node, CLString, StoreName);
+			if(History == null)
+			{
+				throw new AutomationException("Couldn't get history for {0}", ShowHistoryParam);
+			}
+
+			PrintDetailedChanges(History, P4Env.Changelist);
 		}
-		else if(!ParseParam("ListOnly"))
+		else
 		{
-			ExecuteNodes(EC, OrderedToDo, bFake, bFakeEC, bSaveSharedTempStorage, CLString, StoreName, FakeFail);
+			string FakeFail = ParseParamValue("FakeFail");
+			if(CommanderSetup)
+			{
+				DoCommanderSetup(EC, AllNodes, AllAggregates, NodesToDo, OrderedToDo, TimeIndex, TimeQuantum, bSkipTriggers, bFake, bFakeEC, CLString, ExplicitTrigger, UnfinishedTriggers, FakeFail, bPreflightBuild);
+			}
+			else if(ParseParam("SaveGraph"))
+			{
+				SaveGraphVisualization(OrderedToDo);
+			}
+			else if(!ParseParam("ListOnly"))
+			{
+				ExecuteNodes(EC, OrderedToDo, bFake, bFakeEC, bSaveSharedTempStorage, CLString, StoreName, FakeFail);
+			}
 		}
-        PrintRunTime();
+		PrintRunTime();
 	}
 
 	/// <summary>
@@ -385,102 +405,87 @@ public partial class GUBP : BuildCommand
 		}
     }
 
-    static List<P4Connection.ChangeRecord> GetChanges(int LastOutputForChanges, int TopCL, int LastGreen)
+	/// <summary>
+	/// Finds all the source code changes between the given range.
+	/// </summary>
+	/// <param name="MinimumCL">The minimum (inclusive) changelist to include</param>
+	/// <param name="MaximumCL">The maximum (inclusive) changelist to include</param>
+	/// <returns>Lists of changelist records in the given range</returns>
+    static List<P4Connection.ChangeRecord> GetSourceChangeRecords(int MinimumCL, int MaximumCL)
     {
-        List<P4Connection.ChangeRecord> Result = new List<P4Connection.ChangeRecord>();
-        if (TopCL > LastGreen)
+        if (MinimumCL < 1990000)
         {
-            if (LastOutputForChanges > 1990000)
-            {
-                string Cmd = String.Format("{0}@{1},{2} {3}@{4},{5}",
-                    CombinePaths(PathSeparator.Slash, P4Env.BuildRootP4, "...", "Source", "..."), LastOutputForChanges + 1, TopCL,
-                    CombinePaths(PathSeparator.Slash, P4Env.BuildRootP4, "...", "Build", "..."), LastOutputForChanges + 1, TopCL
-                    );
-                List<P4Connection.ChangeRecord> ChangeRecords;
-				if (P4.Changes(out ChangeRecords, Cmd, false, true, LongComment: true))
-                {
-                    foreach (P4Connection.ChangeRecord Record in ChangeRecords)
-                    {
-                        if (!Record.User.Equals("buildmachine", StringComparison.InvariantCultureIgnoreCase))
-                        {
-                            Result.Add(Record);
-                        }
-                    }
-                }
-                else
-                {
-                    throw new AutomationException("Could not get changes; cmdline: p4 changes {0}", Cmd);
-                }
-            }
-            else
-            {
-                throw new AutomationException("That CL looks pretty far off {0}", LastOutputForChanges);
-            }
-        }
-        return Result;
-    }
+            throw new AutomationException("That CL looks pretty far off {0}", MinimumCL);
+		}
 
-    static int PrintChanges(int LastOutputForChanges, int TopCL, int LastGreen)
-    {
-        List<P4Connection.ChangeRecord> ChangeRecords = GetChanges(LastOutputForChanges, TopCL, LastGreen);
-        foreach (P4Connection.ChangeRecord Record in ChangeRecords)
+		// Query the changes from Perforce
+		StringBuilder FileSpec = new StringBuilder();
+		FileSpec.AppendFormat("{0}@{1},{2} ", CombinePaths(PathSeparator.Slash, P4Env.BuildRootP4, "...", "Source", "..."), MinimumCL, MaximumCL);
+		FileSpec.AppendFormat("{0}@{1},{2} ", CombinePaths(PathSeparator.Slash, P4Env.BuildRootP4, "...", "Build", "..."), MinimumCL, MaximumCL);
+
+        List<P4Connection.ChangeRecord> ChangeRecords;
+		if (!P4.Changes(out ChangeRecords, FileSpec.ToString(), false, true, LongComment: true))
         {
-            string Summary = Record.Summary.Replace("\r", "\n");
-            if (Summary.IndexOf("\n") > 0)
-            {
-                Summary = Summary.Substring(0, Summary.IndexOf("\n"));
-            }
-            Log("         {0} {1} {2}", Record.CL, Record.UserEmail, Summary);
-        }
-        return TopCL;
+            throw new AutomationException("Could not get changes; cmdline: p4 changes {0}", FileSpec.ToString());
+		}
+
+		// Filter out all the changes by the buildmachine user; the CIS counter or promoted builds aren't of interest to us.
+		ChangeRecords.RemoveAll(x => x.User.Equals("buildmachine", StringComparison.InvariantCultureIgnoreCase));
+		return ChangeRecords;
     }
 
-    static void PrintDetailedChanges(NodeHistory History, bool bShowAllChanges = false)
+	/// <summary>
+	/// Print a list of source changes, along with the success state for other builds of this node.
+	/// </summary>
+	/// <param name="History">History for this node</param>
+    static void PrintDetailedChanges(NodeHistory History, int CurrentCL)
     {
         DateTime StartTime = DateTime.UtcNow;
 
-        string Me = String.Format("{0}   <<<< local sync", P4Env.Changelist);
-        int LastOutputForChanges = 0;
-        int LastGreen = History.LastSucceeded;
-        if (bShowAllChanges)
-        {
-            if (History.AllStarted.Count > 0)
-            {
-                LastGreen = History.AllStarted[0];
-            }
-        }
-        foreach (int cl in History.AllStarted)
-        {
-            if (cl < LastGreen)
-            {
-                continue;
-            }
-            if (P4Env.Changelist < cl && Me != "")
-            {
-                LastOutputForChanges = PrintChanges(LastOutputForChanges, P4Env.Changelist, LastGreen);
-                Log("         {0}", Me);
-                Me = "";
-            }
-            string Status = "In Process";
-            if (History.AllSucceeded.Contains(cl))
+		// Find all the changelists that we're interested in
+		SortedSet<int> BuildChanges = new SortedSet<int>();
+		BuildChanges.UnionWith(History.AllStarted.Where(x => x >= History.LastSucceeded));
+		BuildChanges.Add(CurrentCL);
+
+		// Find all the changelists that we're interested in
+        List<P4Connection.ChangeRecord> ChangeRecords = GetSourceChangeRecords(BuildChanges.First(), BuildChanges.Last());
+
+		// Print all the changes in the set
+		int LastCL = BuildChanges.First();
+		foreach(int BuildCL in BuildChanges)
+		{
+			// Show all the changes in this range
+			foreach(P4Connection.ChangeRecord Record in ChangeRecords.Where(x => x.CL > LastCL && x.CL <= BuildCL))
+			{
+				string[] Lines = Record.Summary.Split(new char[]{ '\r', '\n' }, StringSplitOptions.RemoveEmptyEntries);
+				Log("             {0} {1} {2}", Record.CL, Record.UserEmail, (Lines.Length > 0)? Lines[0] : "");
+			}
+
+			// Show the status of this build
+			string BuildStatus;
+			if(BuildCL == CurrentCL)
+			{
+				BuildStatus = "this sync";
+			}
+            else if (History.AllSucceeded.Contains(BuildCL))
             {
-                Status = "ok";
+				BuildStatus = "built";
             }
-            if (History.AllFailed.Contains(cl))
+            else if (History.AllFailed.Contains(BuildCL))
             {
-                Status = "FAIL";
+				BuildStatus = "FAILED";
             }
-            LastOutputForChanges = PrintChanges(LastOutputForChanges, cl, LastGreen);
-            Log("         {0}   {1}", cl, Status);
-        }
-        if (Me != "")
-        {
-            LastOutputForChanges = PrintChanges(LastOutputForChanges, P4Env.Changelist, LastGreen);
-            Log("         {0}", Me);
-        }
-        double BuildDuration = (DateTime.UtcNow - StartTime).TotalMilliseconds;
-        Log("Took {0}s to get P4 history", BuildDuration / 1000);
+			else
+			{
+				BuildStatus = "running";
+			}
+			Log(" {0} {1} {2}", (BuildCL == CurrentCL)? ">>>>" : "    ", BuildCL, BuildStatus.ToString());
 
+			// Update the last CL now that we've output this one
+			LastCL = BuildCL;
+		}
+
+        Log("Took {0:0.0}s to get P4 history", (DateTime.UtcNow - StartTime).TotalSeconds);
     }
 
     void PrintNodes(GUBP bp, List<BuildNode> Nodes, IEnumerable<AggregateNode> Aggregates, List<TriggerNode> UnfinishedTriggers, int TimeQuantum)
@@ -934,13 +939,16 @@ public partial class GUBP : BuildCommand
             History.AllSucceeded = ConvertCLToIntList(TempStorage.FindTempStorageManifests(CmdEnv, NodeStoreWildCard + SucceededTempStorageSuffix, false, true, GameNameIfAny));
             History.AllFailed = ConvertCLToIntList(TempStorage.FindTempStorageManifests(CmdEnv, NodeStoreWildCard + FailedTempStorageSuffix, false, true, GameNameIfAny));
 
+			int CL;
+			int.TryParse(CLString, out CL);
+
             if (History.AllFailed.Count > 0)
             {
-                History.LastFailed = History.AllFailed[History.AllFailed.Count - 1];
+                History.LastFailed = History.AllFailed.LastOrDefault(x => x < CL);
             }
             if (History.AllSucceeded.Count > 0)
             {
-                History.LastSucceeded = History.AllSucceeded[History.AllSucceeded.Count - 1];
+                History.LastSucceeded = History.AllSucceeded.LastOrDefault(x => x < CL);
 
                 foreach (int Failed in History.AllFailed)
                 {
@@ -992,11 +1000,15 @@ public partial class GUBP : BuildCommand
                     Log(System.Diagnostics.TraceEventType.Warning, LogUtils.FormatException(Ex));
                 }
 
-                List<P4Connection.ChangeRecord> ChangeRecords = GetChanges(History.LastSucceeded, LastNonDuplicateFail, History.LastSucceeded);
-                foreach (P4Connection.ChangeRecord Record in ChangeRecords)
-                {
-                    FailCauserEMails = GUBPNode.MergeSpaceStrings(FailCauserEMails, Record.UserEmail);
-                }
+				if(LastNonDuplicateFail > History.LastSucceeded)
+				{
+					List<P4Connection.ChangeRecord> ChangeRecords = GetSourceChangeRecords(History.LastSucceeded + 1, LastNonDuplicateFail);
+					foreach (P4Connection.ChangeRecord Record in ChangeRecords)
+					{
+						FailCauserEMails = GUBPNode.MergeSpaceStrings(FailCauserEMails, Record.UserEmail);
+					}
+				}
+
                 if (!String.IsNullOrEmpty(FailCauserEMails))
                 {
                     NumPeople++;
@@ -1787,7 +1799,7 @@ public partial class GUBP : BuildCommand
                         Log("");
                         Log("");
                         Log("");
-                        PrintDetailedChanges(History);
+                        PrintDetailedChanges(History, P4Env.Changelist);
                         Log("End changes since last green");
                     }
 
