Fix formatting.

diff --git a/Engine/Source/Programs/AutomationTool/GUBP/LegacyBranchSetup.cs b/Engine/Source/Programs/AutomationTool/GUBP/LegacyBranchSetup.cs
--- a/Engine/Source/Programs/AutomationTool/GUBP/LegacyBranchSetup.cs
+++ b/Engine/Source/Programs/AutomationTool/GUBP/LegacyBranchSetup.cs
@@ -762,54 +762,54 @@ partial class GUBP
                                     }
                                     if (BranchOptions.ProjectsToCook.Contains(NonCodeProject.GameName) || BranchOptions.ProjectsToCook.Count == 0)
                                     {
-                                    if (!BranchConfig.HasNode(CookNode.StaticGetFullName(HostPlatform, NonCodeProject, CookedPlatform)))
-                                    {
-                                        GameCookNodes.Add(BranchConfig.AddNode(new CookNode(BranchConfig, HostPlatform, NonCodeProject, Plat, CookedPlatform)));
-                                    }
-                                    if (!BranchConfig.HasNode(GamePlatformCookedAndCompiledNode.StaticGetFullName(HostPlatform, NonCodeProject, Plat)))
-                                    {
-										BranchConfig.AddNode(new GamePlatformCookedAndCompiledNode(BranchConfig, HostPlatform, NonCodeProject, Plat, false));
+										if (!BranchConfig.HasNode(CookNode.StaticGetFullName(HostPlatform, NonCodeProject, CookedPlatform)))
+										{
+											GameCookNodes.Add(BranchConfig.AddNode(new CookNode(BranchConfig, HostPlatform, NonCodeProject, Plat, CookedPlatform)));
+										}
+										if (!BranchConfig.HasNode(GamePlatformCookedAndCompiledNode.StaticGetFullName(HostPlatform, NonCodeProject, Plat)))
+										{
+											BranchConfig.AddNode(new GamePlatformCookedAndCompiledNode(BranchConfig, HostPlatform, NonCodeProject, Plat, false));
 
-                                        if (NonCodeFormalBuilds.ContainsKey(NonCodeProject.GameName))
-                                        {
-                                            var PlatList = NonCodeFormalBuilds[NonCodeProject.GameName];
-                                            foreach (var PlatPair in PlatList)
-                                            {
-                                                if (PlatPair.TargetPlatform == Plat)
-                                                {
-													var NodeName = BranchConfig.AddNode(new FormalBuildNode(BranchConfig, NonCodeProject, HostPlatform, new List<UnrealTargetPlatform>() { Plat }, new List<UnrealTargetConfiguration>() { PlatPair.TargetConfig }));                                                    
-													// we don't want this delayed
-													// this would normally wait for the testing phase, we just want to build it right away
-													BranchConfig.RemovePseudodependencyFromNode(
-														CookNode.StaticGetFullName(HostPlatform, NonCodeProject, CookedPlatform),
-														WaitForTestShared.StaticGetFullName());
-													string BuildAgentSharingGroup = "";
-													if (Options.bSeparateGamePromotion)
+											if (NonCodeFormalBuilds.ContainsKey(NonCodeProject.GameName))
+											{
+												var PlatList = NonCodeFormalBuilds[NonCodeProject.GameName];
+												foreach (var PlatPair in PlatList)
+												{
+													if (PlatPair.TargetPlatform == Plat)
 													{
-														BuildAgentSharingGroup = NonCodeProject.GameName + "_MakeFormalBuild_" + Plat.ToString() + HostPlatformNode.StaticGetHostPlatformSuffix(HostPlatform);
-														if (Plat == UnrealTargetPlatform.IOS || Plat == UnrealTargetPlatform.Android) // These trash build products, so we need to use different agents
+														var NodeName = BranchConfig.AddNode(new FormalBuildNode(BranchConfig, NonCodeProject, HostPlatform, new List<UnrealTargetPlatform>() { Plat }, new List<UnrealTargetConfiguration>() { PlatPair.TargetConfig }));                                                    
+														// we don't want this delayed
+														// this would normally wait for the testing phase, we just want to build it right away
+														BranchConfig.RemovePseudodependencyFromNode(
+															CookNode.StaticGetFullName(HostPlatform, NonCodeProject, CookedPlatform),
+															WaitForTestShared.StaticGetFullName());
+														string BuildAgentSharingGroup = "";
+														if (Options.bSeparateGamePromotion)
 														{
-															BuildAgentSharingGroup = "";
+															BuildAgentSharingGroup = NonCodeProject.GameName + "_MakeFormalBuild_" + Plat.ToString() + HostPlatformNode.StaticGetHostPlatformSuffix(HostPlatform);
+															if (Plat == UnrealTargetPlatform.IOS || Plat == UnrealTargetPlatform.Android) // These trash build products, so we need to use different agents
+															{
+																BuildAgentSharingGroup = "";
+															}
+															BranchConfig.FindNode(CookNode.StaticGetFullName(HostPlatform, NonCodeProject, CookedPlatform)).AgentSharingGroup = BuildAgentSharingGroup;
+															BranchConfig.FindNode(NodeName).AgentSharingGroup = BuildAgentSharingGroup;
 														}
-														BranchConfig.FindNode(CookNode.StaticGetFullName(HostPlatform, NonCodeProject, CookedPlatform)).AgentSharingGroup = BuildAgentSharingGroup;
-														BranchConfig.FindNode(NodeName).AgentSharingGroup = BuildAgentSharingGroup;
-													}
-													else
-													{
-														//GUBPNodes[NodeName].AgentSharingGroup = FormalAgentSharingGroup;
-															if (Plat == UnrealTargetPlatform.XboxOne)
+														else
 														{
-															BranchConfig.FindNode(NodeName).AgentSharingGroup = "";
+															//GUBPNodes[NodeName].AgentSharingGroup = FormalAgentSharingGroup;
+																if (Plat == UnrealTargetPlatform.XboxOne)
+															{
+																BranchConfig.FindNode(NodeName).AgentSharingGroup = "";
+															}
+														}
+														if (PlatPair.bTest)
+														{
+															BranchConfig.AddNode(new FormalBuildTestNode(BranchConfig, NonCodeProject, HostPlatform, Plat, PlatPair.TargetConfig));
 														}
-													}
-													if (PlatPair.bTest)
-													{
-														BranchConfig.AddNode(new FormalBuildTestNode(BranchConfig, NonCodeProject, HostPlatform, Plat, PlatPair.TargetConfig));
 													}
 												}
 											}
 										}
-                                    }
 
 										if (!bNoAutomatedTesting)
                                         {
