GUBP: add UnrealPak location on Linux for MakeFeaturePacks node.

diff --git a/Engine/Source/Programs/AutomationTool/GUBP/LegacyNodes.cs b/Engine/Source/Programs/AutomationTool/GUBP/LegacyNodes.cs
--- a/Engine/Source/Programs/AutomationTool/GUBP/LegacyNodes.cs
+++ b/Engine/Source/Programs/AutomationTool/GUBP/LegacyNodes.cs
@@ -1227,6 +1227,10 @@ partial class GUBP
 				{
 					UnrealPakExe = CombinePaths(CmdEnv.LocalRoot, "Engine/Binaries/Win64/UnrealPak.exe");
 				}
+				else if (HostPlatform == UnrealTargetPlatform.Linux)
+				{
+					UnrealPakExe = CombinePaths(CmdEnv.LocalRoot, "Engine/Binaries/Linux/UnrealPak");
+				}
 				else
 				{
 					throw new AutomationException("Unknown path to UnrealPak for host platform ({0})", HostPlatform);
