Add AutomationUtils.Automation.dll to UAT build products.

diff --git a/Engine/Source/Programs/AutomationTool/AutomationUtils/UE4Build.cs b/Engine/Source/Programs/AutomationTool/AutomationUtils/UE4Build.cs
--- a/Engine/Source/Programs/AutomationTool/AutomationUtils/UE4Build.cs
+++ b/Engine/Source/Programs/AutomationTool/AutomationUtils/UE4Build.cs
@@ -1591,6 +1591,7 @@ namespace AutomationTool
 						"AutomationTool.exe.config",
 						"UnrealBuildTool.exe",
 						"UnrealBuildTool.exe.config",
+						"AutomationUtils.Automation.dll",
 						"DotNETUtilities.dll",
 					});
 
