UEB-117 - Move secure signing to the platforms

diff --git a/Engine/Source/Programs/AutomationTool/AutomationUtils/CommandUtils.cs b/Engine/Source/Programs/AutomationTool/AutomationUtils/CommandUtils.cs
--- a/Engine/Source/Programs/AutomationTool/AutomationUtils/CommandUtils.cs
+++ b/Engine/Source/Programs/AutomationTool/AutomationUtils/CommandUtils.cs
@@ -2568,6 +2568,7 @@ namespace AutomationTool
 					FilesToSignBuilder.Clear();
 				}
 			}
+			FinalListSignStrings.Add(FilesToSignBuilder.ToString());
 			foreach(string FilesToSign in FinalListSignStrings)
 			{
 				SignListFilesIfEXEOrDLL(FilesToSign);
diff --git a/Engine/Source/Programs/AutomationTool/AutomationUtils/ProjectParams.cs b/Engine/Source/Programs/AutomationTool/AutomationUtils/ProjectParams.cs
--- a/Engine/Source/Programs/AutomationTool/AutomationUtils/ProjectParams.cs
+++ b/Engine/Source/Programs/AutomationTool/AutomationUtils/ProjectParams.cs
@@ -297,6 +297,7 @@ namespace AutomationTool
             this.Prebuilt = InParams.Prebuilt;
             this.RunTimeoutSeconds = InParams.RunTimeoutSeconds;
 			this.bIsCodeBasedProject = InParams.bIsCodeBasedProject;
+			this.bCodeSign = InParams.bCodeSign;
 		}
 
 		/// <summary>
@@ -399,7 +400,8 @@ namespace AutomationTool
 			string SpecifiedArchitecture = null,
             bool? IterativeDeploy = null,
 			bool? FastCook = null,
-			bool? IgnoreCookErrors = null
+			bool? IgnoreCookErrors = null,
+			bool? CodeSign = null
 			)
 		{
 			//
@@ -514,6 +516,7 @@ namespace AutomationTool
 			}
 			this.StageDirectoryParam = ParseParamValueIfNotSpecified(Command, StageDirectoryParam, "stagingdirectory", String.Empty, true);
             this.StageNonMonolithic = GetParamValueIfNotSpecified(Command, StageNonMonolithic, this.StageNonMonolithic, "StageNonMonolithic");
+			this.bCodeSign = GetParamValueIfNotSpecified(Command, CodeSign, CommandUtils.IsBuildMachine, "CodeSign");
 			this.Manifests = GetParamValueIfNotSpecified(Command, Manifests, this.Manifests, "manifests");
             this.CreateChunkInstall = GetParamValueIfNotSpecified(Command, CreateChunkInstall, this.CreateChunkInstall, "createchunkinstall");
 			this.ChunkInstallDirectory = ParseParamValueIfNotSpecified(Command, ChunkInstallDirectory, "chunkinstalldirectory", String.Empty, true);
@@ -1201,6 +1204,11 @@ namespace AutomationTool
 		/// </summary>
         public bool NoBootstrapExe { get; set; }
 
+		/// <summary>
+		/// By default we don't code sign unless it is required or requested
+		/// </summary>
+		public bool bCodeSign = false;
+
 		#endregion
 
 		#region Run
diff --git a/Engine/Source/Programs/AutomationTool/GUBP/LegacyNodes.cs b/Engine/Source/Programs/AutomationTool/GUBP/LegacyNodes.cs
--- a/Engine/Source/Programs/AutomationTool/GUBP/LegacyNodes.cs
+++ b/Engine/Source/Programs/AutomationTool/GUBP/LegacyNodes.cs
@@ -361,14 +361,6 @@ partial class GUBP
                     AddBuildProduct(Product);
                 }
                 RemoveOveralppingBuildProducts();
-                if (CommandUtils.IsBuildMachine)
-                {
-                    // Sign everything we built
-					using(TelemetryStopwatch SignStopwatch = new TelemetryStopwatch("SignBuildProducts"))
-					{
-						CodeSign.SignMultipleIfEXEOrDLL(bp, BuildProducts);
-					}
-                }
                 PostBuildProducts(bp);
             }
 			if (Agenda == null || (BuildProducts.Count == 0 && BranchConfig.bForceIncrementalCompile))
diff --git a/Engine/Source/Programs/AutomationTool/Mac/MacPlatform.Automation.cs b/Engine/Source/Programs/AutomationTool/Mac/MacPlatform.Automation.cs
--- a/Engine/Source/Programs/AutomationTool/Mac/MacPlatform.Automation.cs
+++ b/Engine/Source/Programs/AutomationTool/Mac/MacPlatform.Automation.cs
@@ -418,8 +418,21 @@ public class MacPlatform : Platform
 		}
 		return false;
 	}
+
 	public override bool ShouldStageCommandLine(ProjectParams Params, DeploymentContext SC)
 	{
 		return false; // !String.IsNullOrEmpty(Params.StageCommandline) || !String.IsNullOrEmpty(Params.RunCommandline) || (!Params.IsCodeBasedProject && Params.NoBootstrapExe);
 	}
+
+	public override bool SignExecutables(DeploymentContext SC, ProjectParams Params)
+	{
+		// Sign everything we built
+		List<string> FilesToSign = GetExecutableNames(SC);
+		foreach (var File in FilesToSign)
+		{
+			CodeSign.SignMacFileOrFolder(File);
+		}
+
+		return true;
+	}
 }
diff --git a/Engine/Source/Programs/AutomationTool/Scripts/CopyBuildToStagingDirectory.Automation.cs b/Engine/Source/Programs/AutomationTool/Scripts/CopyBuildToStagingDirectory.Automation.cs
--- a/Engine/Source/Programs/AutomationTool/Scripts/CopyBuildToStagingDirectory.Automation.cs
+++ b/Engine/Source/Programs/AutomationTool/Scripts/CopyBuildToStagingDirectory.Automation.cs
@@ -1609,6 +1609,11 @@ public partial class Project : CommandUtils
 							WriteDeltaManifest(Params, SC, DeployedNonUFSFiles, StagedNonUFSFiles, DeploymentContext.NonUFSDeployDeltaFileName);
 						}
 					}
+
+					if (Params.bCodeSign)
+					{
+						SC.StageTargetPlatform.SignExecutables(SC, Params);
+					}
 				}
 			}
 
diff --git a/Engine/Source/Programs/AutomationTool/Scripts/RocketBuild.Automation.cs b/Engine/Source/Programs/AutomationTool/Scripts/RocketBuild.Automation.cs
--- a/Engine/Source/Programs/AutomationTool/Scripts/RocketBuild.Automation.cs
+++ b/Engine/Source/Programs/AutomationTool/Scripts/RocketBuild.Automation.cs
@@ -787,6 +787,9 @@ namespace Rocket
 			FilterRocketNode FilterNode = (FilterRocketNode)BranchConfig.FindNode(FilterRocketNode.StaticGetFullName(HostPlatform));
 			CopyManifestFilesToOutput(FilterNode.DepotManifestPath, CommandUtils.CmdEnv.LocalRoot, OutputDir);
 
+			// sign the executables
+			CodeSignManifestFiles(bp, FilterNode.DepotManifestPath, OutputDir);
+
 			// Copy the stripped files to the output directory
 			foreach(KeyValuePair<string, string> StrippedManifestPath in FilterNode.StrippedNodeManifestPaths)
 			{
@@ -807,6 +810,20 @@ namespace Rocket
 			SaveRecordOfSuccessAndAddToBuildProducts();
 		}
 
+		static void CodeSignManifestFiles(BuildCommand bp, string ManifestPath, string OutputDir)
+		{
+			// Read the files from the manifest
+			CommandUtils.Log("Reading manifest: '{0}'", ManifestPath);
+			string[] Files = CommandUtils.ReadAllLines(ManifestPath).Select(x => x.Trim()).Where(x => x.Length > 0).ToArray();
+
+			// Create lists of source and target files
+			CommandUtils.Log("Code sign files...");
+			string[] SourceFiles = Files.Select(x => CommandUtils.CombinePaths(OutputDir, x)).Where(x => (x.Contains(".dll") || x.Contains(".exe")) && !(Path.GetDirectoryName(x).Replace("\\", "/")).Contains("Binaries/XboxOne")).ToArray();
+
+			// Copy everything
+			CodeSign.SignMultipleIfEXEOrDLL(bp, SourceFiles.ToList());
+		}
+
 		static void CopyManifestFilesToOutput(string ManifestPath, string InputDir, string OutputDir)
 		{
 			// Read the files from the manifest
diff --git a/Engine/Source/Programs/AutomationTool/Win/WinPlatform.Automation.cs b/Engine/Source/Programs/AutomationTool/Win/WinPlatform.Automation.cs
--- a/Engine/Source/Programs/AutomationTool/Win/WinPlatform.Automation.cs
+++ b/Engine/Source/Programs/AutomationTool/Win/WinPlatform.Automation.cs
@@ -243,6 +243,15 @@ public abstract class BaseWinPlatform : Platform
 	{
 		return new List<string> { ".pdb", ".map" };
 	}
+
+	public override bool SignExecutables(DeploymentContext SC, ProjectParams Params)
+	{
+		// Sign everything we built
+		List<string> FilesToSign = GetExecutableNames(SC);
+		CodeSign.SignMultipleFilesIfEXEOrDLL(FilesToSign);
+
+		return true;
+	}
 }
 
 public class Win64Platform : BaseWinPlatform
