file missing from previous signing submission

diff --git a/Engine/Source/Programs/AutomationTool/AutomationUtils/Platform.cs b/Engine/Source/Programs/AutomationTool/AutomationUtils/Platform.cs
--- a/Engine/Source/Programs/AutomationTool/AutomationUtils/Platform.cs
+++ b/Engine/Source/Programs/AutomationTool/AutomationUtils/Platform.cs
@@ -508,6 +508,11 @@ namespace AutomationTool
 			return false;
 		}
 
+		public virtual bool SignExecutables(DeploymentContext SC, ProjectParams Params)
+		{
+			return true;
+		}
+
 		public virtual UnrealTargetPlatform[] GetStagePlatforms()
 		{
 			return new UnrealTargetPlatform[] { PlatformType };
