Don't submit the build.version file; it needs to be hand edited.

diff --git a/Engine/Source/Programs/AutomationTool/GUBP/LegacyNodes.cs b/Engine/Source/Programs/AutomationTool/GUBP/LegacyNodes.cs
--- a/Engine/Source/Programs/AutomationTool/GUBP/LegacyNodes.cs
+++ b/Engine/Source/Programs/AutomationTool/GUBP/LegacyNodes.cs
@@ -2072,6 +2072,7 @@ partial class GUBP
 							!Product.Replace('\\', '/').Contains("/Intermediate/") &&
 							!Product.Replace('\\', '/').Contains("/Engine/Saved/") &&
 							!Product.Replace('\\', '/').Contains("/DerivedDataCache/") &&
+							!Product.EndsWith("build.version", StringComparison.InvariantCultureIgnoreCase) && 
 							!Product.EndsWith(".lib") &&
 							!Product.EndsWith(".a") && 
 							!Product.EndsWith(".bc")
