Ensure build.version is writeable before updating it.

diff --git a/Engine/Source/Programs/AutomationTool/AutomationUtils/UE4Build.cs b/Engine/Source/Programs/AutomationTool/AutomationUtils/UE4Build.cs
--- a/Engine/Source/Programs/AutomationTool/AutomationUtils/UE4Build.cs
+++ b/Engine/Source/Programs/AutomationTool/AutomationUtils/UE4Build.cs
@@ -457,33 +457,6 @@ namespace AutomationTool
 
 			var Result = new List<String>();
 			string Branch = P4Enabled ? P4Env.BuildRootEscaped : "";
-			{
-				string VerFile = CombinePaths(CmdEnv.LocalRoot, "Engine", "Build", "build.properties");
-				if (bDoUpdateVersionFiles)
-				{
-					Log("Updating {0} with:", VerFile);
-					Log("  TimestampForBVT={0}", CmdEnv.TimestampAsString);
-					Log("  EngineVersion={0}", ChangelistNumber.ToString());
-					Log("  BranchName={0}", Branch);					
-
-					PrepareBuildProduct(VerFile, false);
-					VersionFileUpdater BuildProperties = new VersionFileUpdater(VerFile);
-					BuildProperties.ReplaceLine("TimestampForBVT=", CmdEnv.TimestampAsString);
-					BuildProperties.ReplaceLine("EngineVersion=", ChangelistNumber.ToString());
-					BuildProperties.ReplaceLine("BranchName=", Branch);
-                    BuildProperties.Commit();
-				}
-				else
-				{
-					Log("{0} will not be updated because P4 is not enabled.", VerFile);
-				}
-				if (IsBuildMachine)
-				{
-					// Only BuildMachines can check this file in
-					AddBuildProduct(VerFile);
-				}
-				Result.Add(VerFile);
-			}
 			{
 				string VerFile = CombinePaths(CmdEnv.LocalRoot, "Engine", "Build", "build.version");
 				if (bDoUpdateVersionFiles)
@@ -502,6 +475,7 @@ namespace AutomationTool
 					Pairs["BranchName"] = Branch;
 					VerText = Serializer.Serialize(Pairs).Replace("{\"", "{\n\t\"").Replace(",\"", ",\n\t\"").Replace("\":", "\": ").Replace("\"}", "\"\n}");
 
+					CommandUtils.SetFileAttributes(VerFile, ReadOnly: false);
 					CommandUtils.WriteAllText(VerFile, VerText);
 				}
 				else
