Move the "IsSticky" and "Priority" properties into the build node class, populating them from GUBPNode in LegacyNode.

diff --git a/Engine/Source/Programs/AutomationTool/GUBP/ElectricCommander.cs b/Engine/Source/Programs/AutomationTool/GUBP/ElectricCommander.cs
--- a/Engine/Source/Programs/AutomationTool/GUBP/ElectricCommander.cs
+++ b/Engine/Source/Programs/AutomationTool/GUBP/ElectricCommander.cs
@@ -233,7 +233,7 @@ namespace AutomationTool
 				if (!NodeToDo.IsComplete) // if something is already finished, we don't put it into EC
 				{
 					bHaveECNodes = true;
-					if (NodeToDo.Node.IsSticky())
+					if (NodeToDo.IsSticky)
 					{
 						LastSticky = NodeToDo;
 						if (HitNonSticky && !bSkipTriggers)
@@ -282,7 +282,7 @@ namespace AutomationTool
 							}
 						}
 					}
-					if (NodeToDo.Node.IsSticky())
+					if (NodeToDo.IsSticky)
 					{
 						if (!StickyChain.Contains(NodeToDo))
 						{
@@ -297,7 +297,7 @@ namespace AutomationTool
 						List<string> NodeProps = GetECPropsForNode(NodeToDo);
 						ECProps.AddRange(NodeProps);
 
-						bool Sticky = NodeToDo.Node.IsSticky();
+						bool Sticky = NodeToDo.IsSticky;
 						bool DoParallel = !Sticky;
 						if (NodeToDo.Node.ECProcedure() == "GUBP_UAT_Node_Parallel_AgentShare_Editor")
 						{
@@ -308,7 +308,7 @@ namespace AutomationTool
 							throw new AutomationException("Node {1} is sticky but has agent requirements.", NodeToDo.Name);
 						}
 						string Procedure = NodeToDo.Node.ECProcedure();
-						if (NodeToDo.Node.IsSticky() && NodeToDo == LastSticky)
+						if (NodeToDo.IsSticky && NodeToDo == LastSticky)
 						{
 							Procedure = Procedure + "_Release";
 						}
@@ -530,7 +530,7 @@ namespace AutomationTool
 					}
 				}
 			}
-			if (NodeToDo.Node.IsSticky())
+			if (NodeToDo.IsSticky)
 			{
 				List<BuildNode> MyChain = StickyChain;
 				int MyIndex = MyChain.IndexOf(NodeToDo);
diff --git a/Engine/Source/Programs/AutomationTool/GUBP/GUBP.cs b/Engine/Source/Programs/AutomationTool/GUBP/GUBP.cs
--- a/Engine/Source/Programs/AutomationTool/GUBP/GUBP.cs
+++ b/Engine/Source/Programs/AutomationTool/GUBP/GUBP.cs
@@ -556,7 +556,7 @@ public partial class GUBP : BuildCommand
 			{
 				Builder.Append(" - (TriggerNode)");
 			}
-			if(NodeToDo.Node.IsSticky())
+			if(NodeToDo.IsSticky)
 			{
 				Builder.Append(" - (Sticky)");
 			}
@@ -802,7 +802,7 @@ public partial class GUBP : BuildCommand
                         }
                     }
                 }
-                float Priority = NodeToDo.Node.Priority();
+                float Priority = NodeToDo.Priority;
 
                 if (bReady && BestNode != null)
                 {
@@ -812,11 +812,11 @@ public partial class GUBP : BuildCommand
                     }
                     else if (String.Compare(BestNode.ControllingTriggerDotName, NodeToDo.ControllingTriggerDotName) == 0) //sorted by controlling trigger
                     {
-                        if (BestNode.Node.IsSticky() && !NodeToDo.Node.IsSticky()) //sticky nodes first
+                        if (BestNode.IsSticky && !NodeToDo.IsSticky) //sticky nodes first
                         {
                             bReady = false;
                         }
-                        else if (BestNode.Node.IsSticky() == NodeToDo.Node.IsSticky())
+                        else if (BestNode.IsSticky == NodeToDo.IsSticky)
                         {
                             if (BestPseudoReady && !bPseudoReady)
                             {
@@ -1538,7 +1538,7 @@ public partial class GUBP : BuildCommand
 	{
 		foreach (BuildNode NodeToDo in OrdereredToDo)
 		{
-			if ((NodeToDo is TriggerNode) && (NodeToDo.Node.IsSticky() || NodeToDo.IsComplete)) // these sticky triggers are ok, everything is already completed anyway
+			if ((NodeToDo is TriggerNode) && (NodeToDo.IsSticky || NodeToDo.IsComplete)) // these sticky triggers are ok, everything is already completed anyway
 			{
 				continue;
 			}
@@ -1659,7 +1659,7 @@ public partial class GUBP : BuildCommand
 
             // this is kinda complicated
             bool SaveSuccessRecords = (IsBuildMachine || bFakeEC) && // no real reason to make these locally except for fakeEC tests
-                (!(NodeToDo is TriggerNode) || NodeToDo.Node.IsSticky()); // trigger nodes are run twice, one to start the new workflow and once when it is actually triggered, we will save reconds for the latter
+                (!(NodeToDo is TriggerNode) || NodeToDo.IsSticky); // trigger nodes are run twice, one to start the new workflow and once when it is actually triggered, we will save reconds for the latter
 
             Log("***** Running GUBP Node {0} -> {1} : {2}", NodeToDo.Name, GameNameIfAny, NodeStoreName);
             if (NodeToDo.IsComplete)
diff --git a/Engine/Source/Programs/AutomationTool/GUBP/Nodes/BuildNode.cs b/Engine/Source/Programs/AutomationTool/GUBP/Nodes/BuildNode.cs
--- a/Engine/Source/Programs/AutomationTool/GUBP/Nodes/BuildNode.cs
+++ b/Engine/Source/Programs/AutomationTool/GUBP/Nodes/BuildNode.cs
@@ -14,6 +14,8 @@ namespace AutomationTool
 		public int FrequencyShift;
 		public int AgentMemoryRequirement;
 		public int TimeoutInMinutes;
+		public bool IsSticky;
+		public float Priority;
 
 		public GUBP.GUBPNode Node;
 		public BuildNode[] Dependencies;
diff --git a/Engine/Source/Programs/AutomationTool/GUBP/Nodes/LegacyNode.cs b/Engine/Source/Programs/AutomationTool/GUBP/Nodes/LegacyNode.cs
--- a/Engine/Source/Programs/AutomationTool/GUBP/Nodes/LegacyNode.cs
+++ b/Engine/Source/Programs/AutomationTool/GUBP/Nodes/LegacyNode.cs
@@ -16,6 +16,8 @@ namespace AutomationTool
 			AgentMemoryRequirement = Node.AgentMemoryRequirement();
 			TimeoutInMinutes = Node.TimeoutInMinutes();
 			SendSuccessEmail = Node.SendSuccessEmail();
+			IsSticky = Node.IsSticky();
+			Priority = Node.Priority();
 		}
 	}
 }
