Fix case of build.version for Linux.

diff --git a/Engine/Source/Programs/AutomationTool/AutomationUtils/UE4Build.cs b/Engine/Source/Programs/AutomationTool/AutomationUtils/UE4Build.cs
--- a/Engine/Source/Programs/AutomationTool/AutomationUtils/UE4Build.cs
+++ b/Engine/Source/Programs/AutomationTool/AutomationUtils/UE4Build.cs
@@ -458,7 +458,7 @@ namespace AutomationTool
 			var Result = new List<String>();
 			string Branch = P4Enabled ? P4Env.BuildRootEscaped : "";
 			{
-				string VerFile = CombinePaths(CmdEnv.LocalRoot, "Engine", "Build", "build.version");
+				string VerFile = CombinePaths(CmdEnv.LocalRoot, "Engine", "Build", "Build.version");
 				if (bDoUpdateVersionFiles)
 				{
 					Log("Updating {0} with:", VerFile);
