Fix triggers only being sticky if they're being explicitly triggers. Fixes build failures running promotable triggers.

diff --git a/Engine/Source/Programs/AutomationTool/GUBP/Nodes/BuildNode.cs b/Engine/Source/Programs/AutomationTool/GUBP/Nodes/BuildNode.cs
--- a/Engine/Source/Programs/AutomationTool/GUBP/Nodes/BuildNode.cs
+++ b/Engine/Source/Programs/AutomationTool/GUBP/Nodes/BuildNode.cs
@@ -14,7 +14,6 @@ namespace AutomationTool
 		public int FrequencyShift;
 		public int AgentMemoryRequirement;
 		public int TimeoutInMinutes;
-		public bool IsSticky;
 		public float Priority;
 
 		public GUBP.GUBPNode Node;
@@ -28,6 +27,11 @@ namespace AutomationTool
 		public bool AddSubmittersToFailureEmails;
 		public bool SendSuccessEmail;
 
+		public abstract bool IsSticky
+		{
+			get;
+		}
+
 		public string ControllingTriggerDotName
 		{
 			get { return String.Join(".", ControllingTriggers.Select(x => x.Name)); }
diff --git a/Engine/Source/Programs/AutomationTool/GUBP/Nodes/LegacyNode.cs b/Engine/Source/Programs/AutomationTool/GUBP/Nodes/LegacyNode.cs
--- a/Engine/Source/Programs/AutomationTool/GUBP/Nodes/LegacyNode.cs
+++ b/Engine/Source/Programs/AutomationTool/GUBP/Nodes/LegacyNode.cs
@@ -16,8 +16,12 @@ namespace AutomationTool
 			AgentMemoryRequirement = Node.AgentMemoryRequirement();
 			TimeoutInMinutes = Node.TimeoutInMinutes();
 			SendSuccessEmail = Node.SendSuccessEmail();
-			IsSticky = Node.IsSticky();
 			Priority = Node.Priority();
 		}
+
+		public override bool IsSticky
+		{
+			get { return Node.IsSticky(); }
+		}
 	}
 }
