Add an error if nodes within an agent sharing group have multiple controlling triggers. Fix Rocket publish symbols nodes incorrectly having a dependency on the promotion trigger when debugging locally.

diff --git a/Engine/Source/Programs/AutomationTool/GUBP/GUBP.cs b/Engine/Source/Programs/AutomationTool/GUBP/GUBP.cs
--- a/Engine/Source/Programs/AutomationTool/GUBP/GUBP.cs
+++ b/Engine/Source/Programs/AutomationTool/GUBP/GUBP.cs
@@ -728,6 +728,15 @@ public partial class GUBP : BuildCommand
             {
                 SortedAgentGroupChains.Add(Chain.Key, TopologicalSort(new HashSet<BuildNode>(Chain.Value), ExplicitTrigger, true, DoNotConsiderCompletion));
             }
+			foreach(KeyValuePair<string, List<BuildNode>> Chain in SortedAgentGroupChains)
+			{
+				string[] ControllingTriggers = Chain.Value.Select(x => x.ControllingTriggerDotName).Distinct().OrderBy(x => x).ToArray();
+				if(ControllingTriggers.Length > 1)
+				{
+					string Triggers = String.Join(", ", ControllingTriggers.Select(x => String.Format("'{0}' ({1})", x, String.Join("+", Chain.Value.Where(y => y.ControllingTriggerDotName == x).Select(y => y.Name)))));
+					throw new AutomationException("Agent sharing group '{0}' has multiple controlling triggers: {1}", Chain.Key, Triggers);
+				}
+			}
             Log("***************Done with recursion");
         }
 
diff --git a/Engine/Source/Programs/AutomationTool/Scripts/RocketBuild.Automation.cs b/Engine/Source/Programs/AutomationTool/Scripts/RocketBuild.Automation.cs
--- a/Engine/Source/Programs/AutomationTool/Scripts/RocketBuild.Automation.cs
+++ b/Engine/Source/Programs/AutomationTool/Scripts/RocketBuild.Automation.cs
@@ -898,7 +898,6 @@ namespace Rocket
 			BranchConfig = InBranchConfig;
 			SymbolsOutputDir = InSymbolsOutputDir;
 
-			AddDependency(GUBP.WaitForSharedPromotionUserInput.StaticGetFullName(false));
 			AddDependency(GUBP.ToolsForCompileNode.StaticGetFullName(HostPlatform));
 			AddDependency(GUBP.RootEditorNode.StaticGetFullName(HostPlatform));
 			AddDependency(GUBP.ToolsNode.StaticGetFullName(HostPlatform));
