Remove DirectDependencies and DependedOnBy from GUBP job properties. As best I can tell, nothing actually uses them.

diff --git a/Engine/Source/Programs/AutomationTool/GUBP/ElectricCommander.cs b/Engine/Source/Programs/AutomationTool/GUBP/ElectricCommander.cs
--- a/Engine/Source/Programs/AutomationTool/GUBP/ElectricCommander.cs
+++ b/Engine/Source/Programs/AutomationTool/GUBP/ElectricCommander.cs
@@ -156,10 +156,8 @@ namespace AutomationTool
 			}
 		}
 
-		public void DoCommanderSetup(IEnumerable<BuildNode> AllNodes, IEnumerable<AggregateNode> AllAggregates, HashSet<BuildNode> NodesToDo, List<BuildNode> OrdereredToDo, List<BuildNode> SortedNodes, int TimeIndex, int TimeQuantum, bool bSkipTriggers, bool bFake, bool bFakeEC, string CLString, TriggerNode ExplicitTrigger, List<TriggerNode> UnfinishedTriggers, string FakeFail)
+		public void DoCommanderSetup(IEnumerable<BuildNode> AllNodes, IEnumerable<AggregateNode> AllAggregates, List<BuildNode> OrdereredToDo, List<BuildNode> SortedNodes, int TimeIndex, int TimeQuantum, bool bSkipTriggers, bool bFake, bool bFakeEC, string CLString, TriggerNode ExplicitTrigger, List<TriggerNode> UnfinishedTriggers, string FakeFail)
 		{
-			Dictionary<string, string> FullNodeDependedOnBy = GetFullNodeDependedOnBy(NodesToDo);
-
 			List<AggregateNode> SeparatePromotables = FindPromotables(AllAggregates);
 			Dictionary<BuildNode, List<AggregateNode>> DependentPromotions = FindDependentPromotables(AllNodes, SeparatePromotables);
 
@@ -175,18 +173,10 @@ namespace AutomationTool
 			{
 				ECProps.Add(string.Format("AllNodes/{0}={1}", Node.Name, GetNodeForAllNodesProperty(Node, TimeQuantum)));
 			}
-			foreach (BuildNode Node in SortedNodes)
-			{
-				ECProps.Add(string.Format("DirectDependencies/{0}={1}", Node.Name, String.Join(" ", Node.AllDirectDependencies.Select(x => x.Name))));
-			}
 			foreach (KeyValuePair<BuildNode, int> NodePair in FullNodeListSortKey)
 			{
 				ECProps.Add(string.Format("SortKey/{0}={1}", NodePair.Key.Name, NodePair.Value));
 			}
-			foreach (KeyValuePair<string, string> NodePair in FullNodeDependedOnBy)
-			{
-				ECProps.Add(string.Format("DependedOnBy/{0}={1}", NodePair.Key, NodePair.Value));
-			}
 			foreach (KeyValuePair<BuildNode, List<AggregateNode>> NodePair in DependentPromotions)
 			{
 				ECProps.Add(string.Format("DependentPromotions/{0}={1}", NodePair.Key.Name, String.Join(" ", NodePair.Value.Select(x => x.Name))));
@@ -650,35 +640,6 @@ namespace AutomationTool
 			return DependentPromotions;
 		}
 
-		private static Dictionary<string, string> GetFullNodeDependedOnBy(HashSet<BuildNode> NodesToDo)
-		{
-			//find things that depend on our nodes and setup commander dictionary
-			Dictionary<string, string> FullNodeDependedOnBy = new Dictionary<string, string>();
-			foreach (BuildNode NodeToDo in NodesToDo)
-			{
-				FullNodeDependedOnBy[NodeToDo.Name] = "";
-			}
-			foreach (BuildNode NodeToDo in NodesToDo)
-			{
-				if (!NodeToDo.Node.IsTest())
-				{
-					foreach (BuildNode Dep in NodeToDo.AllDirectDependencies)
-					{
-						string CurrentValue;
-						if (!FullNodeDependedOnBy.TryGetValue(Dep.Name, out CurrentValue) || CurrentValue.Length == 0)
-						{
-							FullNodeDependedOnBy[Dep.Name] = NodeToDo.Name;
-						}
-						else
-						{
-							FullNodeDependedOnBy[Dep.Name] += " " + NodeToDo.Name;
-						}
-					}
-				}
-			}
-			return FullNodeDependedOnBy;
-		}
-
 		/// <summary>
 		/// Sorts a list of nodes to display in EC. The default order is based on execution order and agent groups, whereas this function arranges nodes by
 		/// frequency then execution order, while trying to group nodes on parallel paths (eg. Mac/Windows editor nodes) together.
