Remove low-level overrides of ElectricCommander properties from GUBPNode: EC setup is all done inside ElectricCommander.cs now.

diff --git a/Engine/Source/Programs/AutomationTool/GUBP/ElectricCommander.cs b/Engine/Source/Programs/AutomationTool/GUBP/ElectricCommander.cs
--- a/Engine/Source/Programs/AutomationTool/GUBP/ElectricCommander.cs
+++ b/Engine/Source/Programs/AutomationTool/GUBP/ElectricCommander.cs
@@ -4,6 +4,7 @@ using System.IO;
 using System.Linq;
 using System.Text;
 using System.Xml;
+using UnrealBuildTool;
 
 namespace AutomationTool
 {
@@ -70,7 +71,7 @@ namespace AutomationTool
 			List<string> ECProps = new List<string>();
 			ECProps.Add("FailEmails/" + NodeToDo.Name + "=" + String.Join(" ", NodeToDo.RecipientsForFailureEmails));
 	
-			string AgentReq = NodeToDo.Node.ECAgentString();
+			string AgentReq = NodeToDo.AgentRequirements;
 			if(Command.ParseParamValue("AgentOverride") != "" && !NodeToDo.Name.Contains("OnMac"))
 			{
 				AgentReq = Command.ParseParamValue("AgentOverride");
@@ -238,7 +239,6 @@ namespace AutomationTool
 				}
 			}
 
-			List<string> FakeECArgs = new List<string>();
 			using(CommandUtils.TelemetryStopwatch PerlOutputStopwatch = new CommandUtils.TelemetryStopwatch("PerlOutput"))
 			{
 				string ParentPath = Command.ParseParamValue("ParentPath");
@@ -288,34 +288,90 @@ namespace AutomationTool
 						ECProps.AddRange(NodeProps);
 
 						bool Sticky = NodeToDo.IsSticky;
-						bool DoParallel = !Sticky;
-						if (NodeToDo.Node.ECProcedure() == "GUBP_UAT_Node_Parallel_AgentShare_Editor")
+						if(NodeToDo.IsSticky)
 						{
-							DoParallel = true;
+							if(NodeToDo.AgentSharingGroup != "")
+							{
+								throw new AutomationException("Node {0} is both agent sharing and sitcky.", NodeToDo.Name);
+							}
+							if(NodeToDo.AgentPlatform != UnrealTargetPlatform.Win64)
+							{
+								throw new AutomationException("Node {0} is sticky, but {1} hosted. Sticky nodes must be PC hosted.", NodeToDo.Name, NodeToDo.AgentPlatform);
+							}
+							if(NodeToDo.AgentRequirements != "")
+							{
+								throw new AutomationException("Node {0} is sticky but has agent requirements.", NodeToDo.Name);
+							}
 						}
-						if (Sticky && NodeToDo.Node.ECAgentString() != "")
+
+						string ProcedureInfix = "";
+						if(NodeToDo.AgentPlatform != UnrealTargetPlatform.Unknown && NodeToDo.AgentPlatform != UnrealTargetPlatform.Win64)
 						{
-							throw new AutomationException("Node {1} is sticky but has agent requirements.", NodeToDo.Name);
+							ProcedureInfix = "_" + NodeToDo.AgentPlatform.ToString();
 						}
-						string Procedure = NodeToDo.Node.ECProcedure();
-						if (NodeToDo.IsSticky && NodeToDo == LastSticky)
+
+						bool DoParallel = !Sticky || NodeToDo.IsParallelAgentShareEditor;
+						
+						TriggerNode TriggerNodeToDo = NodeToDo as TriggerNode;
+
+						List<Tuple<string, string>> Parameters = new List<Tuple<string,string>>();
+
+						Parameters.Add(new Tuple<string, string>("NodeName", NodeToDo.Name));
+						Parameters.Add(new Tuple<string, string>("Sticky", NodeToDo.IsSticky ? "1" : "0"));
+
+						if (NodeToDo.AgentSharingGroup != "")
 						{
-							Procedure = Procedure + "_Release";
+							Parameters.Add(new Tuple<string, string>("AgentSharingGroup", NodeToDo.AgentSharingGroup));
 						}
-						string Args = String.Format("{0}, subprocedure => '{1}', parallel => '{2}', jobStepName => '{3}', actualParameter => [{{actualParameterName => 'NodeName', value =>'{4}'}}",
-							BaseArgs, Procedure, DoParallel ? 1 : 0, NodeToDo.Name, NodeToDo.Name);
-						string ProcedureParams = NodeToDo.Node.ECProcedureParams();
-						if (!String.IsNullOrEmpty(ProcedureParams))
+
+						string Procedure;
+						if(TriggerNodeToDo == null || TriggerNodeToDo.IsTriggered)
 						{
-							Args = Args + ProcedureParams;
+							if (NodeToDo.IsParallelAgentShareEditor)
+							{
+								Procedure = "GUBP_UAT_Node_Parallel_AgentShare_Editor";
+							}
+							else
+							{
+								Procedure = "GUBP" + ProcedureInfix + "_UAT_Node";
+								if (!NodeToDo.IsSticky)
+								{
+									Procedure += "_Parallel";
+								}
+								if (NodeToDo.AgentSharingGroup != "")
+								{
+									Procedure += "_AgentShare";
+								}
+							}
+							if (NodeToDo.IsSticky && NodeToDo == LastSticky)
+							{
+								Procedure += "_Release";
+							}
 						}
-
-						if ((Procedure == "GUBP_UAT_Trigger" || Procedure == "GUBP_Hardcoded_Trigger") && NodeToDo.RecipientsForFailureEmails.Length > 0)
+						else
 						{
-							Args = Args + ", {actualParameterName => 'EmailsForTrigger', value => \'" + String.Join(" ", NodeToDo.RecipientsForFailureEmails) + "\'}";
+							if(TriggerNodeToDo.RequiresRecursiveWorkflow)
+							{
+				                Procedure = "GUBP_UAT_Trigger"; //here we run a recursive workflow to wait for the trigger
+							}
+							else
+							{
+								Procedure = "GUBP_Hardcoded_Trigger"; //here we advance the state in the hardcoded workflow so folks can approve
+							}
+
+							Parameters.Add(new Tuple<string, string>("TriggerState", TriggerNodeToDo.StateName));
+							Parameters.Add(new Tuple<string, string>("ActionText", TriggerNodeToDo.ActionText));
+							Parameters.Add(new Tuple<string, string>("DescText", TriggerNodeToDo.DescriptionText));
+
+							if (NodeToDo.RecipientsForFailureEmails.Length > 0)
+							{
+								Parameters.Add(new Tuple<string, string>("EmailsForTrigger", String.Join(" ", NodeToDo.RecipientsForFailureEmails)));
+							}
 						}
-						Args = Args + "]";
 
+						string ActualParameterArgs = String.Join(", ", Parameters.Select(x => String.Format("{{actualParameterName => '{0}', value => '{1}'}}", x.Item1, x.Item2)));
+						string Args = String.Format("{0}, subprocedure => '{1}', parallel => '{2}', jobStepName => '{3}', actualParameter => [{4}]", BaseArgs, Procedure, DoParallel? 1 : 0, NodeToDo.Name, ActualParameterArgs);
+						
 						List<BuildNode> UncompletedEcDeps = new List<BuildNode>();
 						{
 							foreach (BuildNode Dep in NodeToDo.AllDirectDependencies)
@@ -368,7 +424,7 @@ namespace AutomationTool
 								}
 								{
 									string NestArgs = String.Format("$batch->createJobStep({{parentPath => '{0}/jobSteps[{1}]', jobStepName => '{2}_GetPool', subprocedure => 'GUBP{3}_AgentShare_GetPool', parallel => '1', actualParameter => [{{actualParameterName => 'AgentSharingGroup', value => '{4}'}}, {{actualParameterName => 'NodeName', value => '{5}'}}]",
-										ParentPath, MyAgentGroup, MyAgentGroup, NodeToDo.Node.ECProcedureInfix(), MyAgentGroup, NodeToDo.Name);
+										ParentPath, MyAgentGroup, MyAgentGroup, ProcedureInfix, MyAgentGroup, NodeToDo.Name);
 									if (!String.IsNullOrEmpty(PreCondition))
 									{
 										NestArgs = NestArgs + ", precondition => " + PreCondition;
@@ -378,7 +434,7 @@ namespace AutomationTool
 								}
 								{
 									string NestArgs = String.Format("$batch->createJobStep({{parentPath => '{0}/jobSteps[{1}]', jobStepName => '{2}_GetAgent', subprocedure => 'GUBP{3}_AgentShare_GetAgent', parallel => '1', exclusiveMode => 'call', resourceName => '{4}', actualParameter => [{{actualParameterName => 'AgentSharingGroup', value => '{5}'}}, {{actualParameterName => 'NodeName', value=> '{6}'}}]",
-										ParentPath, MyAgentGroup, MyAgentGroup, NodeToDo.Node.ECProcedureInfix(),
+										ParentPath, MyAgentGroup, MyAgentGroup, ProcedureInfix,
 										String.Format("$[/myJob/jobSteps[{0}]/ResourcePool]", MyAgentGroup),
 										MyAgentGroup, NodeToDo.Name);
 									{
@@ -421,31 +477,6 @@ namespace AutomationTool
 						}
 						Args = Args + "});";
 						StepList.Add(Args);
-						if (bFakeEC &&
-							!UnfinishedTriggers.Contains(NodeToDo) &&
-							(NodeToDo.Node.ECProcedure().StartsWith("GUBP_UAT_Node") || NodeToDo.Node.ECProcedure().StartsWith("GUBP_Mac_UAT_Node")) // other things we really can't test
-							) // unfinished triggers are never run directly by EC, rather it does another job setup
-						{
-							string Arg = String.Format("gubp -Node={0} -FakeEC {1} {2} {3} {4} {5}",
-								NodeToDo.Name,
-								bFake ? "-Fake" : "",
-								Command.ParseParam("AllPlatforms") ? "-AllPlatforms" : "",
-								Command.ParseParam("UnfinishedTriggersFirst") ? "-UnfinishedTriggersFirst" : "",
-								Command.ParseParam("UnfinishedTriggersParallel") ? "-UnfinishedTriggersParallel" : "",
-								Command.ParseParam("WithMac") ? "-WithMac" : ""
-								);
-
-							string Node = Command.ParseParamValue("-Node");
-							if (!String.IsNullOrEmpty(Node))
-							{
-								Arg = Arg + " -Node=" + Node;
-							}
-							if (!String.IsNullOrEmpty(FakeFail))
-							{
-								Arg = Arg + " -FakeFail=" + FakeFail;
-							}
-							FakeECArgs.Add(Arg);
-						}
 
 						if (MyAgentGroup != "" && !bDoNestedJobstep)
 						{
diff --git a/Engine/Source/Programs/AutomationTool/GUBP/GUBP.cs b/Engine/Source/Programs/AutomationTool/GUBP/GUBP.cs
--- a/Engine/Source/Programs/AutomationTool/GUBP/GUBP.cs
+++ b/Engine/Source/Programs/AutomationTool/GUBP/GUBP.cs
@@ -203,12 +203,11 @@ public partial class GUBP : BuildCommand
         {
             if (!String.IsNullOrEmpty(ExplicitTriggerName))
             {
-                foreach (TriggerNode Node in AllNodes.OfType<TriggerNode>())
+                foreach (TriggerNode Trigger in AllNodes.OfType<TriggerNode>())
                 {
-                    if (Node.Name.Equals(ExplicitTriggerName, StringComparison.InvariantCultureIgnoreCase))
+					if (Trigger.Name.Equals(ExplicitTriggerName, StringComparison.InvariantCultureIgnoreCase))
                     {
-                        Node.Node.SetAsExplicitTrigger();
-						ExplicitTrigger = Node;
+						Trigger.Activate();
                         break;
                     }
                 }
@@ -221,9 +220,9 @@ public partial class GUBP : BuildCommand
             {
                 if (bSkipTriggers)
                 {
-                    foreach (TriggerNode Node in AllNodes.OfType<TriggerNode>())
+                    foreach (TriggerNode Trigger in AllNodes.OfType<TriggerNode>())
                     {
-                        Node.Node.SetAsExplicitTrigger();
+						Trigger.Activate();
                     }
                 }
             }
@@ -507,7 +506,6 @@ public partial class GUBP : BuildCommand
 
         bool bShowDependencies = bp.ParseParam("ShowDependencies");
         bool AddEmailProps = bp.ParseParam("ShowEmails");
-        bool ECProc = bp.ParseParam("ShowECProc");
         string LastControllingTrigger = "";
         string LastAgentGroup = "";
 
@@ -561,7 +559,7 @@ public partial class GUBP : BuildCommand
 				Builder.Append(" - (Sticky)");
 			}
 
-            string Agent = NodeToDo.Node.ECAgentString();
+            string Agent = NodeToDo.AgentRequirements;
 			if(ParseParamValue("AgentOverride") != "" && !NodeToDo.Node.GetFullName().Contains("Mac"))
 			{
 				Agent = ParseParamValue("AgentOverride");
@@ -578,10 +576,6 @@ public partial class GUBP : BuildCommand
             {
 				Builder.AppendFormat(" {0}", String.Join(" ", NodeToDo.RecipientsForFailureEmails));
             }
-			if(ECProc)
-			{
-				Builder.AppendFormat("  {0}", NodeToDo.Node.ECProcedure());
-			}
 			Log(Builder.ToString());
 
             if (bShowDependencies)
diff --git a/Engine/Source/Programs/AutomationTool/GUBP/LegacyNodes.cs b/Engine/Source/Programs/AutomationTool/GUBP/LegacyNodes.cs
--- a/Engine/Source/Programs/AutomationTool/GUBP/LegacyNodes.cs
+++ b/Engine/Source/Programs/AutomationTool/GUBP/LegacyNodes.cs
@@ -92,36 +92,11 @@ partial class GUBP
         public virtual float Priority()
         {
             return 100.0f;
-        }
-        public virtual void SetAsExplicitTrigger()
-        {
-
         }
         public virtual string ECAgentString()
         {
             return "";
         }
-        public virtual string ECProcedureInfix()
-        {
-            return "";
-        }
-        public virtual string ECProcedure()
-        {
-            if (IsSticky() && AgentSharingGroup != "")
-            {
-                throw new AutomationException("Node {0} is both agent sharing and sitcky.", GetFullName());
-            }
-            return String.Format("GUBP{0}_UAT_Node{1}{2}", ECProcedureInfix(), IsSticky() ? "" : "_Parallel", AgentSharingGroup != "" ? "_AgentShare" : "");
-        }
-        public virtual string ECProcedureParams()
-        {
-            var Result = String.Format(", {{actualParameterName => 'Sticky', value => '{0}'}}", IsSticky() ? 1 : 0);
-            if (AgentSharingGroup != "")
-            {
-                Result += String.Format(", {{actualParameterName => 'AgentSharingGroup', value => '{0}'}}", AgentSharingGroup);
-            }
-            return Result;
-        }
         public static string MergeSpaceStrings(params string[] EmailLists)
         {
             var Emails = new List<string>();
@@ -270,30 +245,16 @@ partial class GUBP
 			}
             return "";
         }
+		public override BuildNode GetBuildNode()
+		{
+			BuildNode Node = base.GetBuildNode();
+			Node.AgentPlatform = GetAgentPlatform();
+			return Node;
+		}
         public virtual UnrealTargetPlatform GetAgentPlatform()
         {
             return HostPlatform;
         }
-        public override string ECProcedureInfix()
-        {
-            if (GetAgentPlatform() == UnrealTargetPlatform.Mac)
-            {
-                if (IsSticky())
-                {
-                    throw new AutomationException("Node {0} is sticky, but Mac hosted. Sticky nodes must be PC hosted.", GetFullName());
-                }
-                return "_Mac";
-            }
-			if(GetAgentPlatform() == UnrealTargetPlatform.Linux)
-			{
-				if(IsSticky())
-				{
-					throw new AutomationException("Node {0} is sticky, but Linux hosted. Sticky nodes must be PC hosted.", GetFullName());
-				}
-				return "_Linux";
-			}
-            return "";
-        }
         public virtual string GetHostPlatformSuffix()
         {
             return StaticGetHostPlatformSuffix(HostPlatform, GetAgentPlatform());
@@ -404,14 +365,15 @@ partial class GUBP
             }
             return false;
         }
-        public override string ECProcedure()
-        {
-            if (HostPlatform == UnrealTargetPlatform.Win64)
-            {
-                return String.Format("GUBP_UAT_Node_Parallel_AgentShare_Editor");
-            }
-            return base.ECProcedure();
-        }
+		public override BuildNode GetBuildNode()
+		{
+			BuildNode Node = base.GetBuildNode();
+			if(HostPlatform == UnrealTargetPlatform.Win64)
+			{
+				Node.IsParallelAgentShareEditor = true;
+			}
+			return Node;
+		}
         public override bool DeleteBuildProducts()
         {
             return true;
@@ -505,15 +467,16 @@ partial class GUBP
             }
             return false;
         }
-        public override string ECProcedure()
-        {
-            if (HostPlatform == UnrealTargetPlatform.Win64)
-            {
-                return String.Format("GUBP_UAT_Node_Parallel_AgentShare_Editor");
-            }
-            return base.ECProcedure();
-        }
-        public override UE4Build.BuildAgenda GetAgenda(GUBP bp)
+		public override BuildNode GetBuildNode()
+		{
+			BuildNode Node = base.GetBuildNode();
+			if (HostPlatform == UnrealTargetPlatform.Win64)
+			{
+				Node.IsParallelAgentShareEditor = true;
+			}
+			return Node;
+		}
+		public override UE4Build.BuildAgenda GetAgenda(GUBP bp)
         {
             var Agenda = new UE4Build.BuildAgenda();
 
@@ -1093,20 +1056,21 @@ partial class GUBP
         public override bool IsSticky()
         {
             if (HostPlatform == UnrealTargetPlatform.Win64)
-        {
+	        {
                 return true;
             }
             return false;
         }
-        public override string ECProcedure()
-        {
-            if (HostPlatform == UnrealTargetPlatform.Win64)
-            {
-                return String.Format("GUBP_UAT_Node_Parallel_AgentShare_Editor");
-            }
-            return base.ECProcedure();
-        }
-        public override string GameNameIfAnyForTempStorage()
+		public override BuildNode GetBuildNode()
+		{
+			BuildNode Node = base.GetBuildNode();
+			if (HostPlatform == UnrealTargetPlatform.Win64)
+			{
+				Node.IsParallelAgentShareEditor = true;
+			}
+			return Node;
+		}
+		public override string GameNameIfAnyForTempStorage()
         {
             return GameProjects[0].Options(HostPlatform).GroupName ?? GameProjects[0].GameName;
             }                    
@@ -1798,10 +1762,8 @@ partial class GUBP
 
     public class WaitForUserInput : GUBPNode
     {
-        protected bool bTriggerWasTriggered;
         public WaitForUserInput()
         {
-            bTriggerWasTriggered = false;
         }
 
         public override void DoBuild(GUBP bp)
@@ -1813,14 +1775,6 @@ partial class GUBP
         {
 			return new TriggerNode(this);
         }
-        public override void SetAsExplicitTrigger()
-        {
-            bTriggerWasTriggered = true;
-        }
-        public override bool IsSticky()
-        {
-            return bTriggerWasTriggered;
-        }
         public virtual string GetTriggerStateName()
         {
             return "GenericTrigger";
@@ -1837,28 +1791,6 @@ partial class GUBP
         {
             return true;
         }
-        public override string ECProcedure()
-        {
-            if (bTriggerWasTriggered)
-            {
-                return base.ECProcedure(); // after this user hits the trigger, we want to run this as an ordinary node
-            }
-            if (TriggerRequiresRecursiveWorkflow())
-            {
-                return String.Format("GUBP_UAT_Trigger"); //here we run a recursive workflow to wait for the trigger
-            }
-            return String.Format("GUBP_Hardcoded_Trigger"); //here we advance the state in the hardcoded workflow so folks can approve
-        }
-        public override string ECProcedureParams()
-        {
-            var Result = base.ECProcedureParams();
-            if (!bTriggerWasTriggered)
-            {
-                Result += String.Format(", {{actualParameterName => 'TriggerState', value => '{0}'}}, {{actualParameterName => 'ActionText', value =>\"{1}\"}}, {{actualParameterName => 'DescText', value =>\"{2}\"}}", GetTriggerStateName(), GetTriggerActionText(), GetTriggerDescText());
-                //Result += String.Format(" --actualParameter TriggerState={0} --actualParameter ActionText=\"{1}\" --actualParameter DescText=\"{2}\"", GetTriggerStateName(), GetTriggerActionText(), GetTriggerDescText());
-            }
-            return Result;
-        }
         public override int TimeoutInMinutes()
         {
             return 0;
diff --git a/Engine/Source/Programs/AutomationTool/GUBP/Nodes/BuildNode.cs b/Engine/Source/Programs/AutomationTool/GUBP/Nodes/BuildNode.cs
--- a/Engine/Source/Programs/AutomationTool/GUBP/Nodes/BuildNode.cs
+++ b/Engine/Source/Programs/AutomationTool/GUBP/Nodes/BuildNode.cs
@@ -3,6 +3,7 @@ using System.Collections.Generic;
 using System.Diagnostics;
 using System.Linq;
 using System.Text;
+using UnrealBuildTool;
 
 namespace AutomationTool
 {
@@ -10,6 +11,8 @@ namespace AutomationTool
 	public abstract class BuildNode
 	{
 		public string Name;
+		public UnrealTargetPlatform AgentPlatform = UnrealTargetPlatform.Win64;
+		public string AgentRequirements;
 		public string AgentSharingGroup;
 		public int FrequencyShift;
 		public int AgentMemoryRequirement;
@@ -26,6 +29,7 @@ namespace AutomationTool
 		public string[] RecipientsForFailureEmails;
 		public bool AddSubmittersToFailureEmails;
 		public bool SendSuccessEmail;
+		public bool IsParallelAgentShareEditor;
 
 		public abstract bool IsSticky
 		{
diff --git a/Engine/Source/Programs/AutomationTool/GUBP/Nodes/LegacyNode.cs b/Engine/Source/Programs/AutomationTool/GUBP/Nodes/LegacyNode.cs
--- a/Engine/Source/Programs/AutomationTool/GUBP/Nodes/LegacyNode.cs
+++ b/Engine/Source/Programs/AutomationTool/GUBP/Nodes/LegacyNode.cs
@@ -12,6 +12,7 @@ namespace AutomationTool
 		{
 			Name = InNode.GetFullName();
 			Node = InNode;
+			AgentRequirements = Node.ECAgentString();
 			AgentSharingGroup = Node.AgentSharingGroup;
 			AgentMemoryRequirement = Node.AgentMemoryRequirement();
 			TimeoutInMinutes = Node.TimeoutInMinutes();
diff --git a/Engine/Source/Programs/AutomationTool/GUBP/Nodes/TriggerNode.cs b/Engine/Source/Programs/AutomationTool/GUBP/Nodes/TriggerNode.cs
--- a/Engine/Source/Programs/AutomationTool/GUBP/Nodes/TriggerNode.cs
+++ b/Engine/Source/Programs/AutomationTool/GUBP/Nodes/TriggerNode.cs
@@ -8,9 +8,30 @@ namespace AutomationTool
 {
 	public class TriggerNode : LegacyNode
 	{
+        public bool IsTriggered;
+		public bool RequiresRecursiveWorkflow;
+		public string StateName;
+		public string DescriptionText;
+		public string ActionText;
+
 		public TriggerNode(GUBP.WaitForUserInput InNode) : base(InNode)
 		{
 			AddSubmittersToFailureEmails = false;
+
+			StateName = InNode.GetTriggerStateName();
+			DescriptionText = InNode.GetTriggerDescText();
+			ActionText = InNode.GetTriggerActionText();
+			RequiresRecursiveWorkflow = InNode.TriggerRequiresRecursiveWorkflow();
+		}
+
+		public void Activate()
+		{ 
+			IsTriggered = true;
+		}
+
+		public override bool IsSticky
+		{
+			get { return IsTriggered; }
 		}
 	}
 }
