New APIDocTool program added to ToolsForDocumentationNode.

diff --git a/Engine/Source/Programs/AutomationTool/Scripts/BuildDocumentation.Automation.cs b/Engine/Source/Programs/AutomationTool/Scripts/BuildDocumentation.Automation.cs
--- a/Engine/Source/Programs/AutomationTool/Scripts/BuildDocumentation.Automation.cs
+++ b/Engine/Source/Programs/AutomationTool/Scripts/BuildDocumentation.Automation.cs
@@ -36,7 +36,8 @@ class ToolsForDocumentationNode : GUBP.CompileNode
         var Agenda = new UE4Build.BuildAgenda();
 		Agenda.DotNetSolutions.Add("Engine/Source/Programs/UnrealDocTool/APIDocTool/APIDocTool.sln");
 		Agenda.DotNetSolutions.Add("Engine/Source/Programs/UnrealDocTool/UnrealDocTool/UnrealDocTool.sln");
-        return Agenda;
+		Agenda.DotNetSolutions.Add("Engine/Source/Programs/UnrealDocTool/DoxygenInputFilter/DoxygenInputFilter.sln");
+		return Agenda;
     }
 }
 	
