Reduce max temp storage time to 1 day.

diff --git a/Engine/Source/Programs/AutomationTool/GUBP/TempStorage.cs b/Engine/Source/Programs/AutomationTool/GUBP/TempStorage.cs
--- a/Engine/Source/Programs/AutomationTool/GUBP/TempStorage.cs
+++ b/Engine/Source/Programs/AutomationTool/GUBP/TempStorage.cs
@@ -588,7 +588,7 @@ namespace AutomationTool
             }
             TopDirectoryTestedForClean.Add(TopDirectory);
 
-            const int MaximumDaysToKeepTempStorage = 2;
+            const int MaximumDaysToKeepTempStorage = 1;
             var StartTimeDir = DateTime.UtcNow;
             DirectoryInfo DirInfo = new DirectoryInfo(TopDirectory);
             var TopLevelDirs = DirInfo.GetDirectories();
