Enable HTML5 networking plugin by default.
Allow UAT to deploy websocket dedicated servers for HTML5.
UEPLAT-793 & UEPLAT-795

diff --git a/Engine/Config/BaseEngine.ini b/Engine/Config/BaseEngine.ini
--- a/Engine/Config/BaseEngine.ini
+++ b/Engine/Config/BaseEngine.ini
@@ -2345,3 +2345,21 @@ bSendLogFile=true
 
 [UnrealHeaderTool]
 GeneratedCodeVersion=V1
+
+[/Script/HTML5Networking.WebSocketNetDriver]
+AllowPeerConnections=False
+AllowPeerVoice=False
+ConnectionTimeout=60.0
+InitialConnectTimeout=120.0
+AckTimeout=10.0
+KeepAliveTime=20.2
+MaxClientRate=15000
+MaxInternetClientRate=10000
+RelevantTimeout=5.0
+SpawnPrioritySeconds=1.0
+ServerTravelPause=4.0
+NetServerMaxTickRate=30
+LanServerMaxTickRate=35
+WebSocketPort=8889
+NetConnectionClassName="/Script/HTML5Networking.WebSocketConnection"
+MaxPortCountToTry=512
diff --git a/Engine/Config/HTML5/HTML5Engine.ini b/Engine/Config/HTML5/HTML5Engine.ini
--- a/Engine/Config/HTML5/HTML5Engine.ini
+++ b/Engine/Config/HTML5/HTML5Engine.ini
@@ -9,3 +9,7 @@ benabled=false
 [SystemSettings]
 r.setres=800x600
 
+[/Script/Engine.GameEngine]
+!NetDriverDefinitions=ClearArray
++NetDriverDefinitions=(DefName="GameNetDriver", DriverClassName="/Script/HTML5Networking.WebSocketNetDriver", DriverClassNameFallback="/Script/HTML5Networking.WebSocketNetDriver")
+
diff --git a/Engine/Plugins/Experimental/HTML5Networking/HTML5Networking.uplugin b/Engine/Plugins/Experimental/HTML5Networking/HTML5Networking.uplugin
--- a/Engine/Plugins/Experimental/HTML5Networking/HTML5Networking.uplugin
+++ b/Engine/Plugins/Experimental/HTML5Networking/HTML5Networking.uplugin
@@ -9,7 +9,7 @@
 	"EngineVersion" : "4.8.0",
 	"Description" : "HTML5 Networking",
 	"Category" : "Misc",
-	"EnabledByDefault" : false,
+	"EnabledByDefault" : true,
 
 	"Modules" :
 	[
diff --git a/Engine/Plugins/Experimental/HTML5Networking/Source/HTML5Networking/Private/WebSocket.cpp b/Engine/Plugins/Experimental/HTML5Networking/Source/HTML5Networking/Private/WebSocket.cpp
--- a/Engine/Plugins/Experimental/HTML5Networking/Source/HTML5Networking/Private/WebSocket.cpp
+++ b/Engine/Plugins/Experimental/HTML5Networking/Source/HTML5Networking/Private/WebSocket.cpp
@@ -234,6 +234,7 @@ void FWebSocket::HandlePacket()
 
 void FWebSocket::Flush()
 {
+	auto PendingMesssages = OutgoingBuffer.Num();
 	while (OutgoingBuffer.Num() > 0 && !IsServerSide)
 	{
 #if !PLATFORM_HTML5
@@ -243,6 +244,11 @@ void FWebSocket::Flush()
 			libwebsocket_callback_on_writable(Context, Wsi);
 #endif
 		HandlePacket();
+		if (PendingMesssages >= OutgoingBuffer.Num()) 
+		{
+			UE_LOG(LogHTML5Networking, Warning, TEXT("Unable to flush all of OutgoingBuffer in FWebSocket."));
+			break;
+		}
 	};
 }
 
diff --git a/Engine/Source/Programs/AutomationTool/AutomationUtils/DeploymentContext.cs b/Engine/Source/Programs/AutomationTool/AutomationUtils/DeploymentContext.cs
--- a/Engine/Source/Programs/AutomationTool/AutomationUtils/DeploymentContext.cs
+++ b/Engine/Source/Programs/AutomationTool/AutomationUtils/DeploymentContext.cs
@@ -195,6 +195,11 @@ public class DeploymentContext //: ProjectParams
 	/// </summary>
 	public bool DedicatedServer;
 
+	/// <summary>
+	/// The dedicated server and client use
+	/// </summary>
+	public bool bUseWebsocketNetDriver;
+
 	/// <summary>
 	/// True if this build is staged
 	/// </summary>
@@ -238,7 +243,8 @@ public class DeploymentContext //: ProjectParams
 		bool InCookOnTheFly,
 		bool InArchive,
 		bool InProgram,
-		bool bHasDedicatedServerAndClient
+		bool bHasDedicatedServerAndClient,
+		bool bInUseWebsocketNetDriver = false
 		)
 	{
 		bStageCrashReporter = InStageCrashReporter;
@@ -254,6 +260,7 @@ public class DeploymentContext //: ProjectParams
 		ShortProjectName = ProjectUtils.GetShortProjectName(RawProjectPath);
 		Stage = InStage;
 		Archive = InArchive;
+		bUseWebsocketNetDriver = bInUseWebsocketNetDriver;
 
         if (CookSourcePlatform != null && InCooked)
         {
diff --git a/Engine/Source/Programs/AutomationTool/Scripts/CopyBuildToStagingDirectory.Automation.cs b/Engine/Source/Programs/AutomationTool/Scripts/CopyBuildToStagingDirectory.Automation.cs
--- a/Engine/Source/Programs/AutomationTool/Scripts/CopyBuildToStagingDirectory.Automation.cs
+++ b/Engine/Source/Programs/AutomationTool/Scripts/CopyBuildToStagingDirectory.Automation.cs
@@ -385,7 +385,28 @@ public partial class Project : CommandUtils
             // Game ufs (content)
 
             SC.StageFiles(StagedFileType.UFS, CombinePaths(SC.ProjectRoot), "*.uproject", false, null, CombinePaths(SC.RelativeProjectRootForStage), true, !Params.UsePak(SC.StageTargetPlatform));
-            SC.StageFiles(StagedFileType.UFS, CombinePaths(SC.ProjectRoot, "Config"), "*", true, null, CombinePaths(SC.RelativeProjectRootForStage, "Config"), true, !Params.UsePak(SC.StageTargetPlatform)); // TODO: Exclude localization data generation config files.
+			if (SC.StageTargetPlatform.PlatformType != UnrealTargetPlatform.HTML5 && SC.bUseWebsocketNetDriver)
+			{
+				var EngineIniPath = Path.Combine(SC.ProjectRoot, "Config", "DefaultEngine.ini");
+				var IntermediateEngineIniDir = Path.Combine(GetIntermediateCommandlineDir(SC), SC.RelativeProjectRootForStage, "Config");
+				Directory.CreateDirectory(IntermediateEngineIniDir);
+				var IntermediateEngineIniPath = Path.Combine(IntermediateEngineIniDir, "DefaultEngine.ini");
+				List<String> IniLines = new List<String>();
+				if (File.Exists(EngineIniPath))
+				{
+					IniLines = File.ReadAllLines(EngineIniPath).ToList();
+				}
+				IniLines.Add("[/Script/Engine.GameEngine]");
+				IniLines.Add("!NetDriverDefinitions=ClearArray");
+				IniLines.Add("+NetDriverDefinitions=(DefName=\"GameNetDriver\", DriverClassName=\"/Script/HTML5Networking.WebSocketNetDriver\", DriverClassNameFallback=\"/Script/HTML5Networking.WebSocketNetDriver\")");
+				File.WriteAllLines(IntermediateEngineIniPath, IniLines);
+				SC.StageFiles(StagedFileType.UFS, CombinePaths(SC.ProjectRoot, "Config"), "*", true, new string[] {"DefaultEngine.ini"}, CombinePaths(SC.RelativeProjectRootForStage, "Config"), true, !Params.UsePak(SC.StageTargetPlatform)); // TODO: Exclude localization data generation config files.
+				SC.StageFiles(StagedFileType.UFS, IntermediateEngineIniDir, "*.ini", true, null, CombinePaths(SC.RelativeProjectRootForStage, "Config"), true, !Params.UsePak(SC.StageTargetPlatform)); // TODO: Exclude localization data generation config files.
+			}
+			else
+			{
+				SC.StageFiles(StagedFileType.UFS, CombinePaths(SC.ProjectRoot, "Config"), "*", true, null, CombinePaths(SC.RelativeProjectRootForStage, "Config"), true, !Params.UsePak(SC.StageTargetPlatform)); // TODO: Exclude localization data generation config files.
+			}
             foreach (string Culture in CulturesToStage)
             {
                 StageLocalizationDataForCulture(SC, Culture, CombinePaths(SC.ProjectRoot, "Content/Localization/Game"), CombinePaths(SC.RelativeProjectRootForStage, "Content/Localization/Game"), !Params.UsePak(SC.StageTargetPlatform));
@@ -595,7 +616,7 @@ public partial class Project : CommandUtils
 			CopyManifestFilesToStageDir(SC.NonUFSStagingFilesDebug, SC.StageDirectory, "DebugFiles", SC.StageTargetPlatform.GetFilesForCRCCheck());
 		}
 
-        bool bStageUnrealFileSystemFiles = !Params.CookOnTheFly && !ShouldCreatePak(Params) && !Params.FileServer;
+		bool bStageUnrealFileSystemFiles = !Params.CookOnTheFly && !Params.UsePak(SC.StageTargetPlatform) && !Params.FileServer;
 		if (bStageUnrealFileSystemFiles)
 		{
 			CopyManifestFilesToStageDir(SC.UFSStagingFiles, SC.StageDirectory, "UFSFiles", SC.StageTargetPlatform.GetFilesForCRCCheck());
@@ -1420,6 +1441,7 @@ public partial class Project : CommandUtils
 	{
 		ParamList<string> ListToProcess = InDedicatedServer && (Params.Cook || Params.CookOnTheFly) ? Params.ServerCookedTargets : Params.ClientCookedTargets;
 		var ConfigsToProcess = InDedicatedServer && (Params.Cook || Params.CookOnTheFly) ? Params.ServerConfigsToBuild : Params.ClientConfigsToBuild;
+		var CreateWebSocketsServer = Params.ServerTargetPlatforms.Count() > 0 && Params.ClientTargetPlatforms.Contains(UnrealTargetPlatform.HTML5);
 
 		List<UnrealTargetPlatform> PlatformsToStage = Params.ClientTargetPlatforms;
 		if (InDedicatedServer && (Params.Cook || Params.CookOnTheFly))
@@ -1461,7 +1483,7 @@ public partial class Project : CommandUtils
 				}
             }
 
-            string StageDirectory = (Params.Stage || !String.IsNullOrEmpty(Params.StageDirectoryParam)) ? Params.BaseStageDirectory : "";
+			string StageDirectory = ((ShouldCreatePak(Params) || (Params.Stage && !Params.SkipStage)) || !String.IsNullOrEmpty(Params.StageDirectoryParam)) ? Params.BaseStageDirectory : "";
             string ArchiveDirectory = (Params.Archive || !String.IsNullOrEmpty(Params.ArchiveDirectoryParam)) ? Params.BaseArchiveDirectory : "";
             if (prefixArchiveDir && (StagePlatform == UnrealTargetPlatform.Win32 || StagePlatform == UnrealTargetPlatform.Win64))
             {
@@ -1544,7 +1566,8 @@ public partial class Project : CommandUtils
 				Params.CookOnTheFly,
 				Params.Archive,
 				Params.IsProgramTarget,
-				Params.HasDedicatedServerAndClient
+				Params.HasDedicatedServerAndClient,
+				bInUseWebsocketNetDriver: CreateWebSocketsServer
 				);
 			LogDeploymentContext(SC);
 
