UAT: Get better diagnostics when we encounter a PathTooLongException in TempStorage.

diff --git a/Engine/Source/Programs/AutomationTool/GUBP/TempStorage.cs b/Engine/Source/Programs/AutomationTool/GUBP/TempStorage.cs
--- a/Engine/Source/Programs/AutomationTool/GUBP/TempStorage.cs
+++ b/Engine/Source/Programs/AutomationTool/GUBP/TempStorage.cs
@@ -558,9 +558,16 @@ namespace AutomationTool
                     {
                         var NewFile = CommandUtils.CombinePaths(NewBaseDir, ThisFileInfo.Name);
 
-                        FileInfo Info = new FileInfo(NewFile);
-
-                        Result.Add(Info.FullName);
+						try
+						{
+							FileInfo Info = new FileInfo(NewFile);
+							Result.Add(Info.FullName);
+						}
+						catch (PathTooLongException)
+						{
+							CommandUtils.LogError("Path too long ... failed to create FileInfo for {0}", NewFile);
+							throw;
+						}
                     }
                 }
                 if (Result.Count < 1)
