Fix to packaging blacklist to ignore blank lines in file
#ue4

diff --git a/Engine/Source/Programs/AutomationTool/Scripts/CopyBuildToStagingDirectory.Automation.cs b/Engine/Source/Programs/AutomationTool/Scripts/CopyBuildToStagingDirectory.Automation.cs
--- a/Engine/Source/Programs/AutomationTool/Scripts/CopyBuildToStagingDirectory.Automation.cs
+++ b/Engine/Source/Programs/AutomationTool/Scripts/CopyBuildToStagingDirectory.Automation.cs
@@ -647,14 +647,25 @@ public partial class Project : CommandUtils
 	private static Dictionary<string, string> CreatePakResponseFileFromStagingManifest(DeploymentContext SC)
 	{
 		// look for optional packaging blacklist if only one config active
-		string[] Blacklist = null;
+		List<string> Blacklist = null;
 		if (SC.StageTargetConfigurations.Count == 1)
 		{
 			var PakBlacklistFilename = CombinePaths(SC.ProjectRoot, "Build", SC.PlatformDir, string.Format("PakBlacklist-{0}.txt", SC.StageTargetConfigurations[0].ToString()));
 			if (File.Exists(PakBlacklistFilename))
 			{
 				Log("Applying PAK blacklist file {0}", PakBlacklistFilename);
-				Blacklist = File.ReadAllLines(PakBlacklistFilename);
+				string[] BlacklistContents = File.ReadAllLines(PakBlacklistFilename);
+				foreach (string Candidate in BlacklistContents)
+				{
+					if (Candidate.Trim().Length > 0)
+					{
+						if (Blacklist == null)
+						{
+							Blacklist = new List<string>();
+						}
+						Blacklist.Add(Candidate);
+					}
+				}
 			}
 		}
 
