#html5 fix packaging when prebuildsync was not called ( packaging content only games, sync'd from perforce)

diff --git a/Engine/Source/Programs/AutomationTool/HTML5/HTML5Platform.Automation.cs b/Engine/Source/Programs/AutomationTool/HTML5/HTML5Platform.Automation.cs
--- a/Engine/Source/Programs/AutomationTool/HTML5/HTML5Platform.Automation.cs
+++ b/Engine/Source/Programs/AutomationTool/HTML5/HTML5Platform.Automation.cs
@@ -26,7 +26,11 @@ public class HTML5Platform : Platform
 	public override void Package(ProjectParams Params, DeploymentContext SC, int WorkingCL)
 	{
 		Log("Package {0}", Params.RawProjectPath);
-		
+
+		Log("Setting Emscripten SDK for packaging..");
+		HTML5SDKInfo.SetupEmscriptenTemp();
+		HTML5SDKInfo.SetUpEmscriptenConfigFile();
+
 		string PackagePath = Path.Combine(Path.GetDirectoryName(Params.RawProjectPath), "Binaries", "HTML5");
 		if (!Directory.Exists(PackagePath))
 		{
diff --git a/Engine/Source/Programs/UnrealBuildTool/HTML5/HTML5SDKInfo.cs b/Engine/Source/Programs/UnrealBuildTool/HTML5/HTML5SDKInfo.cs
--- a/Engine/Source/Programs/UnrealBuildTool/HTML5/HTML5SDKInfo.cs
+++ b/Engine/Source/Programs/UnrealBuildTool/HTML5/HTML5SDKInfo.cs
@@ -101,7 +101,7 @@ namespace UnrealBuildTool
 
 		public static string  SetupEmscriptenTemp()
 		{
-			string HTML5Intermediatory = Path.GetFullPath(Path.Combine(BuildConfiguration.RelativeEnginePath, BuildConfiguration.PlatformIntermediateFolder));
+			string HTML5Intermediatory = Path.GetFullPath(Path.Combine(BuildConfiguration.RelativeEnginePath, BuildConfiguration.BaseIntermediateFolder, "HTML5"));
 			string TempPath= Path.Combine(HTML5Intermediatory, "EmscriptenTemp");
 			try
 			{
