Fix promotable triggers failing.

diff --git a/Engine/Source/Programs/AutomationTool/GUBP/GUBP.cs b/Engine/Source/Programs/AutomationTool/GUBP/GUBP.cs
--- a/Engine/Source/Programs/AutomationTool/GUBP/GUBP.cs
+++ b/Engine/Source/Programs/AutomationTool/GUBP/GUBP.cs
@@ -208,6 +208,7 @@ public partial class GUBP : BuildCommand
 					if (Trigger.Name.Equals(ExplicitTriggerName, StringComparison.InvariantCultureIgnoreCase))
                     {
 						Trigger.Activate();
+						ExplicitTrigger = Trigger;
                         break;
                     }
                 }
