fixes for executable signing for Rocket build

diff --git a/Engine/Source/Programs/AutomationTool/Scripts/RocketBuild.Automation.cs b/Engine/Source/Programs/AutomationTool/Scripts/RocketBuild.Automation.cs
--- a/Engine/Source/Programs/AutomationTool/Scripts/RocketBuild.Automation.cs
+++ b/Engine/Source/Programs/AutomationTool/Scripts/RocketBuild.Automation.cs
@@ -86,12 +86,17 @@ namespace Rocket
 				// Get the temp directory for stripped files for this host
 				string StrippedDir = Path.GetFullPath(CommandUtils.CombinePaths(CommandUtils.CmdEnv.LocalRoot, "Engine", "Saved", "Rocket", HostPlatform.ToString()));
 
+				// Get the temp directory for signed files for this host
+				string SignedDir = Path.GetFullPath(CommandUtils.CombinePaths(CommandUtils.CmdEnv.LocalRoot, "Engine", "Saved", "Rocket", "Signed", HostPlatform.ToString()));
+
 				// Strip the host platform
 				if (StripRocketNode.IsRequiredForPlatform(HostPlatform))
 				{
 					BranchConfig.AddNode(new StripRocketToolsNode(BranchConfig, HostPlatform, StrippedDir));
 					BranchConfig.AddNode(new StripRocketEditorNode(BranchConfig, HostPlatform, StrippedDir));
 				}
+				BranchConfig.AddNode(new SignRocketToolsNode(BranchConfig, HostPlatform, SignedDir));
+				BranchConfig.AddNode(new SignRocketEditorNode(BranchConfig, HostPlatform, SignedDir));
 
 				// Strip all the target platforms that are built on this host
 				foreach (UnrealTargetPlatform TargetPlatform in TargetPlatforms)
@@ -100,6 +105,10 @@ namespace Rocket
 					{
 						BranchConfig.AddNode(new StripRocketMonolithicsNode(BranchConfig, HostPlatform, TargetPlatform, StrippedDir));
 					}
+					if (GetSourceHostPlatform(BranchConfig.HostPlatforms, HostPlatform, TargetPlatform) == HostPlatform && SignRocketNode.IsRequiredForPlatform(TargetPlatform))
+					{
+						BranchConfig.AddNode(new SignRocketMonolithicsNode(BranchConfig, HostPlatform, TargetPlatform, SignedDir));
+					}
 				}
 
 				// Build the DDC
@@ -486,6 +495,185 @@ namespace Rocket
 		}
 	}
 
+	public abstract class SignRocketNode : GUBP.HostPlatformNode
+	{
+		public GUBP.GUBPBranchConfig BranchConfig;
+		public UnrealTargetPlatform TargetPlatform;
+		public string SignedDir;
+		public List<string> NodesToSign = new List<string>();
+
+		public SignRocketNode(GUBP.GUBPBranchConfig InBranchConfig, UnrealTargetPlatform InHostPlatform, UnrealTargetPlatform InTargetPlatform, string InSignedDir)
+			: base(InHostPlatform)
+		{
+			BranchConfig = InBranchConfig;
+			TargetPlatform = InTargetPlatform;
+			SignedDir = InSignedDir;
+		}
+
+		public override abstract string GetFullName();
+
+		public void AddNodeToSign(string NodeName)
+		{
+			NodesToSign.Add(NodeName);
+			AddDependency(NodeName);
+		}
+
+		public static bool IsRequiredForPlatform(UnrealTargetPlatform Platform)
+		{
+			return Platform == UnrealTargetPlatform.Mac || Platform == UnrealTargetPlatform.Win64 || Platform == UnrealTargetPlatform.Win32;
+		}
+
+		public override void DoBuild(GUBP bp)
+		{
+			BuildProducts = new List<string>();
+
+			string InputDir = Path.GetFullPath(CommandUtils.CmdEnv.LocalRoot);
+			string RulesFileName = CommandUtils.CombinePaths(CommandUtils.CmdEnv.LocalRoot, "Engine", "Build", "InstalledEngineFilters.ini");
+
+			// Read the filter for files on this platform
+			FileFilter SignFilter = new FileFilter();
+			if (HostPlatform == UnrealTargetPlatform.Mac)
+			{
+				SignFilter.AddRule("*.dylib");
+				SignFilter.AddRule("*.app");
+			}
+			else
+			{
+				SignFilter.AddRule("*.exe");
+				SignFilter.AddRule("*.dll");
+			}
+
+			// Apply the filter to the build products
+			List<string> SourcePaths = new List<string>();
+			List<string> TargetPaths = new List<string>();
+			foreach (string NodeToSign in NodesToSign)
+			{
+				GUBP.GUBPNode Node = BranchConfig.FindNode(NodeToSign);
+				foreach (string DependencyBuildProduct in Node.BuildProducts)
+				{
+					string RelativePath = CommandUtils.StripBaseDirectory(Path.GetFullPath(DependencyBuildProduct), InputDir);
+					if (SignFilter.Matches(RelativePath))
+					{
+						SourcePaths.Add(CommandUtils.CombinePaths(InputDir, RelativePath));
+						TargetPaths.Add(CommandUtils.CombinePaths(SignedDir, RelativePath));
+					}
+				}
+			}
+
+			// Strip the files and add them to the build products
+			SignFiles(bp, SourcePaths.ToArray(), TargetPaths.ToArray());
+			BuildProducts.AddRange(TargetPaths);
+
+			SaveRecordOfSuccessAndAddToBuildProducts();
+		}
+
+		public void SignFiles(GUBP bp, string[] SourceFileNames, string[] TargetFileNames)
+		{
+			// copy the files from source to target
+			for (int Idx = 0; Idx < SourceFileNames.Length; Idx++)
+			{
+				CommandUtils.CreateDirectory(Path.GetDirectoryName(TargetFileNames[Idx]));
+				CommandUtils.Log("Signing code: {0} -> {1}", SourceFileNames[Idx], TargetFileNames[Idx]);
+				CommandUtils.CopyFile(SourceFileNames[Idx], TargetFileNames[Idx]);
+			}
+
+			// Sign everything we built
+			CodeSign.SignMultipleIfEXEOrDLL(bp, TargetFileNames.ToList());
+		}
+	}
+
+	public class SignRocketToolsNode : SignRocketNode
+	{
+		public SignRocketToolsNode(GUBP.GUBPBranchConfig InBranchConfig, UnrealTargetPlatform InHostPlatform, string InSignedDir)
+			: base(InBranchConfig, InHostPlatform, InHostPlatform, InSignedDir)
+		{
+			AddNodeToSign(GUBP.ToolsForCompileNode.StaticGetFullName(HostPlatform));
+			AddNodeToSign(GUBP.ToolsNode.StaticGetFullName(HostPlatform));
+			AddDependency(StripRocketToolsNode.StaticGetFullName(HostPlatform));
+			AgentSharingGroup = "ToolsGroup" + StaticGetHostPlatformSuffix(InHostPlatform);
+		}
+
+		public static string StaticGetFullName(UnrealTargetPlatform InHostPlatform)
+		{
+			return "Sign" + GUBP.ToolsNode.StaticGetFullName(InHostPlatform);
+		}
+
+		public override string GetFullName()
+		{
+			return StaticGetFullName(HostPlatform);
+		}
+	}
+
+	public class SignRocketEditorNode : SignRocketNode
+	{
+		public SignRocketEditorNode(GUBP.GUBPBranchConfig InBranchConfig, UnrealTargetPlatform InHostPlatform, string InSignedDir)
+			: base(InBranchConfig, InHostPlatform, InHostPlatform, InSignedDir)
+		{
+			AddNodeToSign(GUBP.RootEditorNode.StaticGetFullName(HostPlatform));
+			AddDependency(StripRocketEditorNode.StaticGetFullName(HostPlatform));
+			AgentSharingGroup = "Editor" + StaticGetHostPlatformSuffix(HostPlatform);
+		}
+
+		public override float Priority()
+		{
+			return 1000000.0f;
+		}
+
+		public static string StaticGetFullName(UnrealTargetPlatform InHostPlatform)
+		{
+			return "Sign" + GUBP.RootEditorNode.StaticGetFullName(InHostPlatform);
+		}
+
+		public override string GetFullName()
+		{
+			return StaticGetFullName(HostPlatform);
+		}
+	}
+
+	public class SignRocketMonolithicsNode : SignRocketNode
+	{
+		BranchInfo.BranchUProject Project;
+		bool bIsCodeTargetPlatform;
+
+		public SignRocketMonolithicsNode(GUBP.GUBPBranchConfig InBranchConfig, UnrealTargetPlatform InHostPlatform, UnrealTargetPlatform InTargetPlatform, string InSignedDir)
+			: base(InBranchConfig, InHostPlatform, InTargetPlatform, InSignedDir)
+		{
+			Project = InBranchConfig.Branch.BaseEngineProject;
+			bIsCodeTargetPlatform = RocketBuild.IsCodeTargetPlatform(InHostPlatform, InTargetPlatform);
+
+			GUBP.GUBPNode Node = InBranchConfig.FindNode(GUBP.GamePlatformMonolithicsNode.StaticGetFullName(HostPlatform, Project, InTargetPlatform, Precompiled: bIsCodeTargetPlatform));
+			if (String.IsNullOrEmpty(Node.AgentSharingGroup))
+			{
+				Node.AgentSharingGroup = BranchConfig.Branch.BaseEngineProject.GameName + "_MonolithicsGroup_" + InTargetPlatform + StaticGetHostPlatformSuffix(InHostPlatform);
+			}
+			AddNodeToSign(Node.GetFullName());
+			string StripNode = StripRocketMonolithicsNode.StaticGetFullName(HostPlatform, BranchConfig.Branch.BaseEngineProject, TargetPlatform, RocketBuild.IsCodeTargetPlatform(HostPlatform, TargetPlatform));
+			AddDependency(StripNode);
+
+			AgentSharingGroup = Node.AgentSharingGroup;
+		}
+
+		public override string GetDisplayGroupName()
+		{
+			return Project.GameName + "_Monolithics" + (bIsCodeTargetPlatform ? "_Precompiled" : "");
+		}
+
+		public static string StaticGetFullName(UnrealTargetPlatform InHostPlatform, BranchInfo.BranchUProject InProject, UnrealTargetPlatform InTargetPlatform, bool bIsCodeTargetPlatform)
+		{
+			string Name = InProject.GameName + "_" + InTargetPlatform + "_Mono";
+			if (bIsCodeTargetPlatform)
+			{
+				Name += "_Precompiled";
+			}
+			return Name + "_Sign" + StaticGetHostPlatformSuffix(InHostPlatform);
+		}
+
+		public override string GetFullName()
+		{
+			return StaticGetFullName(HostPlatform, Project, TargetPlatform, bIsCodeTargetPlatform);
+		}
+	}
+
 	public class FilterRocketNode : GUBP.HostPlatformNode
 	{
 		GUBP.GUBPBranchConfig BranchConfig;
@@ -495,6 +683,7 @@ namespace Rocket
 		string[] CurrentTemplates;
 		public readonly string DepotManifestPath;
 		public Dictionary<string, string> StrippedNodeManifestPaths = new Dictionary<string, string>();
+		public Dictionary<string, string> SignedNodeManifestPaths = new Dictionary<string, string>();
 
 		public FilterRocketNode(GUBP.GUBPBranchConfig InBranchConfig, UnrealTargetPlatform InHostPlatform, List<UnrealTargetPlatform> InTargetPlatforms, string[] InCurrentFeaturePacks, string[] InCurrentTemplates)
 			: base(InHostPlatform)
@@ -532,6 +721,19 @@ namespace Rocket
 				}
 			}
 
+			// Also add signed node for all the target platforms that require it
+			List<string> SignedNodeNames = new List<string>();
+			foreach (UnrealTargetPlatform TargetPlatform in TargetPlatforms)
+			{
+				if (SignRocketNode.IsRequiredForPlatform(TargetPlatform))
+				{
+					UnrealTargetPlatform SourceHostPlatform = RocketBuild.GetSourceHostPlatform(BranchConfig.HostPlatforms, HostPlatform, TargetPlatform);
+					string SignNode = SignRocketMonolithicsNode.StaticGetFullName(SourceHostPlatform, BranchConfig.Branch.BaseEngineProject, TargetPlatform, RocketBuild.IsCodeTargetPlatform(SourceHostPlatform, TargetPlatform));
+					AddDependency(SignNode);
+					SignedNodeNames.Add(SignNode);
+				}
+			}
+
 			// Add win64 tools on Mac, to get the win64 build of UBT, UAT and IPP
 			if (HostPlatform == UnrealTargetPlatform.Mac && BranchConfig.HostPlatforms.Contains(UnrealTargetPlatform.Win64))
 			{
@@ -554,16 +756,31 @@ namespace Rocket
 			{
 				AddDependency(StripRocketToolsNode.StaticGetFullName(HostPlatform));
 				StrippedNodeNames.Add(StripRocketToolsNode.StaticGetFullName(HostPlatform));
-
 				AddDependency(StripRocketEditorNode.StaticGetFullName(HostPlatform));
 				StrippedNodeNames.Add(StripRocketEditorNode.StaticGetFullName(HostPlatform));
 			}
 
+			// Add the signed host platforms
+			if (SignRocketNode.IsRequiredForPlatform(HostPlatform))
+			{
+				AddDependency(SignRocketToolsNode.StaticGetFullName(HostPlatform));
+				SignedNodeNames.Add(SignRocketToolsNode.StaticGetFullName(HostPlatform));
+
+				AddDependency(SignRocketEditorNode.StaticGetFullName(HostPlatform));
+				SignedNodeNames.Add(SignRocketEditorNode.StaticGetFullName(HostPlatform));
+			}
+
 			// Set all the stripped manifest paths
 			foreach(string StrippedNodeName in StrippedNodeNames)
 			{
 				StrippedNodeManifestPaths.Add(StrippedNodeName, Path.Combine(Path.GetDirectoryName(DepotManifestPath), "Filter_" + StrippedNodeName + ".txt"));
 			}
+
+			// Set all the signed manifest paths
+			foreach (string SignedNodeName in SignedNodeNames)
+			{
+				SignedNodeManifestPaths.Add(SignedNodeName, Path.Combine(Path.GetDirectoryName(DepotManifestPath), "Filter_" + SignedNodeName + ".txt"));
+			}
 		}
 
 		public override int CISFrequencyQuantumShift(GUBP.GUBPBranchConfig BranchConfig)
@@ -676,6 +893,31 @@ namespace Rocket
 				BuildProducts.Add(StrippedNodeManifestPath.Value);
 			}
 
+			// Run the filter on the signed code, and remove those files from the copy filter
+			List<string> AllSignedFiles = new List<string>();
+			foreach (KeyValuePair<string, string> SignedNodeManifestPath in SignedNodeManifestPaths)
+			{
+				List<string> SignedFiles = new List<string>();
+
+				SignRocketNode SignNode = (SignRocketNode)BranchConfig.FindNode(SignedNodeManifestPath.Key);
+				foreach (string BuildProduct in SignNode.BuildProducts)
+				{
+					if (Utils.IsFileUnderDirectory(BuildProduct, SignNode.SignedDir))
+					{
+						string RelativePath = CommandUtils.StripBaseDirectory(Path.GetFullPath(BuildProduct), SignNode.SignedDir);
+						if (Filter.Matches(RelativePath))
+						{
+							SignedFiles.Add(RelativePath);
+							AllSignedFiles.Add(RelativePath);
+							Filter.Exclude("/" + RelativePath);
+						}
+					}
+				}
+
+				WriteManifest(SignedNodeManifestPath.Value, SignedFiles);
+				BuildProducts.Add(SignedNodeManifestPath.Value);
+			}
+
 			// Write the filtered list of depot files to disk, removing any symlinks
 			List<string> DepotFiles = Filter.ApplyToDirectory(CommandUtils.CmdEnv.LocalRoot, true).ToList();
 			WriteManifest(DepotManifestPath, DepotFiles);
@@ -770,9 +1012,6 @@ namespace Rocket
 			FilterRocketNode FilterNode = (FilterRocketNode)BranchConfig.FindNode(FilterRocketNode.StaticGetFullName(HostPlatform));
 			CopyManifestFilesToOutput(FilterNode.DepotManifestPath, CommandUtils.CmdEnv.LocalRoot, OutputDir);
 
-			// sign the executables
-			CodeSignManifestFiles(bp, FilterNode.DepotManifestPath, OutputDir);
-
 			// Copy the stripped files to the output directory
 			foreach(KeyValuePair<string, string> StrippedManifestPath in FilterNode.StrippedNodeManifestPaths)
 			{
@@ -780,6 +1019,13 @@ namespace Rocket
 				CopyManifestFilesToOutput(StrippedManifestPath.Value, StripNode.StrippedDir, OutputDir);
 			}
 
+			// Copy the signed files to the output directory
+			foreach (KeyValuePair<string, string> SignedManifestPath in FilterNode.SignedNodeManifestPaths)
+			{
+				SignRocketNode SignNode = (SignRocketNode)BranchConfig.FindNode(SignedManifestPath.Key);
+				CopyManifestFilesToOutput(SignedManifestPath.Value, SignNode.SignedDir, OutputDir);
+			}
+
 			// Copy the DDC to the output directory
 			BuildDerivedDataCacheNode DerivedDataCacheNode = (BuildDerivedDataCacheNode)BranchConfig.FindNode(BuildDerivedDataCacheNode.StaticGetFullName(HostPlatform));
 			CopyManifestFilesToOutput(DerivedDataCacheNode.SavedManifestPath, DerivedDataCacheNode.SavedDir, OutputDir);
@@ -793,20 +1039,6 @@ namespace Rocket
 			SaveRecordOfSuccessAndAddToBuildProducts();
 		}
 
-		static void CodeSignManifestFiles(BuildCommand bp, string ManifestPath, string OutputDir)
-		{
-			// Read the files from the manifest
-			CommandUtils.Log("Reading manifest: '{0}'", ManifestPath);
-			string[] Files = CommandUtils.ReadAllLines(ManifestPath).Select(x => x.Trim()).Where(x => x.Length > 0).ToArray();
-
-			// Create lists of source and target files
-			CommandUtils.Log("Code sign files...");
-			string[] SourceFiles = Files.Select(x => CommandUtils.CombinePaths(OutputDir, x)).Where(x => (x.Contains(".dll") || x.Contains(".exe")) && !(Path.GetDirectoryName(x).Replace("\\", "/")).Contains("Binaries/XboxOne")).ToArray();
-
-			// Copy everything
-			CodeSign.SignMultipleIfEXEOrDLL(bp, SourceFiles.ToList());
-		}
-
 		static void CopyManifestFilesToOutput(string ManifestPath, string InputDir, string OutputDir)
 		{
 			// Read the files from the manifest
