Fix incorrect GUBP error message.

diff --git a/Engine/Source/Programs/AutomationTool/GUBP/GUBP.cs b/Engine/Source/Programs/AutomationTool/GUBP/GUBP.cs
--- a/Engine/Source/Programs/AutomationTool/GUBP/GUBP.cs
+++ b/Engine/Source/Programs/AutomationTool/GUBP/GUBP.cs
@@ -1341,7 +1341,7 @@ public partial class GUBP : BuildCommand
 			}
 			if (FoundNames == 0)
 			{
-				throw new AutomationException("Could not find node named {0}", FoundNames);
+				throw new AutomationException("Could not find node named {0}", NameToDo);
 			}
 		}
 
