Fix unused variable warning causing build failure on Mono.

diff --git a/Engine/Source/Programs/AutomationTool/Scripts/RocketBuild.Automation.cs b/Engine/Source/Programs/AutomationTool/Scripts/RocketBuild.Automation.cs
--- a/Engine/Source/Programs/AutomationTool/Scripts/RocketBuild.Automation.cs
+++ b/Engine/Source/Programs/AutomationTool/Scripts/RocketBuild.Automation.cs
@@ -528,7 +528,6 @@ namespace Rocket
 			BuildProducts = new List<string>();
 
 			string InputDir = Path.GetFullPath(CommandUtils.CmdEnv.LocalRoot);
-			string RulesFileName = CommandUtils.CombinePaths(CommandUtils.CmdEnv.LocalRoot, "Engine", "Build", "InstalledEngineFilters.ini");
 
 			// Read the filter for files on this platform
 			FileFilter SignFilter = new FileFilter();
