Put all sample monolithic builds and cooks behind a separate trigger.

diff --git a/Engine/Source/Programs/AutomationTool/GUBP/LegacyBranchSetup.cs b/Engine/Source/Programs/AutomationTool/GUBP/LegacyBranchSetup.cs
--- a/Engine/Source/Programs/AutomationTool/GUBP/LegacyBranchSetup.cs
+++ b/Engine/Source/Programs/AutomationTool/GUBP/LegacyBranchSetup.cs
@@ -701,7 +701,8 @@ partial class GUBP
                     }
                 }
 
-                var CookedAgentSharingGroup = "Shared_CookedTests" + HostPlatformNode.StaticGetHostPlatformSuffix(HostPlatform);
+                var CookedSharedAgentSharingGroup = "Shared_CookedTests" + HostPlatformNode.StaticGetHostPlatformSuffix(HostPlatform);
+                var CookedSampleAgentSharingGroup = "Sample_CookedTests" + HostPlatformNode.StaticGetHostPlatformSuffix(HostPlatform);
 
                 var GameTestNodes = new List<string>();
                 var GameCookNodes = new List<string>();
@@ -746,7 +747,7 @@ partial class GUBP
                                     foreach (var Test in GameTests)
                                     {
                                         var TestName = Test.Key + "_" + Plat.ToString();
-                                        ThisMonoGameTestNodes.Add(BranchConfig.AddNode(new UATTestNode(BranchConfig, HostPlatform, Branch.BaseEngineProject, TestName, Test.Value, CookedAgentSharingGroup, false, RequiredPlatforms)));
+                                        ThisMonoGameTestNodes.Add(BranchConfig.AddNode(new UATTestNode(BranchConfig, HostPlatform, Branch.BaseEngineProject, TestName, Test.Value, CookedSharedAgentSharingGroup, false, RequiredPlatforms)));
                                     }
                                     if (ThisMonoGameTestNodes.Count > 0)
                                     {										
@@ -818,6 +819,7 @@ partial class GUBP
 											foreach (var Test in GameTests)
 											{
 												var TestName = Test.Key + "_" + Plat.ToString();
+												string CookedAgentSharingGroup = GamePlatformMonolithicsNode.IsSample(BranchConfig, NonCodeProject)? CookedSampleAgentSharingGroup : CookedSharedAgentSharingGroup;
 												ThisMonoGameTestNodes.Add(BranchConfig.AddNode(new UATTestNode(BranchConfig, HostPlatform, NonCodeProject, TestName, Test.Value, CookedAgentSharingGroup, false, RequiredPlatforms)));
 											}
 											if (ThisMonoGameTestNodes.Count > 0)
@@ -884,7 +886,14 @@ partial class GUBP
                 var AgentShareName = CodeProj.GameName;
                 if (!Options.bSeparateGamePromotion)
                 {
-                    AgentShareName = "Shared";
+					if(GamePlatformMonolithicsNode.IsSample(BranchConfig, CodeProj))
+					{
+						AgentShareName = "Sample";
+					}
+					else
+					{
+						AgentShareName = "Shared";
+					}
                 }
 
 				if (!BranchOptions.ExcludePlatformsForEditor.Contains(HostPlatform) && !Options.bIsNonCode)
@@ -1151,7 +1160,9 @@ partial class GUBP
             BranchConfig.AddNode(new SharedLabelPromotableSuccessNode());
             BranchConfig.AddNode(new WaitForTestShared(this));
             BranchConfig.AddNode(new WaitForSharedPromotionUserInput(this, true));
-			BranchConfig.AddNode(new SharedLabelPromotableNode(BranchConfig, true));			
+			BranchConfig.AddNode(new SharedLabelPromotableNode(BranchConfig, true));		
+
+			BranchConfig.AddNode(new WaitToPackageSamplesNode());
         }
 
 		AddCustomNodes(BranchConfig, HostPlatforms, ActivePlatforms);
diff --git a/Engine/Source/Programs/AutomationTool/GUBP/LegacyNodes.cs b/Engine/Source/Programs/AutomationTool/GUBP/LegacyNodes.cs
--- a/Engine/Source/Programs/AutomationTool/GUBP/LegacyNodes.cs
+++ b/Engine/Source/Programs/AutomationTool/GUBP/LegacyNodes.cs
@@ -1225,6 +1225,11 @@ partial class GUBP
 				AddDependency(ToolsNode.StaticGetFullName(InHostPlatform));
 			}
 
+			if(IsSample(BranchConfig, InGameProj))
+			{
+				AddDependency(WaitToPackageSamplesNode.StaticGetFullName());
+			}
+
             if (InGameProj.GameName != BranchConfig.Branch.BaseEngineProject.GameName && GameProj.Properties.Targets.ContainsKey(TargetRules.TargetType.Editor))
             {
 				if (!BranchConfig.BranchOptions.ExcludePlatformsForEditor.Contains(InHostPlatform))
@@ -1250,6 +1255,11 @@ partial class GUBP
             }
         }
 
+		public static bool IsSample(GUBPBranchConfig BranchConfig, BranchInfo.BranchUProject GameProj)
+		{
+			return (GameProj.GameName != BranchConfig.Branch.BaseEngineProject.GameName && (GameProj.FilePath.Contains("Samples") || GameProj.FilePath.Contains("Templates")));
+		}
+
 		public override string GetDisplayGroupName()
 		{
 			return GameProj.GameName + "_Monolithics" + (Precompiled? "_Precompiled" : "");
@@ -1797,6 +1807,34 @@ partial class GUBP
         }
     }
 
+    public class WaitToPackageSamplesNode : WaitForUserInput
+    {
+        public WaitToPackageSamplesNode()
+        {
+			AddDependency(WaitForSharedPromotionUserInput.StaticGetFullName(false));
+		}
+
+        public static string StaticGetFullName()
+        {
+            return "WaitToPackageSamples";
+        }
+
+        public override string GetFullName()
+        {
+            return StaticGetFullName();
+        }
+
+        public override string GetTriggerDescText()
+        {
+			return "Ready to package samples";
+        }
+
+        public override string GetTriggerActionText()
+        {
+			return "Package samples";
+        }
+    }
+
     public class WaitForPromotionUserInput : WaitForUserInput
     {
         string PromotionLabelPrefix;
@@ -2179,7 +2217,13 @@ partial class GUBP
                 {
                     AddPseudodependency(CookNode.StaticGetFullName(HostPlatform, BranchConfig.Branch.BaseEngineProject, BaseCookedPlatform));
                 }
-            }			
+            }
+
+			if(GamePlatformMonolithicsNode.IsSample(BranchConfig, GameProj))
+			{
+				AddDependency(WaitToPackageSamplesNode.StaticGetFullName());
+				AgentSharingGroup = "SampleCooks" + StaticGetHostPlatformSuffix(HostPlatform);
+			}
         }
         public static string StaticGetFullName(UnrealTargetPlatform InHostPlatform, BranchInfo.BranchUProject InGameProj, string InCookPlatform)
         {
@@ -2271,9 +2315,9 @@ partial class GUBP
 							if(Target.Rules.GUBP_BuildWindowsXPMonolithics())
 							{
 								AddDependency(GamePlatformMonolithicsNode.StaticGetFullName(HostPlatform, GameProj, TargetPlatform, true));
-                        }
-                    }
-                }
+							}
+						}
+					}
                 }
                 else
                 {
