Missing file.

diff --git a/Engine/Source/Programs/AutomationTool/Scripts/Tests.Automation.cs b/Engine/Source/Programs/AutomationTool/Scripts/Tests.Automation.cs
--- a/Engine/Source/Programs/AutomationTool/Scripts/Tests.Automation.cs
+++ b/Engine/Source/Programs/AutomationTool/Scripts/Tests.Automation.cs
@@ -1344,22 +1344,11 @@ public class ZeroEngineVersions : BuildCommand
 
 		Log("Checking if engine version files need to be reset...");
 		List<string> FilesToSubmit = new List<string>();
-		{
-			VersionFileUpdater ObjectVersionCpp = new VersionFileUpdater(ObjectVersionFilename);
-			if (ObjectVersionCpp.Contains("#define	ENGINE_VERSION	0") == false)
-			{
-				Log("Zeroing out engine versions in {0}", ObjectVersionFilename);
-				ObjectVersionCpp.ReplaceLine("#define	ENGINE_VERSION	", "0");
-				ObjectVersionCpp.Commit();
-				FilesToSubmit.Add(ObjectVersionFilename);
-			}
-		}
 		{
 			VersionFileUpdater VersionH = new VersionFileUpdater(VersionFilename);
 			if (VersionH.Contains("#define ENGINE_VERSION 0") == false)
 			{
 				Log("Zeroing out engine versions in {0}", VersionFilename);
-				VersionH.ReplaceLine("#define ENGINE_VERSION ", "0");
 				VersionH.ReplaceLine("#define BRANCH_NAME ", "\"" + P4Env.BranchName + "\"");
 				VersionH.ReplaceLine("#define BUILT_FROM_CHANGELIST ", "0");
 
