Stopped create release version from failing to copy.

diff --git a/Engine/Source/Programs/AutomationTool/Scripts/CopyBuildToStagingDirectory.Automation.cs b/Engine/Source/Programs/AutomationTool/Scripts/CopyBuildToStagingDirectory.Automation.cs
--- a/Engine/Source/Programs/AutomationTool/Scripts/CopyBuildToStagingDirectory.Automation.cs
+++ b/Engine/Source/Programs/AutomationTool/Scripts/CopyBuildToStagingDirectory.Automation.cs
@@ -812,7 +812,7 @@ public partial class Project : CommandUtils
             //string ReleaseVersionPath = CombinePaths( SC.ProjectRoot, "Releases", Params.CreateReleaseVersion, SC.StageTargetPlatform.GetCookPlatform(Params.DedicatedServer, false, Params.CookFlavor), Path.GetFileName(OutputLocation) );
             string ReleaseVersionPath = GetReleasePakFilePath(SC, Params, Params.CreateReleaseVersion, Path.GetFileName(OutputLocation));
 
-            File.Copy(OutputLocation, ReleaseVersionPath);
+            File.Copy(OutputLocation, ReleaseVersionPath, true);
         }
 
 		if (Params.CreateChunkInstall)
