UE-18919 - fix for signing issues on Mac
#uat

diff --git a/Engine/Source/Programs/AutomationTool/Mac/MacPlatform.Automation.cs b/Engine/Source/Programs/AutomationTool/Mac/MacPlatform.Automation.cs
--- a/Engine/Source/Programs/AutomationTool/Mac/MacPlatform.Automation.cs
+++ b/Engine/Source/Programs/AutomationTool/Mac/MacPlatform.Automation.cs
@@ -114,7 +114,7 @@ public class MacPlatform : Platform
 					// ensure the ue4game binary exists, if applicable
 					if (!SC.IsCodeBasedProject && !Directory.Exists(AbsoluteBundlePath) && !SC.bIsCombiningMultiplePlatforms)
 					{
-						Log("Failed to find app bundle " + AbsoluteBundlePath);
+						LogError("Failed to find app bundle " + AbsoluteBundlePath);
 						throw new AutomationException(ErrorCodes.Error_MissingExecutable, "Could not find app bundle {0}. You may need to build the UE4 project with your target configuration and platform.", AbsoluteBundlePath);
 					}
 
@@ -428,9 +428,18 @@ public class MacPlatform : Platform
 	{
 		// Sign everything we built
 		List<string> FilesToSign = GetExecutableNames(SC);
-		foreach (var File in FilesToSign)
+		foreach (var Exe in FilesToSign)
 		{
-			CodeSign.SignMacFileOrFolder(File);
+			string AppBundlePath = "";
+			if (Exe.StartsWith(CombinePaths(SC.RuntimeProjectRootDir, "Binaries", SC.PlatformDir)))
+			{
+				AppBundlePath = CombinePaths(SC.ShortProjectName, "Binaries", SC.PlatformDir, Path.GetFileNameWithoutExtension(Exe) + ".app");
+			}
+			else if (Exe.StartsWith(CombinePaths(SC.RuntimeRootDir, "Engine/Binaries", SC.PlatformDir)))
+			{
+				AppBundlePath = CombinePaths("Engine/Binaries", SC.PlatformDir, Path.GetFileNameWithoutExtension(Exe) + ".app");
+			}
+			CodeSign.SignMacFileOrFolder(AppBundlePath);
 		}
 
 		return true;
