Make sure directory exists for release paks

diff --git a/Engine/Source/Programs/AutomationTool/Scripts/CopyBuildToStagingDirectory.Automation.cs b/Engine/Source/Programs/AutomationTool/Scripts/CopyBuildToStagingDirectory.Automation.cs
--- a/Engine/Source/Programs/AutomationTool/Scripts/CopyBuildToStagingDirectory.Automation.cs
+++ b/Engine/Source/Programs/AutomationTool/Scripts/CopyBuildToStagingDirectory.Automation.cs
@@ -812,7 +812,8 @@ public partial class Project : CommandUtils
             //string ReleaseVersionPath = CombinePaths( SC.ProjectRoot, "Releases", Params.CreateReleaseVersion, SC.StageTargetPlatform.GetCookPlatform(Params.DedicatedServer, false, Params.CookFlavor), Path.GetFileName(OutputLocation) );
             string ReleaseVersionPath = GetReleasePakFilePath(SC, Params, Params.CreateReleaseVersion, Path.GetFileName(OutputLocation));
 
-            File.Copy(OutputLocation, ReleaseVersionPath, true);
+			InternalUtils.SafeCreateDirectory(Path.GetDirectoryName(ReleaseVersionPath));
+			InternalUtils.SafeCopyFile(OutputLocation, ReleaseVersionPath);
         }
 
 		if (Params.CreateChunkInstall)
