Remove Remove some redundant and unused functions;
* InternalUtils::ExecutingAssemblyLocation
* InternalUtils::ExecutingAssemblyDirectory
* InternalUtils::ExecutableVersion
* CommandUtils::WriteToFile (one overload not used)
* CommandUtils::ExeFilename
* CommandUtils::ExeDirectory
* CommandUtils::CurrentDirectory
Added:
* Tools.DotNETCommon.AssemblyUtils::ExecutableVersion

diff --git a/Engine/Source/Programs/AutomationTool/AutomationUtils/BuildEnvironment.cs b/Engine/Source/Programs/AutomationTool/AutomationUtils/BuildEnvironment.cs
--- a/Engine/Source/Programs/AutomationTool/AutomationUtils/BuildEnvironment.cs
+++ b/Engine/Source/Programs/AutomationTool/AutomationUtils/BuildEnvironment.cs
@@ -6,6 +6,7 @@ using System.Text;
 using System.Reflection;
 using Microsoft.Win32;
 using System.Diagnostics;
+using Tools.DotNETCommon;
 using UnrealBuildTool;
 
 namespace AutomationTool
@@ -75,7 +76,7 @@ namespace AutomationTool
 		{
 			if (String.IsNullOrEmpty(UATExe))
 			{
-				UATExe = CommandUtils.CombinePaths(Path.GetFullPath(InternalUtils.ExecutingAssemblyLocation));
+				UATExe = Assembly.GetExecutingAssembly().GetOriginalLocation();
 			}
 			if (!CommandUtils.FileExists_NoExceptions(UATExe))
 			{
diff --git a/Engine/Source/Programs/AutomationTool/AutomationUtils/CommandUtils.cs b/Engine/Source/Programs/AutomationTool/AutomationUtils/CommandUtils.cs
--- a/Engine/Source/Programs/AutomationTool/AutomationUtils/CommandUtils.cs
+++ b/Engine/Source/Programs/AutomationTool/AutomationUtils/CommandUtils.cs
@@ -13,6 +13,8 @@ using Tools.DotNETCommon.XmlHandler;
 using UnrealBuildTool;
 using System.Runtime.CompilerServices;
 using System.Linq;
+using System.Reflection;
+using Tools.DotNETCommon;
 using Tools.DotNETCommon.CaselessDictionary;
 using Tools.DotNETCommon.HarvestEnvVars;
 
@@ -971,18 +973,6 @@ namespace AutomationTool
 			return Attributes;
 		}
 
-		/// <summary>
-		/// Writes a line of formatted string to a file. Creates the file if it does not exists.
-		/// If the file does exists, appends a new line.
-		/// </summary>
-		/// <param name="Filename">Filename</param>
-		/// <param name="Format">Format string</param>
-		/// <param name="Args">Arguments</param>
-		public static void WriteToFile(string Filename, string Format, params object[] Args)
-		{
-			WriteToFile(Filename, String.Format(Format, Args));
-		}
-
 		/// <summary>
 		/// Writes a line of formatted string to a file. Creates the file if it does not exists.
 		/// If the file does exists, appends a new line.
@@ -991,32 +981,14 @@ namespace AutomationTool
 		/// <param name="Text">Text to write</param>
 		public static void WriteToFile(string Filename, string Text)
 		{
-			try
+            Filename = ConvertSeparators(PathSeparator.Default, Filename);
+            try
 			{
-				Filename = ConvertSeparators(PathSeparator.Default, Filename);
-				FileStream Stream;
-				if (File.Exists(Filename))
-				{
-					Stream = new FileStream(Filename, FileMode.Append, FileAccess.Write, FileShare.Read);
-				}
-				else
-				{
-					Stream = new FileStream(Filename, FileMode.OpenOrCreate, FileAccess.Write, FileShare.Read);
-				}
-				if (Stream != null)
-				{
-					using (StreamWriter Writer = new StreamWriter(Stream))
-					{
-						Writer.WriteLine(Text);
-						Writer.Flush();
-						Writer.Close();
-					}
-					Stream.Dispose();
-				}
+                File.AppendAllText(Filename, Text);
 			}
 			catch (Exception Ex)
 			{
-				throw new AutomationException(String.Format("Failed to Write to file {0}", Filename), Ex);
+				throw new AutomationException(Ex, "Failed to Write to file {0}", Filename);
 			}
 		}
 
@@ -1950,31 +1922,6 @@ namespace AutomationTool
 			set { CommandLineParams = value; }
 		}
 
-		/// <summary>
-		/// Path to the AutomationTool executable.
-		/// </summary>
-		public static string ExeFilename
-		{
-			get { return InternalUtils.ExecutingAssemblyLocation; }
-		}
-
-		/// <summary>
-		/// Directory where the AutomationTool executable sits.
-		/// </summary>
-		public static string ExeDirectory
-		{
-			get { return InternalUtils.ExecutingAssemblyDirectory; }
-		}
-
-		/// <summary>
-		/// Current directory.
-		/// </summary>
-		public static string CurrentDirectory
-		{
-			get { return Environment.CurrentDirectory; }
-			set { ChDir(value); }
-		}
-
 		/// <summary>
 		/// Checks if this command is running on a build machine.
 		/// </summary>
diff --git a/Engine/Source/Programs/AutomationTool/AutomationUtils/Utils.cs b/Engine/Source/Programs/AutomationTool/AutomationUtils/Utils.cs
--- a/Engine/Source/Programs/AutomationTool/AutomationUtils/Utils.cs
+++ b/Engine/Source/Programs/AutomationTool/AutomationUtils/Utils.cs
@@ -9,6 +9,7 @@ using System.Diagnostics;
 using System.Reflection;
 using UnrealBuildTool;
 using System.Text.RegularExpressions;
+using Tools.DotNETCommon;
 
 namespace AutomationTool
 {
@@ -690,17 +691,18 @@ namespace AutomationTool
                 return;
 			}
 			var bCreatedMutex = false;
-			var LocationHash = InternalUtils.ExecutingAssemblyLocation.GetHashCode();
-			var MutexName = "Global/" + Path.GetFileNameWithoutExtension(ExecutingAssemblyLocation) + "_" + LocationHash.ToString() + "_Mutex";
+            var EntryAssemblyLocation = Assembly.GetEntryAssembly().GetOriginalLocation();
+			var LocationHash = EntryAssemblyLocation.GetHashCode();
+            var MutexName = "Global/" + Path.GetFileNameWithoutExtension(EntryAssemblyLocation) + "_" + LocationHash.ToString() + "_Mutex";
 			using (Mutex SingleInstanceMutex = new Mutex(true, MutexName, out bCreatedMutex))
 			{
 				if (!bCreatedMutex)
 				{
-                    throw new AutomationException("Another instance of {0} is already running.", ExecutingAssemblyLocation);
+                    throw new AutomationException("Another instance of {0} is already running.", EntryAssemblyLocation);
 				}
 				else
 				{
-					Log.WriteLine(TraceEventType.Verbose, "No other instance of {0} is running.", ExecutingAssemblyLocation);
+                    Log.WriteLine(TraceEventType.Verbose, "No other instance of {0} is running.", EntryAssemblyLocation);
 				}
 
 				Main(Param);
@@ -708,39 +710,6 @@ namespace AutomationTool
 				SingleInstanceMutex.ReleaseMutex();
 			}
 		}
-
-		/// <summary>
-		/// Path to the executable which runs this code.
-		/// </summary>
-		public static string ExecutingAssemblyDirectory
-		{
-			get
-			{
-				return CommandUtils.CombinePaths(Path.GetDirectoryName(ExecutingAssemblyLocation));
-			}
-		}
-
-		/// <summary>
-		/// Filename of the executable which runs this code.
-		/// </summary>
-		public static string ExecutingAssemblyLocation
-		{
-			get
-			{
-				return CommandUtils.CombinePaths(new Uri(System.Reflection.Assembly.GetEntryAssembly().CodeBase).LocalPath);
-			}
-		}
-
-		/// <summary>
-		/// Version info of the executable which runs this code.
-		/// </summary>
-		public static FileVersionInfo ExecutableVersion
-		{
-			get
-			{
-				return FileVersionInfo.GetVersionInfo(ExecutingAssemblyLocation);
-			}
-		}
 	}
 
 	#endregion
diff --git a/Engine/Source/Programs/AutomationTool/Program.cs b/Engine/Source/Programs/AutomationTool/Program.cs
--- a/Engine/Source/Programs/AutomationTool/Program.cs
+++ b/Engine/Source/Programs/AutomationTool/Program.cs
@@ -8,6 +8,7 @@ using System.Threading;
 using System.Diagnostics;
 using UnrealBuildTool;
 using System.Reflection;
+using Tools.DotNETCommon;
 
 namespace AutomationTool
 {
@@ -33,8 +34,8 @@ namespace AutomationTool
                 XmlConfigLoader.Init();
 
                 // Log if we're running from the launcher
-                var ExecutingAssemblyLocation = CommandUtils.CombinePaths(Assembly.GetExecutingAssembly().Location);
-                if (String.Compare(ExecutingAssemblyLocation, CommandUtils.CombinePaths(InternalUtils.ExecutingAssemblyLocation), true) != 0)
+                var ExecutingAssemblyLocation = Assembly.GetExecutingAssembly().Location;
+                if (string.Compare(ExecutingAssemblyLocation, Assembly.GetEntryAssembly().GetOriginalLocation(), StringComparison.OrdinalIgnoreCase) != 0)
                 {
                     Log.WriteLine(TraceEventType.Information, "Executed from AutomationToolLauncher ({0})", ExecutingAssemblyLocation);
                 }
@@ -46,7 +47,7 @@ namespace AutomationTool
                 Domain.DomainUnload += Domain_ProcessExit;
                 HostPlatform.Current.SetConsoleCtrlHandler(ProgramCtrlHandler);
 
-                var Version = InternalUtils.ExecutableVersion;
+                var Version = AssemblyUtils.ExecutableVersion;
                 Log.WriteLine(TraceEventType.Verbose, "{0} ver. {1}", Version.ProductName, Version.ProductVersion);
 
                 // Don't allow simultaneous execution of AT (in the same branch)
diff --git a/Engine/Source/Programs/DotNETCommon/DotNETUtilities/AssemblyUtils.cs b/Engine/Source/Programs/DotNETCommon/DotNETUtilities/AssemblyUtils.cs
--- a/Engine/Source/Programs/DotNETCommon/DotNETUtilities/AssemblyUtils.cs
+++ b/Engine/Source/Programs/DotNETCommon/DotNETUtilities/AssemblyUtils.cs
@@ -1,6 +1,7 @@
 ﻿// Copyright 1998-2015 Epic Games, Inc. All Rights Reserved.
 
 using System;
+using System.Diagnostics;
 using System.IO;
 using System.Reflection;
 
@@ -18,5 +19,16 @@ namespace Tools.DotNETCommon
         {
             return new Uri(ThisAssembly.CodeBase).LocalPath;
         }
-	}
+    
+        /// <summary>
+        /// Version info of the executable which runs this code.
+        /// </summary>
+        public static FileVersionInfo ExecutableVersion
+        {
+            get
+            {
+                return FileVersionInfo.GetVersionInfo(Assembly.GetEntryAssembly().GetOriginalLocation());
+            }
+        }
+    }
 }
