Move Robust_XXX methods into InternalUtils, since it is called outside of TempStorage.
#UEB-307

diff --git a/Engine/Source/Programs/AutomationTool/AutomationUtils/Utils.cs b/Engine/Source/Programs/AutomationTool/AutomationUtils/Utils.cs
--- a/Engine/Source/Programs/AutomationTool/AutomationUtils/Utils.cs
+++ b/Engine/Source/Programs/AutomationTool/AutomationUtils/Utils.cs
@@ -109,7 +109,7 @@ namespace AutomationTool
 				}
 				if (Result == false && Attempts + 1 < MaxAttempts)
 				{
-					System.Threading.Thread.Sleep(1000);
+					Thread.Sleep(1000);
 				}
 			} while (Result == false && ++Attempts < MaxAttempts);
 
@@ -375,7 +375,7 @@ namespace AutomationTool
 				}
 				catch (Exception Ex)
 				{
-					Log.WriteLine(System.Diagnostics.TraceEventType.Warning, "SafeCopyFile Exception was {0}", LogUtils.FormatException(Ex));
+					Log.WriteLine(TraceEventType.Warning, "SafeCopyFile Exception was {0}", LogUtils.FormatException(Ex));
 					Retry = true;
 				}
 
@@ -464,7 +464,7 @@ namespace AutomationTool
 			// Windows needs the symlinks though because that's how deduplication works on Windows server, 
 			// see https://answers.unrealengine.com/questions/212888/automated-buildjenkins-failing-due-to-symlink-chec.html
 			// FIXME: ZFS, JFS and other fs that can be case-insensitive on Linux should use the faster path as well.
-			if (UnrealBuildTool.BuildHostPlatform.Current.Platform != UnrealTargetPlatform.Linux)
+			if (BuildHostPlatform.Current.Platform != UnrealTargetPlatform.Linux)
 			{
 				return Directory.GetFiles(Path, SearchPattern, Recursive ? SearchOption.AllDirectories : SearchOption.TopDirectoryOnly);
 			}
@@ -683,7 +683,7 @@ namespace AutomationTool
 		/// </summary>
 		/// <param name="Main"></param>
 		/// <param name="Param"></param>
-		public static void RunSingleInstance(System.Action<object> Main, object Param)
+		public static void RunSingleInstance(Action<object> Main, object Param)
 		{
 			if (Environment.GetEnvironmentVariable("uebp_UATMutexNoWait") == "1")
 			{
@@ -710,6 +710,114 @@ namespace AutomationTool
 				SingleInstanceMutex.ReleaseMutex();
 			}
 		}
+
+	    public static void Robust_CopyFile(string InputFileName, string OutputFileName)
+	    {
+	        if (OutputFileName.StartsWith("/Volumes/"))
+	        {
+	            int Retry = 0;
+	            int NumRetries = 60;
+	            bool bCopied = false;
+	            while (!bCopied && Retry < NumRetries)
+	            {
+	                if (Retry > 0)
+	                {
+	                    CommandUtils.Log("*** Mac temp storage retry {0}", OutputFileName);
+	                    System.Threading.Thread.Sleep(1000);
+	                }
+	                bCopied = CommandUtils.CopyFile_NoExceptions(InputFileName, OutputFileName, true);
+	                Retry++;
+	            }
+	        }
+	        else
+	        {
+	            CommandUtils.CopyFile(InputFileName, OutputFileName);
+	        }
+	    }
+
+	    public static void Robust_FileExists_NoExceptions(string Filename, string Message)
+	    {
+	        Robust_FileExists_NoExceptions(false, Filename, Message);
+	    }
+
+	    public static void Robust_FileExists_NoExceptions(bool bQuiet, string Filename, string Message)
+	    {
+	        if (!CommandUtils.FileExists_NoExceptions(bQuiet, Filename))
+	        {
+	            bool bFound = false;
+	            // mac is terrible on shares, this isn't a solution, but a stop gap
+	            if (Filename.StartsWith("/Volumes/"))
+	            {
+	                int Retry = 0;
+	                while (!bFound && Retry < 60)
+	                {
+	                    CommandUtils.Log(System.Diagnostics.TraceEventType.Information, "*** Mac temp storage retry {0}", Filename);
+	                    System.Threading.Thread.Sleep(10000);
+	                    bFound = CommandUtils.FileExists_NoExceptions(bQuiet, Filename);
+	                    Retry++;
+	                }
+	            }
+	            if (!bFound)
+	            {
+	                throw new AutomationException(Message, Filename);
+	            }
+	        }
+	    }
+
+	    public static bool Robust_DirectoryExists_NoExceptions(string Directoryname, string Message)
+	    {
+	        bool bFound = false;
+	        if (!CommandUtils.DirectoryExists_NoExceptions(Directoryname))
+	        {				
+	            // mac is terrible on shares, this isn't a solution, but a stop gap
+	            if (Directoryname.StartsWith("/Volumes/"))
+	            {
+	                int Retry = 0;
+	                while (!bFound && Retry < 60)
+	                {
+	                    CommandUtils.Log(System.Diagnostics.TraceEventType.Information, "*** Mac temp storage retry {0}", Directoryname);
+	                    System.Threading.Thread.Sleep(10000);
+	                    bFound = CommandUtils.DirectoryExists_NoExceptions(Directoryname);
+	                    Retry++;
+	                }
+	            }				
+	        }
+	        else
+	        {
+	            bFound = true;
+	        }
+	        return bFound;
+	    }
+
+	    public static bool Robust_DirectoryExistsAndIsWritable_NoExceptions(string Directoryname)
+	    {
+	        bool bFound = false;
+	        if (!CommandUtils.DirectoryExistsAndIsWritable_NoExceptions(Directoryname))
+	        {
+	            // mac is terrible on shares, this isn't a solution, but a stop gap
+	            if (Directoryname.StartsWith("/Volumes/"))
+	            {
+	                int Retry = 0;
+	                int NumRetries = 60;
+	                if(!Directoryname.Contains("UE4"))
+	                {
+	                    NumRetries = 2;
+	                }
+	                while (!bFound && Retry < NumRetries)
+	                {
+	                    CommandUtils.Log("*** Mac temp storage retry {0}", Directoryname);
+	                    System.Threading.Thread.Sleep(1000);
+	                    bFound = CommandUtils.DirectoryExistsAndIsWritable_NoExceptions(Directoryname);
+	                    Retry++;
+	                }
+	            }
+	        }
+	        else
+	        {
+	            bFound = true;
+	        }
+	        return bFound;
+	    }
 	}
 
 	#endregion
diff --git a/Engine/Source/Programs/AutomationTool/GUBP/TempStorage.cs b/Engine/Source/Programs/AutomationTool/GUBP/TempStorage.cs
--- a/Engine/Source/Programs/AutomationTool/GUBP/TempStorage.cs
+++ b/Engine/Source/Programs/AutomationTool/GUBP/TempStorage.cs
@@ -96,112 +96,6 @@ namespace AutomationTool
             }
         }
 
-		public static void Robust_CopyFile(string InputFileName, string OutputFileName)
-		{
-            if (OutputFileName.StartsWith("/Volumes/"))
-            {
-                int Retry = 0;
-                int NumRetries = 60;
-                bool bCopied = false;
-                while (!bCopied && Retry < NumRetries)
-                {
-                    if (Retry > 0)
-                    {
-                        CommandUtils.Log("*** Mac temp storage retry {0}", OutputFileName);
-                        System.Threading.Thread.Sleep(1000);
-                    }
-                    bCopied = CommandUtils.CopyFile_NoExceptions(InputFileName, OutputFileName, true);
-                    Retry++;
-                }
-            }
-            else
-            {
-                CommandUtils.CopyFile(InputFileName, OutputFileName);
-            }
-		}
-
-        public static void Robust_FileExists_NoExceptions(string Filename, string Message)
-        {
-            Robust_FileExists_NoExceptions(false, Filename, Message);
-        }
-        public static void Robust_FileExists_NoExceptions(bool bQuiet, string Filename, string Message)
-        {
-            if (!CommandUtils.FileExists_NoExceptions(bQuiet, Filename))
-            {
-                bool bFound = false;
-                // mac is terrible on shares, this isn't a solution, but a stop gap
-                if (Filename.StartsWith("/Volumes/"))
-                {
-                    int Retry = 0;
-                    while (!bFound && Retry < 60)
-                    {
-                        CommandUtils.Log(System.Diagnostics.TraceEventType.Information, "*** Mac temp storage retry {0}", Filename);
-                        System.Threading.Thread.Sleep(10000);
-                        bFound = CommandUtils.FileExists_NoExceptions(bQuiet, Filename);
-						Retry++;
-                    }
-                }
-                if (!bFound)
-                {
-                    throw new AutomationException(Message, Filename);
-                }
-            }
-        }
-
-		public static bool Robust_DirectoryExists_NoExceptions(string Directoryname, string Message)
-		{
-			bool bFound = false;
-			if (!CommandUtils.DirectoryExists_NoExceptions(Directoryname))
-			{				
-				// mac is terrible on shares, this isn't a solution, but a stop gap
-				if (Directoryname.StartsWith("/Volumes/"))
-				{
-					int Retry = 0;
-					while (!bFound && Retry < 60)
-					{
-						CommandUtils.Log(System.Diagnostics.TraceEventType.Information, "*** Mac temp storage retry {0}", Directoryname);
-						System.Threading.Thread.Sleep(10000);
-						bFound = CommandUtils.DirectoryExists_NoExceptions(Directoryname);
-						Retry++;
-					}
-				}				
-			}
-			else
-			{
-				bFound = true;
-			}
-			return bFound;
-		}
-        public static bool Robust_DirectoryExistsAndIsWritable_NoExceptions(string Directoryname)
-        {
-            bool bFound = false;
-            if (!CommandUtils.DirectoryExistsAndIsWritable_NoExceptions(Directoryname))
-            {
-                // mac is terrible on shares, this isn't a solution, but a stop gap
-                if (Directoryname.StartsWith("/Volumes/"))
-                {
-                    int Retry = 0;
-                    int NumRetries = 60;
-                    if(!Directoryname.Contains("UE4"))
-                    {
-                        NumRetries = 2;
-                    }
-                    while (!bFound && Retry < NumRetries)
-                    {
-                        CommandUtils.Log("*** Mac temp storage retry {0}", Directoryname);
-                        System.Threading.Thread.Sleep(1000);
-                        bFound = CommandUtils.DirectoryExistsAndIsWritable_NoExceptions(Directoryname);
-                        Retry++;
-                    }
-                }
-            }
-            else
-            {
-                bFound = true;
-            }
-            return bFound;
-        }
-
         public class TempStorageManifest
         {
             private static readonly string RootElementName = "tempstorage";
@@ -224,11 +118,11 @@ namespace AutomationTool
                 foreach (string InFilename in InFiles)
                 {
                     var Filename = CommandUtils.CombinePaths(InFilename);
-                    Robust_FileExists_NoExceptions(true, Filename, "Could not add {0} to manifest because it does not exist");
+                    InternalUtils.Robust_FileExists_NoExceptions(true, Filename, "Could not add {0} to manifest because it does not exist");
 
                     FileInfo Info = new FileInfo(Filename);
                     Filename = Info.FullName;
-                    Robust_FileExists_NoExceptions(true, Filename, "Could not add {0} to manifest because it does not exist2");
+                    InternalUtils.Robust_FileExists_NoExceptions(true, Filename, "Could not add {0} to manifest because it does not exist2");
 
                     if (!Filename.StartsWith(BaseFolder, StringComparison.InvariantCultureIgnoreCase))
                     {
@@ -576,7 +470,7 @@ namespace AutomationTool
                 }
 
 				CommandUtils.Log("Checking that first file in manifest exists...");
-				Robust_FileExists_NoExceptions(false, Result[0], "Rebased manifest file does not exist {0}");
+                InternalUtils.Robust_FileExists_NoExceptions(false, Result[0], "Rebased manifest file does not exist {0}");
 
                 return Result;
             }
@@ -658,7 +552,7 @@ namespace AutomationTool
                 int Retries = 0;
                 while (Retries < 24)
                 {
-                    if (Robust_DirectoryExistsAndIsWritable_NoExceptions(Dir))
+                    if (InternalUtils.Robust_DirectoryExistsAndIsWritable_NoExceptions(Dir))
                     {
                         return true;
                     }
@@ -667,7 +561,7 @@ namespace AutomationTool
 					Retries++;
                 }
             }
-            else if (Robust_DirectoryExists_NoExceptions(Dir, "Could not find {0}"))
+            else if (InternalUtils.Robust_DirectoryExists_NoExceptions(Dir, "Could not find {0}"))
             {
                 return true;
             }
@@ -683,7 +577,7 @@ namespace AutomationTool
             }
             string Root = CommandUtils.RootSharedTempStorageDirectory();
             string Result = CommandUtils.CombinePaths(Root, GameFolder);
-            if (String.IsNullOrEmpty(GameFolder) || !Robust_DirectoryExistsAndIsWritable_NoExceptions(Result))
+            if (String.IsNullOrEmpty(GameFolder) || !InternalUtils.Robust_DirectoryExistsAndIsWritable_NoExceptions(Result))
             {
                 string GameStr = "Game";
                 bool HadGame = false;
@@ -693,10 +587,10 @@ namespace AutomationTool
                     Result = CommandUtils.CombinePaths(Root, ShortFolder);
                     HadGame = true;
                 }
-                if (!HadGame || !Robust_DirectoryExistsAndIsWritable_NoExceptions(Result))
+                if (!HadGame || !InternalUtils.Robust_DirectoryExistsAndIsWritable_NoExceptions(Result))
                 {
                     Result = CommandUtils.CombinePaths(Root, "UE4");
-                    if (!Robust_DirectoryExistsAndIsWritable_NoExceptions(Result))
+                    if (!InternalUtils.Robust_DirectoryExistsAndIsWritable_NoExceptions(Result))
                     {
                         throw new AutomationException("Could not find an appropriate shared temp folder {0}", Result);
                     }
@@ -711,11 +605,11 @@ namespace AutomationTool
             string SharedSubdir = SharedTempStorageDirectory();
             string Result = CommandUtils.CombinePaths(ResolveSharedBuildDirectory(GameFolder), SharedSubdir);
 
-            if (!Robust_DirectoryExists_NoExceptions(Result, "Could not find {0}"))
+            if (!InternalUtils.Robust_DirectoryExists_NoExceptions(Result, "Could not find {0}"))
             {
                 CommandUtils.CreateDirectory_NoExceptions(Result);
             }
-            if (!Robust_DirectoryExists_NoExceptions(Result, "Could not find {0}"))
+            if (!InternalUtils.Robust_DirectoryExists_NoExceptions(Result, "Could not find {0}"))
             {
                 throw new AutomationException("Could not create an appropriate shared temp folder {0}", Result);
             }
@@ -879,7 +773,7 @@ namespace AutomationTool
                     foreach (string InFilename in Files)
                     {
                         var Filename = CommandUtils.CombinePaths(InFilename);
-                        Robust_FileExists_NoExceptions(false, Filename, "Could not add {0} to manifest because it does not exist");
+                        InternalUtils.Robust_FileExists_NoExceptions(false, Filename, "Could not add {0} to manifest because it does not exist");
 
                         if (!Filename.StartsWith(BaseFolder, StringComparison.InvariantCultureIgnoreCase))
                         {
@@ -892,7 +786,7 @@ namespace AutomationTool
                             throw new AutomationException("Dest file {0} already exists.", DestFile);
                         }
                         CommandUtils.CopyFile(Filename, DestFile, true);
-                        Robust_FileExists_NoExceptions(true, DestFile, "Could not copy to {0}");
+                        InternalUtils.Robust_FileExists_NoExceptions(true, DestFile, "Could not copy to {0}");
 
                         DestFiles.Add(DestFile);
                     }
@@ -903,7 +797,7 @@ namespace AutomationTool
                     foreach (string InFilename in Files)
                     {
                         var Filename = CommandUtils.CombinePaths(InFilename);
-                        Robust_FileExists_NoExceptions(false, Filename, "Could not add {0} to manifest because it does not exist");
+                        InternalUtils.Robust_FileExists_NoExceptions(false, Filename, "Could not add {0} to manifest because it does not exist");
 
                         if (!Filename.StartsWith(BaseFolder, StringComparison.InvariantCultureIgnoreCase))
                         {
@@ -921,7 +815,7 @@ namespace AutomationTool
                     CommandUtils.ThreadedCopyFiles(SrcFiles.ToArray(), DestFiles.ToArray(), ThreadsToCopyWith());
                     foreach (string DestFile in DestFiles)
                     {
-                        Robust_FileExists_NoExceptions(true, DestFile, "Could not copy to {0}");
+                        InternalUtils.Robust_FileExists_NoExceptions(true, DestFile, "Could not copy to {0}");
                     }
                 }
                 var Shared = SaveSharedTempStorageManifest(Env, StorageBlockName, GameFolder, DestFiles);
@@ -1009,7 +903,7 @@ namespace AutomationTool
                 throw new AutomationException("Storage Block Does Not Exists! {0}", BlockPath);
             }
             var SharedManifest = SharedTempStorageManifestFilename(Env, StorageBlockName, GameFolder);
-            Robust_FileExists_NoExceptions(SharedManifest, "Storage Block Manifest Does Not Exists! {0}");
+            InternalUtils.Robust_FileExists_NoExceptions(SharedManifest, "Storage Block Manifest Does Not Exists! {0}");
 
             var Shared = new TempStorageManifest();
             Shared.Load(SharedManifest);
@@ -1022,7 +916,7 @@ namespace AutomationTool
                 foreach (string InFilename in SharedFiles)
                 {
                     var Filename = CommandUtils.CombinePaths(InFilename);
-                    Robust_FileExists_NoExceptions(true, Filename, "Could not add {0} to manifest because it does not exist");
+                    InternalUtils.Robust_FileExists_NoExceptions(true, Filename, "Could not add {0} to manifest because it does not exist");
 
                     if (!Filename.StartsWith(BlockPath, StringComparison.InvariantCultureIgnoreCase))
                     {
@@ -1037,7 +931,7 @@ namespace AutomationTool
                     }
                     CommandUtils.CopyFile(Filename, DestFile, true);
 
-                    Robust_FileExists_NoExceptions(true, DestFile, "Could not copy to {0}");
+                    InternalUtils.Robust_FileExists_NoExceptions(true, DestFile, "Could not copy to {0}");
 
                     if (UnrealBuildTool.Utils.IsRunningOnMono)
                     {
@@ -1074,7 +968,7 @@ namespace AutomationTool
                 var NewDestFiles = new List<string>();
                 foreach (string DestFile in DestFiles)
                 {
-                    Robust_FileExists_NoExceptions(true, DestFile, "Could not copy to {0}");
+                    InternalUtils.Robust_FileExists_NoExceptions(true, DestFile, "Could not copy to {0}");
                     if (UnrealBuildTool.Utils.IsRunningOnMono)
                     {
 						CommandUtils.FixUnixFilePermissions(DestFile);
diff --git a/Engine/Source/Programs/AutomationTool/Scripts/RocketBuild.Automation.cs b/Engine/Source/Programs/AutomationTool/Scripts/RocketBuild.Automation.cs
--- a/Engine/Source/Programs/AutomationTool/Scripts/RocketBuild.Automation.cs
+++ b/Engine/Source/Programs/AutomationTool/Scripts/RocketBuild.Automation.cs
@@ -1108,8 +1108,8 @@ namespace Rocket
 
 			// Copy the files to their final location
 			CommandUtils.Log("Copying files to {0}", PublishDir);
-			TempStorage.Robust_CopyFile(FullZipFileName, Path.Combine(PublishDir, Path.GetFileName(FullZipFileName)));
-			TempStorage.Robust_CopyFile(EditorZipFileName, Path.Combine(PublishDir, Path.GetFileName(EditorZipFileName)));
+			InternalUtils.Robust_CopyFile(FullZipFileName, Path.Combine(PublishDir, Path.GetFileName(FullZipFileName)));
+			InternalUtils.Robust_CopyFile(EditorZipFileName, Path.Combine(PublishDir, Path.GetFileName(EditorZipFileName)));
 			CommandUtils.DeleteFile(FullZipFileName);
 			CommandUtils.DeleteFile(EditorZipFileName);
 
@@ -1167,7 +1167,7 @@ namespace Rocket
 					if(DebugExtensions.Contains(Extension) || Extension == ".exe" || Extension == ".dll") // Need all windows build products for crash reporter
 					{
 						string OutputFileName = CommandUtils.MakeRerootedFilePath(InputFileName, CommandUtils.CmdEnv.LocalRoot, SymbolsOutputDir);
-						TempStorage.Robust_CopyFile(InputFileName, OutputFileName);
+						InternalUtils.Robust_CopyFile(InputFileName, OutputFileName);
 					}
 				}
 			}
