[INTEGRATE] UAT: Add support for cross-app promotion

diff --git a/Engine/Source/Programs/AutomationTool/AutomationUtils/MCPPublic.cs b/Engine/Source/Programs/AutomationTool/AutomationUtils/MCPPublic.cs
--- a/Engine/Source/Programs/AutomationTool/AutomationUtils/MCPPublic.cs
+++ b/Engine/Source/Programs/AutomationTool/AutomationUtils/MCPPublic.cs
@@ -150,11 +150,23 @@ namespace EpicGames.MCP.Automation
         {
             get
             {
-                var BaseFilename = AppName + BuildVersion + "-" + Platform.ToString() + ".manifest";
+				if (!string.IsNullOrEmpty(_ManifestFilename))
+				{
+					return _ManifestFilename;
+				}
+				var BaseFilename = string.Format("{0}{1}-{2}.manifest",
+					AppName,
+					BuildVersion,
+					Platform.ToString());
                 return Regex.Replace(BaseFilename, @"\s+", ""); // Strip out whitespace in order to be compatible with BuildPatchTool
             }
         }
 
+		/// <summary>
+		/// If set, this allows us to over-ride the automatically constructed ManifestFilename
+		/// </summary>
+		private readonly string _ManifestFilename;
+
         /// <summary>
         /// Determine the platform name (Win32/64 becomes Windows, Mac is Mac, the rest we don't currently understand)
         /// </summary>
@@ -216,7 +228,8 @@ namespace EpicGames.MCP.Automation
         /// <param name="InBuildVersion"></param>
         /// <param name="platform"></param>
         /// <param name="stagingDirRelativePath">Relative path from the BuildRootPath where files will be staged. Commonly matches the AppName.</param>
-        public BuildPatchToolStagingInfo(BuildCommand InOwnerCommand, string InAppName, string InMcpConfigKey, int InAppID, string InBuildVersion, UnrealTargetPlatform platform, string stagingDirRelativePath)
+		/// <param name="InManifestFilename">If specifies, will override the value returned by the ManifestFilename property</param>
+        public BuildPatchToolStagingInfo(BuildCommand InOwnerCommand, string InAppName, string InMcpConfigKey, int InAppID, string InBuildVersion, UnrealTargetPlatform platform, string stagingDirRelativePath, string InManifestFilename = null)
             : this(InOwnerCommand, InAppName, InMcpConfigKey, InAppID, InBuildVersion, ToMCPPlatform(platform), stagingDirRelativePath)
         {
         }
@@ -229,10 +242,11 @@ namespace EpicGames.MCP.Automation
         /// <param name="InBuildVersion"></param>
         /// <param name="platform"></param>
         /// <param name="stagingDirRelativePath">Relative path from the BuildRootPath where files will be staged. Commonly matches the AppName.</param>
-        public BuildPatchToolStagingInfo(BuildCommand InOwnerCommand, string InAppName, string InMcpConfigKey, int InAppID, string InBuildVersion, MCPPlatform platform, string stagingDirRelativePath)
+		public BuildPatchToolStagingInfo(BuildCommand InOwnerCommand, string InAppName, string InMcpConfigKey, int InAppID, string InBuildVersion, MCPPlatform platform, string stagingDirRelativePath, string InManifestFilename = null)
         {
             OwnerCommand = InOwnerCommand;
             AppName = InAppName;
+			_ManifestFilename = InManifestFilename;
             McpConfigKey = InMcpConfigKey;
             AppID = InAppID;
             BuildVersion = InBuildVersion;
@@ -246,7 +260,7 @@ namespace EpicGames.MCP.Automation
 		/// <summary>
 		/// Basic constructor with staging dir suffix override, basically to avoid having platform concatenated
 		/// </summary>
-		public BuildPatchToolStagingInfo(BuildCommand InOwnerCommand, string InAppName, string InMcpConfigKey, int InAppID, string InBuildVersion, UnrealTargetPlatform platform, string stagingDirRelativePath, string stagingDirSuffix)
+		public BuildPatchToolStagingInfo(BuildCommand InOwnerCommand, string InAppName, string InMcpConfigKey, int InAppID, string InBuildVersion, UnrealTargetPlatform platform, string stagingDirRelativePath, string stagingDirSuffix, string InManifestFilename = null)
 			: this(InOwnerCommand, InAppName, InMcpConfigKey, InAppID, InBuildVersion, ToMCPPlatform(platform), stagingDirRelativePath, stagingDirSuffix)
 		{
 		}
@@ -254,10 +268,11 @@ namespace EpicGames.MCP.Automation
 		/// <summary>
 		/// Basic constructor with staging dir suffix override, basically to avoid having platform concatenated
 		/// </summary>
-		public BuildPatchToolStagingInfo(BuildCommand InOwnerCommand, string InAppName, string InMcpConfigKey, int InAppID, string InBuildVersion, MCPPlatform platform, string stagingDirRelativePath, string stagingDirSuffix)
+		public BuildPatchToolStagingInfo(BuildCommand InOwnerCommand, string InAppName, string InMcpConfigKey, int InAppID, string InBuildVersion, MCPPlatform platform, string stagingDirRelativePath, string stagingDirSuffix, string InManifestFilename = null)
 		{
 			OwnerCommand = InOwnerCommand;
 			AppName = InAppName;
+			_ManifestFilename = InManifestFilename;
 			McpConfigKey = InMcpConfigKey;
 			AppID = InAppID;
 			BuildVersion = InBuildVersion;
@@ -493,6 +508,15 @@ namespace EpicGames.MCP.Automation
             }
             return Handler;
         }
+
+		/// <summary>
+		/// Determines whether a given build is registered in build info
+		/// </summary>
+		/// <param name="StagingInfo">The staging info representing the build to check.</param>
+		/// <param name="McpConfigName">Name of which MCP config to check against.</param>
+		/// <returns>true if the build is registered, false otherwise</returns>
+		abstract public bool BuildExists(BuildPatchToolStagingInfo StagingInfo, string McpConfigName);
+
         /// <summary>
         /// Given a MCPStagingInfo defining our build info, posts the build to the MCP BuildInfo Service.
         /// </summary>
@@ -513,6 +537,13 @@ namespace EpicGames.MCP.Automation
 		/// <param name="McpConfigName">Which BuildInfo backend to get labels from for this promotion attempt.</param>
 		abstract public List<string> GetBuildLabels(BuildPatchToolStagingInfo StagingInfo, string McpConfigName);
 
+		/// <summary>
+		/// Given a staging info defining our build, return the manifest url for that registered build
+		/// </summary>
+		/// <param name="StagingInfo">Staging Info describing the BuildInfo to query.</param>
+		/// <param name="McpConfigName">Name of which MCP config to query.</param>
+		/// <returns></returns>
+		abstract public string GetBuildManifestUrl(BuildPatchToolStagingInfo StagingInfo, string McpConfigName);
 		/// <summary>
 		/// Get a label string for the specific Platform requested.
 		/// </summary>
