[INTEGRATE]  Change 2627031 by [REDACTED_PERSON][REDACTED_HANDLE] on 2015/07/21 12:06

        #UE4 Fixed a problem where the wrong constructor was matched for BuildPatchToolStagingInfo when supplying two or three strings at the end of the arguments list since the last argument was optional.

diff --git a/Engine/Source/Programs/AutomationTool/AutomationUtils/MCPPublic.cs b/Engine/Source/Programs/AutomationTool/AutomationUtils/MCPPublic.cs
--- a/Engine/Source/Programs/AutomationTool/AutomationUtils/MCPPublic.cs
+++ b/Engine/Source/Programs/AutomationTool/AutomationUtils/MCPPublic.cs
@@ -162,7 +162,7 @@ namespace EpicGames.MCP.Automation
             }
         }
 
-		/// <summary>
+        /// <summary>
 		/// If set, this allows us to over-ride the automatically constructed ManifestFilename
 		/// </summary>
 		private readonly string _ManifestFilename;
@@ -229,7 +229,7 @@ namespace EpicGames.MCP.Automation
         /// <param name="platform"></param>
         /// <param name="stagingDirRelativePath">Relative path from the BuildRootPath where files will be staged. Commonly matches the AppName.</param>
 		/// <param name="InManifestFilename">If specifies, will override the value returned by the ManifestFilename property</param>
-        public BuildPatchToolStagingInfo(BuildCommand InOwnerCommand, string InAppName, string InMcpConfigKey, int InAppID, string InBuildVersion, UnrealTargetPlatform platform, string stagingDirRelativePath, string InManifestFilename = null)
+        public BuildPatchToolStagingInfo(BuildCommand InOwnerCommand, string InAppName, string InMcpConfigKey, int InAppID, string InBuildVersion, UnrealTargetPlatform platform, string stagingDirRelativePath)
             : this(InOwnerCommand, InAppName, InMcpConfigKey, InAppID, InBuildVersion, ToMCPPlatform(platform), stagingDirRelativePath)
         {
         }
@@ -242,11 +242,11 @@ namespace EpicGames.MCP.Automation
         /// <param name="InBuildVersion"></param>
         /// <param name="platform"></param>
         /// <param name="stagingDirRelativePath">Relative path from the BuildRootPath where files will be staged. Commonly matches the AppName.</param>
-		public BuildPatchToolStagingInfo(BuildCommand InOwnerCommand, string InAppName, string InMcpConfigKey, int InAppID, string InBuildVersion, MCPPlatform platform, string stagingDirRelativePath, string InManifestFilename = null)
+		public BuildPatchToolStagingInfo(BuildCommand InOwnerCommand, string InAppName, string InMcpConfigKey, int InAppID, string InBuildVersion, MCPPlatform platform, string stagingDirRelativePath)
         {
             OwnerCommand = InOwnerCommand;
             AppName = InAppName;
-			_ManifestFilename = InManifestFilename;
+			_ManifestFilename = null;
             McpConfigKey = InMcpConfigKey;
             AppID = InAppID;
             BuildVersion = InBuildVersion;
@@ -260,15 +260,15 @@ namespace EpicGames.MCP.Automation
 		/// <summary>
 		/// Basic constructor with staging dir suffix override, basically to avoid having platform concatenated
 		/// </summary>
-		public BuildPatchToolStagingInfo(BuildCommand InOwnerCommand, string InAppName, string InMcpConfigKey, int InAppID, string InBuildVersion, UnrealTargetPlatform platform, string stagingDirRelativePath, string stagingDirSuffix, string InManifestFilename = null)
-			: this(InOwnerCommand, InAppName, InMcpConfigKey, InAppID, InBuildVersion, ToMCPPlatform(platform), stagingDirRelativePath, stagingDirSuffix)
+		public BuildPatchToolStagingInfo(BuildCommand InOwnerCommand, string InAppName, string InMcpConfigKey, int InAppID, string InBuildVersion, UnrealTargetPlatform platform, string stagingDirRelativePath, string stagingDirSuffix, string InManifestFilename)
+			: this(InOwnerCommand, InAppName, InMcpConfigKey, InAppID, InBuildVersion, ToMCPPlatform(platform), stagingDirRelativePath, stagingDirSuffix, InManifestFilename)
 		{
 		}
 
 		/// <summary>
 		/// Basic constructor with staging dir suffix override, basically to avoid having platform concatenated
 		/// </summary>
-		public BuildPatchToolStagingInfo(BuildCommand InOwnerCommand, string InAppName, string InMcpConfigKey, int InAppID, string InBuildVersion, MCPPlatform platform, string stagingDirRelativePath, string stagingDirSuffix, string InManifestFilename = null)
+		public BuildPatchToolStagingInfo(BuildCommand InOwnerCommand, string InAppName, string InMcpConfigKey, int InAppID, string InBuildVersion, MCPPlatform platform, string stagingDirRelativePath, string stagingDirSuffix, string InManifestFilename)
 		{
 			OwnerCommand = InOwnerCommand;
 			AppName = InAppName;
