More logiging for plugin discovery

diff --git a/Engine/Source/Programs/AutomationTool/Scripts/CopyBuildToStagingDirectory.Automation.cs b/Engine/Source/Programs/AutomationTool/Scripts/CopyBuildToStagingDirectory.Automation.cs
--- a/Engine/Source/Programs/AutomationTool/Scripts/CopyBuildToStagingDirectory.Automation.cs
+++ b/Engine/Source/Programs/AutomationTool/Scripts/CopyBuildToStagingDirectory.Automation.cs
@@ -374,11 +374,15 @@ public partial class Project : CommandUtils
 
 			// Engine & Game Plugins
 			ProjectDescriptor Project = ProjectDescriptor.FromFile(SC.RawProjectPath);
+			Console.WriteLine("Searching for plugins with CurrentWorkingDir: " + Directory.GetCurrentDirectory());
+			Console.WriteLine("Searching for plugins in: " + SC.RawProjectPath);
 			List<PluginInfo> AvailablePlugins = Plugins.ReadAvailablePlugins(SC.RawProjectPath);
 			foreach (PluginInfo Plugin in AvailablePlugins)
 			{
+				Console.WriteLine("Considering Plugin for Stage: " + Plugin.FileName);
 				if (UProjectInfo.IsPluginEnabledForProject(Plugin, Project, SC.StageTargetPlatform.PlatformType))
 				{
+					Console.WriteLine("EnabledPlugin: " + Plugin.FileName);
 					SC.StageFiles(StagedFileType.UFS, Plugin.Directory, "*.uplugin", false, null, null, true, !Params.UsePak(SC.StageTargetPlatform), null, true, false);
 				}
 			}
diff --git a/Engine/Source/Programs/UnrealBuildTool/System/Plugins.cs b/Engine/Source/Programs/UnrealBuildTool/System/Plugins.cs
--- a/Engine/Source/Programs/UnrealBuildTool/System/Plugins.cs
+++ b/Engine/Source/Programs/UnrealBuildTool/System/Plugins.cs
@@ -62,9 +62,12 @@ namespace UnrealBuildTool
 			List<PluginInfo> Plugins = new List<PluginInfo>();
 
 			// Read all the engine plugins
+			
 			string EnginePluginsDir = Path.Combine(BuildConfiguration.RelativeEnginePath, "Plugins");
+			Console.WriteLine("EngineDir for plugin search: " + EnginePluginsDir);
 			foreach(string PluginFileName in EnumeratePlugins(EnginePluginsDir))
 			{
+				Console.WriteLine("Found Engine Plugin: " + PluginFileName);
 				PluginInfo Plugin = new PluginInfo(PluginFileName, PluginLoadedFrom.Engine);
 				Plugins.Add(Plugin);
 			}
diff --git a/Engine/Source/Programs/UnrealBuildTool/System/UProjectInfo.cs b/Engine/Source/Programs/UnrealBuildTool/System/UProjectInfo.cs
--- a/Engine/Source/Programs/UnrealBuildTool/System/UProjectInfo.cs
+++ b/Engine/Source/Programs/UnrealBuildTool/System/UProjectInfo.cs
@@ -310,13 +310,16 @@ namespace UnrealBuildTool
 		public static bool IsPluginEnabledForProject(PluginInfo Plugin, ProjectDescriptor Project, UnrealTargetPlatform Platform)
 		{
 			bool bEnabled = Plugin.Descriptor.bEnabledByDefault || Plugin.LoadedFrom == PluginLoadedFrom.GameProject;
+			Console.WriteLine("Plugin passes first enabled check: " + Plugin.FileName);
 			if(Project != null && Project.Plugins != null)
 			{
 				foreach(PluginReferenceDescriptor PluginReference in Project.Plugins)
 				{
 					if(String.Compare(PluginReference.Name, Plugin.Name, true) == 0)
 					{
+						Console.WriteLine("Plugin passes project ref check: " + Plugin.FileName);
 						bEnabled = PluginReference.IsEnabledForPlatform(Platform);
+						Console.WriteLine("Plugin " + Plugin.FileName + "now " + bEnabled.ToString());
 					}
 				}
 			}
