fix for build machine failure when copying DDC stats when running BuildDerivedDataCache commandlet

diff --git a/Engine/Source/Programs/AutomationTool/AutomationUtils/CommandletUtils.cs b/Engine/Source/Programs/AutomationTool/AutomationUtils/CommandletUtils.cs
--- a/Engine/Source/Programs/AutomationTool/AutomationUtils/CommandletUtils.cs
+++ b/Engine/Source/Programs/AutomationTool/AutomationUtils/CommandletUtils.cs
@@ -191,7 +191,7 @@ namespace AutomationTool
 			}
 			// copy the ddc access log to the destination folder
 			string DestDDCLogFile = LogUtils.GetUniqueLogName(CombinePaths(CmdEnv.LogFolder, Commandlet + "-DDC"));
-			string LocalDDCLogFile = CombinePaths(CmdEnv.EngineSavedFolder, Commandlet + "RunDDC.txt");
+			string LocalDDCLogFile = CombinePaths(CmdEnv.EngineSavedFolder, "CookRunDDC.txt");
 			if (!CommandUtils.CopyFile_NoExceptions(LocalDDCLogFile, DestDDCLogFile))
 			{
 				CommandUtils.LogWarning("Commandlet {0} failed to copy the local log file from {1} to {2}. The log file will be lost.", Commandlet, LocalDDCLogFile, DestDDCLogFile);
