Plugin files are now read from the correct working directory while creating the staging manifest.

diff --git a/Engine/Source/Programs/AutomationTool/Scripts/CopyBuildToStagingDirectory.Automation.cs b/Engine/Source/Programs/AutomationTool/Scripts/CopyBuildToStagingDirectory.Automation.cs
--- a/Engine/Source/Programs/AutomationTool/Scripts/CopyBuildToStagingDirectory.Automation.cs
+++ b/Engine/Source/Programs/AutomationTool/Scripts/CopyBuildToStagingDirectory.Automation.cs
@@ -372,7 +372,8 @@ public partial class Project : CommandUtils
                 StageLocalizationDataForCulture(SC, Culture, CombinePaths(SC.LocalRoot, "Engine/Content/Localization/Engine"), null, !Params.UsePak(SC.StageTargetPlatform));
             }
 
-			// Engine & Game Plugins
+			// Engine & Game Plugins. Push the Engine/Source working directory so UBT code has a correct RelativeEnginePath in ReadAvailablePlugins
+			PushDir(CombinePaths(SC.LocalRoot, "Engine", "Source"));
 			ProjectDescriptor Project = ProjectDescriptor.FromFile(SC.RawProjectPath);
 			Console.WriteLine("Searching for plugins with CurrentWorkingDir: " + Directory.GetCurrentDirectory());
 			Console.WriteLine("Searching for plugins in: " + SC.RawProjectPath);
@@ -386,6 +387,7 @@ public partial class Project : CommandUtils
 					SC.StageFiles(StagedFileType.UFS, Plugin.Directory, "*.uplugin", false, null, null, true, !Params.UsePak(SC.StageTargetPlatform), null, true, false);
 				}
 			}
+			PopDir();
 
             // Game ufs (content)
 
