Ensure EngineRelativeDir is correct for plugin staging

diff --git a/Engine/Source/Programs/AutomationTool/Scripts/CopyBuildToStagingDirectory.Automation.cs b/Engine/Source/Programs/AutomationTool/Scripts/CopyBuildToStagingDirectory.Automation.cs
--- a/Engine/Source/Programs/AutomationTool/Scripts/CopyBuildToStagingDirectory.Automation.cs
+++ b/Engine/Source/Programs/AutomationTool/Scripts/CopyBuildToStagingDirectory.Automation.cs
@@ -371,9 +371,12 @@ public partial class Project : CommandUtils
             {
                 StageLocalizationDataForCulture(SC, Culture, CombinePaths(SC.LocalRoot, "Engine/Content/Localization/Engine"), null, !Params.UsePak(SC.StageTargetPlatform));
             }
-
+			
 			// Engine & Game Plugins. Push the Engine/Source working directory so UBT code has a correct RelativeEnginePath in ReadAvailablePlugins
 			PushDir(CombinePaths(SC.LocalRoot, "Engine", "Source"));
+			String OrigEngineRelativeDir = BuildConfiguration.RelativeEnginePath;
+			BuildConfiguration.RelativeEnginePath = "../../Engine/";
+
 			ProjectDescriptor Project = ProjectDescriptor.FromFile(SC.RawProjectPath);
 			Console.WriteLine("Searching for plugins with CurrentWorkingDir: " + Directory.GetCurrentDirectory());
 			Console.WriteLine("Searching for plugins in: " + SC.RawProjectPath);
@@ -387,6 +390,8 @@ public partial class Project : CommandUtils
 					SC.StageFiles(StagedFileType.UFS, Plugin.Directory, "*.uplugin", false, null, null, true, !Params.UsePak(SC.StageTargetPlatform), null, true, false);
 				}
 			}
+
+			BuildConfiguration.RelativeEnginePath = OrigEngineRelativeDir;
 			PopDir();
 
             // Game ufs (content)
