Handle cases where BootstrapPackagedGame is read only.

diff --git a/Engine/Source/Programs/AutomationTool/Win/WinPlatform.Automation.cs b/Engine/Source/Programs/AutomationTool/Win/WinPlatform.Automation.cs
--- a/Engine/Source/Programs/AutomationTool/Win/WinPlatform.Automation.cs
+++ b/Engine/Source/Programs/AutomationTool/Win/WinPlatform.Automation.cs
@@ -92,7 +92,8 @@ public abstract class BaseWinPlatform : Platform
 			InternalUtils.SafeCreateDirectory(IntermediateDir);
 
 			string IntermediateFile = CombinePaths(IntermediateDir, ExeName);
-			File.Copy(InputFile, IntermediateFile, true);
+			CommandUtils.CopyFile(InputFile, IntermediateFile);
+			CommandUtils.SetFileAttributes(IntermediateFile, ReadOnly: false);
 	
 			// currently the icon updating doesn't run under mono
 			if (UnrealBuildTool.BuildHostPlatform.Current.Platform == UnrealTargetPlatform.Win64 ||
