Just do logs for plugin spam

diff --git a/Engine/Source/Programs/AutomationTool/Scripts/CopyBuildToStagingDirectory.Automation.cs b/Engine/Source/Programs/AutomationTool/Scripts/CopyBuildToStagingDirectory.Automation.cs
--- a/Engine/Source/Programs/AutomationTool/Scripts/CopyBuildToStagingDirectory.Automation.cs
+++ b/Engine/Source/Programs/AutomationTool/Scripts/CopyBuildToStagingDirectory.Automation.cs
@@ -378,15 +378,15 @@ public partial class Project : CommandUtils
 			BuildConfiguration.RelativeEnginePath = "../../Engine/";
 
 			ProjectDescriptor Project = ProjectDescriptor.FromFile(SC.RawProjectPath);
-			LogConsole("Searching for plugins with CurrentWorkingDir: " + Directory.GetCurrentDirectory());
-			LogConsole("Searching for plugins in: " + SC.RawProjectPath);
+			Log("Searching for plugins with CurrentWorkingDir: " + Directory.GetCurrentDirectory());
+			Log("Searching for plugins in: " + SC.RawProjectPath);
 			List<PluginInfo> AvailablePlugins = Plugins.ReadAvailablePlugins(SC.RawProjectPath);
 			foreach (PluginInfo Plugin in AvailablePlugins)
 			{
-				LogConsole("Considering Plugin for Stage: " + Plugin.FileName);
+				Log("Considering Plugin for Stage: " + Plugin.FileName);
 				if (UProjectInfo.IsPluginEnabledForProject(Plugin, Project, SC.StageTargetPlatform.PlatformType))
 				{
-					LogConsole("EnabledPlugin: " + Plugin.FileName);
+					Log("EnabledPlugin: " + Plugin.FileName);
 					SC.StageFiles(StagedFileType.UFS, Plugin.Directory, "*.uplugin", false, null, null, true, !Params.UsePak(SC.StageTargetPlatform), null, true, false);
 				}
 			}
