Do iterative cooking for COTF when requested

diff --git a/Engine/Source/Programs/AutomationTool/Scripts/CookCommand.Automation.cs b/Engine/Source/Programs/AutomationTool/Scripts/CookCommand.Automation.cs
--- a/Engine/Source/Programs/AutomationTool/Scripts/CookCommand.Automation.cs
+++ b/Engine/Source/Programs/AutomationTool/Scripts/CookCommand.Automation.cs
@@ -50,10 +50,17 @@ public partial class Project : CommandUtils
 					DeleteDirectory(LogFolderOutsideOfSandbox);
 					CreateDirectory(LogFolderOutsideOfSandbox);
 				}
+
+				String COTFCommandLine = Params.RunCommandline;
+				if (Params.IterativeCooking)
+				{
+					COTFCommandLine += " -iterate";
+				}
+
 				var ServerLogFile = CombinePaths(LogFolderOutsideOfSandbox, "Server.log");
 				Platform ClientPlatformInst = Params.ClientTargetPlatformInstances[0];
 				string TargetCook = ClientPlatformInst.GetCookPlatform(false, Params.HasDedicatedServerAndClient, Params.CookFlavor);
-				ServerProcess = RunCookOnTheFlyServer(Params.RawProjectPath, Params.NoClient ? "" : ServerLogFile, TargetCook, Params.RunCommandline);
+				ServerProcess = RunCookOnTheFlyServer(Params.RawProjectPath, Params.NoClient ? "" : ServerLogFile, TargetCook, COTFCommandLine);
 
 				if (ServerProcess != null)
 				{
