Merging CL 2632258 - Mac codesign fixes

diff --git a/Engine/Source/Programs/AutomationTool/AutomationUtils/CommandUtils.cs b/Engine/Source/Programs/AutomationTool/AutomationUtils/CommandUtils.cs
--- a/Engine/Source/Programs/AutomationTool/AutomationUtils/CommandUtils.cs
+++ b/Engine/Source/Programs/AutomationTool/AutomationUtils/CommandUtils.cs
@@ -1,4 +1,4 @@
-﻿// Copyright 1998-2015 Epic Games, Inc. All Rights Reserved.
+// Copyright 1998-2015 Epic Games, Inc. All Rights Reserved.
 using System;
 using System.Collections;
 using System.Collections.Generic;
@@ -2388,7 +2388,7 @@ namespace AutomationTool
 				return;
 			}
 
-			string SignToolName = "codesign";
+			string SignToolName = "/usr/bin/codesign";
 
 			string CodeSignArgs = String.Format("-f --deep -s \"{0}\" -v \"{1}\"", "Developer ID Application", InPath);
 
diff --git a/Engine/Source/Programs/AutomationTool/Mac/MacPlatform.Automation.cs b/Engine/Source/Programs/AutomationTool/Mac/MacPlatform.Automation.cs
--- a/Engine/Source/Programs/AutomationTool/Mac/MacPlatform.Automation.cs
+++ b/Engine/Source/Programs/AutomationTool/Mac/MacPlatform.Automation.cs
@@ -433,7 +433,7 @@ public class MacPlatform : Platform
 			string AppBundlePath = "";
 			if (Exe.StartsWith(CombinePaths(SC.RuntimeProjectRootDir, "Binaries", SC.PlatformDir)))
 			{
-				AppBundlePath = CombinePaths(SC.ShortProjectName, "Binaries", SC.PlatformDir, Path.GetFileNameWithoutExtension(Exe) + ".app");
+				AppBundlePath = CombinePaths(SC.RuntimeProjectRootDir, "Binaries", SC.PlatformDir, Path.GetFileNameWithoutExtension(Exe) + ".app");
 			}
 			else if (Exe.StartsWith(CombinePaths(SC.RuntimeRootDir, "Engine/Binaries", SC.PlatformDir)))
 			{
