Fix cookonthefly steaming condition.

diff --git a/Engine/Source/Programs/AutomationTool/Scripts/RunProjectCommand.Automation.cs b/Engine/Source/Programs/AutomationTool/Scripts/RunProjectCommand.Automation.cs
--- a/Engine/Source/Programs/AutomationTool/Scripts/RunProjectCommand.Automation.cs
+++ b/Engine/Source/Programs/AutomationTool/Scripts/RunProjectCommand.Automation.cs
@@ -775,7 +775,7 @@ public partial class Project : CommandUtils
                 }
 				TempCmdLine += " ";
 
-				if (!Params.CookOnTheFlyStreaming)
+				if (Params.CookOnTheFlyStreaming)
 				{
 					TempCmdLine += "-streaming ";
 				}
