Explicitly pass the engine root into ReadAvailablePlugins().

diff --git a/Engine/Source/Programs/AutomationTool/AutomationUtils/ProjectUtils.cs b/Engine/Source/Programs/AutomationTool/AutomationUtils/ProjectUtils.cs
--- a/Engine/Source/Programs/AutomationTool/AutomationUtils/ProjectUtils.cs
+++ b/Engine/Source/Programs/AutomationTool/AutomationUtils/ProjectUtils.cs
@@ -205,7 +205,7 @@ namespace AutomationTool
 
 			// Read the project descriptor, and find all the plugins available to this project
 			ProjectDescriptor Project = ProjectDescriptor.FromFile(RawProjectPath);
-			List<PluginInfo> AvailablePlugins = Plugins.ReadAvailablePlugins(RawProjectPath);
+			List<PluginInfo> AvailablePlugins = Plugins.ReadAvailablePlugins(BuildConfiguration.RelativeEnginePath, RawProjectPath);
 
 			// check the target platforms for any differences in build settings or additional plugins
 			bool RetVal = false;
diff --git a/Engine/Source/Programs/AutomationTool/Scripts/CopyBuildToStagingDirectory.Automation.cs b/Engine/Source/Programs/AutomationTool/Scripts/CopyBuildToStagingDirectory.Automation.cs
--- a/Engine/Source/Programs/AutomationTool/Scripts/CopyBuildToStagingDirectory.Automation.cs
+++ b/Engine/Source/Programs/AutomationTool/Scripts/CopyBuildToStagingDirectory.Automation.cs
@@ -373,14 +373,10 @@ public partial class Project : CommandUtils
             }
 			
 			// Engine & Game Plugins. Push the Engine/Source working directory so UBT code has a correct RelativeEnginePath in ReadAvailablePlugins
-			PushDir(CombinePaths(SC.LocalRoot, "Engine", "Source"));
-			String OrigEngineRelativeDir = BuildConfiguration.RelativeEnginePath;
-			BuildConfiguration.RelativeEnginePath = "../../Engine/";
-
 			ProjectDescriptor Project = ProjectDescriptor.FromFile(SC.RawProjectPath);
 			Log("Searching for plugins with CurrentWorkingDir: " + Directory.GetCurrentDirectory());
 			Log("Searching for plugins in: " + SC.RawProjectPath);
-			List<PluginInfo> AvailablePlugins = Plugins.ReadAvailablePlugins(SC.RawProjectPath);
+			List<PluginInfo> AvailablePlugins = Plugins.ReadAvailablePlugins(CombinePaths(SC.LocalRoot, "Engine"), SC.RawProjectPath);
 			foreach (PluginInfo Plugin in AvailablePlugins)
 			{
 				Log("Considering Plugin for Stage: " + Plugin.FileName);
@@ -391,9 +387,6 @@ public partial class Project : CommandUtils
 				}
 			}
 
-			BuildConfiguration.RelativeEnginePath = OrigEngineRelativeDir;
-			PopDir();
-
             // Game ufs (content)
 
             SC.StageFiles(StagedFileType.UFS, CombinePaths(SC.ProjectRoot), "*.uproject", false, null, CombinePaths(SC.RelativeProjectRootForStage), true, !Params.UsePak(SC.StageTargetPlatform));
diff --git a/Engine/Source/Programs/UnrealBuildTool/System/Plugins.cs b/Engine/Source/Programs/UnrealBuildTool/System/Plugins.cs
--- a/Engine/Source/Programs/UnrealBuildTool/System/Plugins.cs
+++ b/Engine/Source/Programs/UnrealBuildTool/System/Plugins.cs
@@ -55,19 +55,17 @@ namespace UnrealBuildTool
 		/// <summary>
 		/// Read all the plugins available to a given project
 		/// </summary>
+		/// <param name="EngineDir">Path to the engine directory</param>
 		/// <param name="ProjectFileName">Path to the project file (or null)</param>
 		/// <returns>Sequence of PluginInfo objects, one for each discovered plugin</returns>
-		public static List<PluginInfo> ReadAvailablePlugins(string ProjectFileName)
+		public static List<PluginInfo> ReadAvailablePlugins(string EngineDir, string ProjectFileName)
 		{
 			List<PluginInfo> Plugins = new List<PluginInfo>();
 
 			// Read all the engine plugins
-			
-			string EnginePluginsDir = Path.Combine(BuildConfiguration.RelativeEnginePath, "Plugins");
-			Log.TraceVerbose("EngineDir for plugin search: " + EnginePluginsDir);
+			string EnginePluginsDir = Path.Combine(EngineDir, "Plugins");
 			foreach(string PluginFileName in EnumeratePlugins(EnginePluginsDir))
 			{
-				Log.TraceVerbose("Found Engine Plugin: " + PluginFileName);
 				PluginInfo Plugin = new PluginInfo(PluginFileName, PluginLoadedFrom.Engine);
 				Plugins.Add(Plugin);
 			}
