Rename LegacyNode -> LegacyBuildNode to distinguish between other types of legacy nodes.

diff --git a/Engine/Source/Programs/AutomationTool/GUBP/GUBP.Automation.csproj b/Engine/Source/Programs/AutomationTool/GUBP/GUBP.Automation.csproj
--- a/Engine/Source/Programs/AutomationTool/GUBP/GUBP.Automation.csproj
+++ b/Engine/Source/Programs/AutomationTool/GUBP/GUBP.Automation.csproj
@@ -64,7 +64,6 @@
     <Compile Include="LegacyNodes.cs" />
     <Compile Include="Nodes\AggregateNode.cs" />
     <Compile Include="Nodes\BuildNode.cs" />
-    <Compile Include="Nodes\LegacyNode.cs" />
     <Compile Include="Nodes\TriggerNode.cs" />
     <Compile Include="TempStorage.cs">
       <SubType>Code</SubType>
diff --git a/Engine/Source/Programs/AutomationTool/GUBP/LegacyNodes.cs b/Engine/Source/Programs/AutomationTool/GUBP/LegacyNodes.cs
--- a/Engine/Source/Programs/AutomationTool/GUBP/LegacyNodes.cs
+++ b/Engine/Source/Programs/AutomationTool/GUBP/LegacyNodes.cs
@@ -39,7 +39,7 @@ partial class GUBP
         }
 		public virtual BuildNode GetBuildNode()
 		{
-			return new LegacyNode(this);
+			return new LegacyBuildNode(this);
 		}
         public virtual void DoBuild(GUBP bp)
         {
diff --git a/Engine/Source/Programs/AutomationTool/GUBP/Nodes/BuildNode.cs b/Engine/Source/Programs/AutomationTool/GUBP/Nodes/BuildNode.cs
--- a/Engine/Source/Programs/AutomationTool/GUBP/Nodes/BuildNode.cs
+++ b/Engine/Source/Programs/AutomationTool/GUBP/Nodes/BuildNode.cs
@@ -52,4 +52,24 @@ namespace AutomationTool
 			return Name;
 		}
 	}
+
+	public class LegacyBuildNode : BuildNode
+	{
+		public LegacyBuildNode(GUBP.GUBPNode InNode)
+		{
+			Name = InNode.GetFullName();
+			Node = InNode;
+			AgentRequirements = Node.ECAgentString();
+			AgentSharingGroup = Node.AgentSharingGroup;
+			AgentMemoryRequirement = Node.AgentMemoryRequirement();
+			TimeoutInMinutes = Node.TimeoutInMinutes();
+			SendSuccessEmail = Node.SendSuccessEmail();
+			Priority = Node.Priority();
+		}
+
+		public override bool IsSticky
+		{
+			get { return Node.IsSticky(); }
+		}
+	}
 }
diff --git a/Engine/Source/Programs/AutomationTool/GUBP/Nodes/LegacyNode.cs b/Engine/Source/Programs/AutomationTool/GUBP/Nodes/LegacyNode.cs
deleted file mode 100644
diff --git a/Engine/Source/Programs/AutomationTool/GUBP/Nodes/TriggerNode.cs b/Engine/Source/Programs/AutomationTool/GUBP/Nodes/TriggerNode.cs
--- a/Engine/Source/Programs/AutomationTool/GUBP/Nodes/TriggerNode.cs
+++ b/Engine/Source/Programs/AutomationTool/GUBP/Nodes/TriggerNode.cs
@@ -6,7 +6,7 @@ using System.Text;
 
 namespace AutomationTool
 {
-	public class TriggerNode : LegacyNode
+	public class TriggerNode : LegacyBuildNode
 	{
         public bool IsTriggered;
 		public bool RequiresRecursiveWorkflow;
