Separate legacy aggregate node functionality from the new implementation.

diff --git a/Engine/Source/Programs/AutomationTool/GUBP/ElectricCommander.cs b/Engine/Source/Programs/AutomationTool/GUBP/ElectricCommander.cs
--- a/Engine/Source/Programs/AutomationTool/GUBP/ElectricCommander.cs
+++ b/Engine/Source/Programs/AutomationTool/GUBP/ElectricCommander.cs
@@ -649,7 +649,7 @@ namespace AutomationTool
 			List<AggregateNode> SeparatePromotables = new List<AggregateNode>();
 			foreach (AggregateNode NodeToDo in NodesToDo)
 			{
-				if (NodeToDo.Node.IsSeparatePromotable())
+				if (NodeToDo.IsSeparatePromotable)
 				{
 					SeparatePromotables.Add(NodeToDo);
 				}
diff --git a/Engine/Source/Programs/AutomationTool/GUBP/GUBP.cs b/Engine/Source/Programs/AutomationTool/GUBP/GUBP.cs
--- a/Engine/Source/Programs/AutomationTool/GUBP/GUBP.cs
+++ b/Engine/Source/Programs/AutomationTool/GUBP/GUBP.cs
@@ -497,7 +497,7 @@ public partial class GUBP : BuildCommand
 			foreach (AggregateNode Aggregate in MatchingAggregates.OrderBy(x => x.Name))
 			{
 				StringBuilder Note = new StringBuilder("    " + Aggregate.Name);
-				if (Aggregate.Node.IsPromotableAggregate())
+				if (Aggregate.IsPromotableAggregate)
 				{
 					Note.Append(" (promotable)");
 				}
@@ -1188,7 +1188,7 @@ public partial class GUBP : BuildCommand
 			Aggregate.Dependencies = new BuildNode[0];
 
 			HashSet<BuildNode> Dependencies = new HashSet<BuildNode>();
-			foreach (string DependencyName in Aggregate.Node.Dependencies)
+			foreach (string DependencyName in Aggregate.DependencyNames)
 			{
 				AggregateNode AggregateDependency;
 				if(AggregateNameToInfo.TryGetValue(DependencyName, out AggregateDependency))
diff --git a/Engine/Source/Programs/AutomationTool/GUBP/LegacyBranchSetup.cs b/Engine/Source/Programs/AutomationTool/GUBP/LegacyBranchSetup.cs
--- a/Engine/Source/Programs/AutomationTool/GUBP/LegacyBranchSetup.cs
+++ b/Engine/Source/Programs/AutomationTool/GUBP/LegacyBranchSetup.cs
@@ -1217,7 +1217,7 @@ partial class GUBP
 		AllNodes.AddRange(BranchConfig.GUBPNodes.Values.Select(x => x.GetBuildNode()));
 
 		AllAggregates = new List<AggregateNode>();
-		AllAggregates.AddRange(BranchConfig.GUBPAggregates.Values.Select(x => new AggregateNode{ Name = x.GetFullName(), Node = x }));
+		AllAggregates.AddRange(BranchConfig.GUBPAggregates.Values.Select(x => new LegacyAggregateNode(x)));
 
 		// Calculate the frequencies for each node
 		FindFrequenciesForNodes(BranchConfig, AllNodes, FrequencyOverrides);
diff --git a/Engine/Source/Programs/AutomationTool/GUBP/Nodes/AggregateNode.cs b/Engine/Source/Programs/AutomationTool/GUBP/Nodes/AggregateNode.cs
--- a/Engine/Source/Programs/AutomationTool/GUBP/Nodes/AggregateNode.cs
+++ b/Engine/Source/Programs/AutomationTool/GUBP/Nodes/AggregateNode.cs
@@ -7,10 +7,54 @@ using System.Text;
 namespace AutomationTool
 {
 	[DebuggerDisplay("{Name}")]
-	public class AggregateNode
+	public abstract class AggregateNode
 	{
 		public string Name;
-		public GUBP.GUBPAggregateNode Node;
 		public BuildNode[] Dependencies;
+
+		public AggregateNode(string InName)
+		{
+			Name = InName;
+		}
+
+		public abstract IEnumerable<string> DependencyNames
+		{
+			get;
+		}
+
+		public abstract bool IsPromotableAggregate
+		{
+			get;
+		}
+
+		public abstract bool IsSeparatePromotable
+		{
+			get;
+		}
+	}
+
+	public class LegacyAggregateNode : AggregateNode
+	{
+		public GUBP.GUBPAggregateNode Node;
+
+		public LegacyAggregateNode(GUBP.GUBPAggregateNode InNode) : base(InNode.GetFullName())
+		{
+			Node = InNode;
+		}
+
+		public override IEnumerable<string> DependencyNames
+		{
+			get { return Node.Dependencies; }
+		}
+
+		public override bool IsPromotableAggregate
+		{
+			get { return Node.IsPromotableAggregate(); }
+		}
+
+		public override bool IsSeparatePromotable
+		{
+			get { return Node.IsSeparatePromotable(); }
+		}
 	}
 }
diff --git a/Engine/Source/Programs/AutomationTool/GUBP/Nodes/BuildNode.cs b/Engine/Source/Programs/AutomationTool/GUBP/Nodes/BuildNode.cs
--- a/Engine/Source/Programs/AutomationTool/GUBP/Nodes/BuildNode.cs
+++ b/Engine/Source/Programs/AutomationTool/GUBP/Nodes/BuildNode.cs
@@ -51,6 +51,18 @@ namespace AutomationTool
 			System.Diagnostics.Trace.TraceWarning("Implicit conversion from NodeInfo to string\n{0}", Environment.StackTrace);
 			return Name;
 		}
+
+		// Legacy stuff that should probably be removed
+
+		public abstract bool IsTest
+		{
+			get;
+		}
+
+		public abstract string DisplayGroupName
+		{
+			get;
+		}
 	}
 
 	public class LegacyBuildNode : BuildNode
@@ -71,5 +83,15 @@ namespace AutomationTool
 		{
 			get { return Node.IsSticky(); }
 		}
+
+		public override bool IsTest
+		{
+			get { return Node.IsTest(); }
+		}
+
+		public override string DisplayGroupName
+		{
+			get { return Node.GetDisplayGroupName(); }
+		}
 	}
 }
