Simplify logic around handling agent sharing groups.

diff --git a/Engine/Source/Programs/AutomationTool/GUBP/ElectricCommander.cs b/Engine/Source/Programs/AutomationTool/GUBP/ElectricCommander.cs
--- a/Engine/Source/Programs/AutomationTool/GUBP/ElectricCommander.cs
+++ b/Engine/Source/Programs/AutomationTool/GUBP/ElectricCommander.cs
@@ -390,27 +390,12 @@ namespace AutomationTool
 						string PreCondition = GetPreConditionForNode(OrdereredToDo, ParentPath, bHasNoop, AgentGroupChains, StickyChain, NodeToDo, UncompletedEcDeps);
 						string RunCondition = GetRunConditionForNode(UncompletedEcDeps, ParentPath);
 
-						string MyAgentGroup = NodeToDo.AgentSharingGroup;
-						bool bDoNestedJobstep = false;
-						bool bDoFirstNestedJobstep = false;
-
-						string NodeParentPath = ParentPath;
-						if (MyAgentGroup != "")
+						if (NodeToDo.AgentSharingGroup != "")
 						{
-							bDoNestedJobstep = true;
-							NodeParentPath = ParentPath + "/jobSteps[" + MyAgentGroup + "]";
+							string MyAgentGroup = NodeToDo.AgentSharingGroup;
 
 							List<BuildNode> MyChain = AgentGroupChains[MyAgentGroup];
-							int MyIndex = MyChain.IndexOf(NodeToDo);
-							if (MyIndex <= 0)
-							{
-								bDoFirstNestedJobstep = bDoNestedJobstep;
-							}
-						}
-
-						if (bDoNestedJobstep)
-						{
-							if (bDoFirstNestedJobstep)
+							if(MyChain.IndexOf(NodeToDo) <= 0)
 							{
 								{
 									string NestArgs = String.Format("$batch->createJobStep({{parentPath => '{0}', jobStepName => '{1}', parallel => '1'",
@@ -452,7 +437,8 @@ namespace AutomationTool
 									PreCondition = PreCondition + ") true;]\"";
 								}
 							}
-							Args = Args.Replace(String.Format("parentPath => '{0}'", ParentPath), String.Format("parentPath => '{0}'", NodeParentPath));
+							Args = Args.Replace(String.Format("parentPath => '{0}'", ParentPath), String.Format("parentPath => '{0}/jobSteps[{1}]'", ParentPath, MyAgentGroup));
+
 							Args = Args.Replace("UAT_Node_Parallel_AgentShare", "UAT_Node_Parallel_AgentShare3");
 						}
 
@@ -464,36 +450,13 @@ namespace AutomationTool
 						{
 							Args = Args + ", condition => " + RunCondition;
 						}
-		#if false
-							// this doesn't work because it includes precondition time
-							if (GUBPNodes[NodeToDo].TimeoutInMinutes() > 0)
-							{
-								Args = Args + String.Format(" --timeLimitUnits minutes --timeLimit {0}", GUBPNodes[NodeToDo].TimeoutInMinutes());
-							}
-		#endif
+
 						if (Sticky && NodeToDo == LastSticky)
 						{
 							Args = Args + ", releaseMode => 'release'";
 						}
 						Args = Args + "});";
 						StepList.Add(Args);
-
-						if (MyAgentGroup != "" && !bDoNestedJobstep)
-						{
-							List<BuildNode> MyChain = AgentGroupChains[MyAgentGroup];
-							int MyIndex = MyChain.IndexOf(NodeToDo);
-							if (MyIndex == MyChain.Count - 1)
-							{
-								string RelPreCondition = "\"\\$\" . \"[/javascript if(";
-								// this runs "parallel", but we a precondition to serialize it
-								RelPreCondition = RelPreCondition + "getProperty('" + ParentPath + "/jobSteps[" + NodeToDo.Name + "]/status') == 'completed'";
-								RelPreCondition = RelPreCondition + ") true;]\"";
-								// we need to release the resource
-								string RelArgs = String.Format("{0}, subprocedure => 'GUBP_Release_AgentShare', parallel => '1', jobStepName => 'Release_{1}', actualParameter => [{{actualParameterName => 'AgentSharingGroup', valued => '{2}'}}], releaseMode => 'release', precondition => '{3}'",
-									BaseArgs, MyAgentGroup, MyAgentGroup, RelPreCondition);
-								StepList.Add(RelArgs);
-							}
-						}
 					}
 				}
 				WriteECPerl(StepList);
