Add a wrapper class around EC job steps, making it easier to construct and manipulate the perl script that gets output.

diff --git a/Engine/Source/Programs/AutomationTool/GUBP/ElectricCommander.cs b/Engine/Source/Programs/AutomationTool/GUBP/ElectricCommander.cs
--- a/Engine/Source/Programs/AutomationTool/GUBP/ElectricCommander.cs
+++ b/Engine/Source/Programs/AutomationTool/GUBP/ElectricCommander.cs
@@ -1,5 +1,6 @@
 ﻿using System;
 using System.Collections.Generic;
+using System.Diagnostics;
 using System.IO;
 using System.Linq;
 using System.Text;
@@ -35,11 +36,128 @@ namespace AutomationTool
 			}
 		}
 
-		static void WriteECPerl(List<string> Args)
+		enum JobStepExclusiveMode
 		{
-			Args.Add("$batch->submit();");
+			None,
+			Call,
+		}
+
+		enum JobStepReleaseMode
+		{
+			Keep,
+			Release,
+		}
+
+		enum JobStepNamePlacement
+		{
+			Default,
+			BeforeParallel,
+			BeforeProcedure,
+		}
+
+		[DebuggerDisplay("{Name}")]
+		class JobStep
+		{
+			public string ParentPath;
+			public string Name;
+			public string SubProcedure;
+			public bool Parallel;
+			public string PreCondition;
+			public string RunCondition;
+			public JobStepExclusiveMode Exclusive;
+			public string ResourceName;
+			public List<Tuple<string, string>> Parameters;
+			public JobStepReleaseMode ReleaseMode;
+			public JobStepNamePlacement NamePlacement;
+
+			public JobStep(string InParentPath, string InName, string InSubProcedure, bool InParallel, string InPreCondition, string InRunCondition, List<Tuple<string, string>> InParameters, JobStepReleaseMode InReleaseMode, JobStepExclusiveMode InExclusive = JobStepExclusiveMode.None, string InResourceName = null, JobStepNamePlacement InNamePlacement = JobStepNamePlacement.Default)
+			{
+				ParentPath = InParentPath;
+				Name = InName;
+				SubProcedure = InSubProcedure;
+				Parallel = InParallel;
+				PreCondition = InPreCondition;
+				RunCondition = InRunCondition;
+				Parameters = InParameters;
+				ReleaseMode = InReleaseMode;
+				Exclusive = InExclusive;
+				ResourceName = InResourceName;
+				NamePlacement = InNamePlacement;
+			}
+
+			public string GetCreationStatement()
+			{
+				StringBuilder Args = new StringBuilder();
+				Args.Append("$batch->createJobStep({");
+				Args.AppendFormat("parentPath => '{0}'", ParentPath);
+				if (!String.IsNullOrEmpty(Name) && NamePlacement == JobStepNamePlacement.BeforeProcedure)
+				{
+					Args.AppendFormat(", jobStepName => '{0}'", Name);
+				}
+				if(!String.IsNullOrEmpty(SubProcedure))
+				{
+					Args.AppendFormat(", subprocedure => '{0}'", SubProcedure);
+				}
+				if (!String.IsNullOrEmpty(Name) && NamePlacement == JobStepNamePlacement.BeforeParallel)
+				{
+					Args.AppendFormat(", jobStepName => '{0}'", Name);
+				}
+				Args.AppendFormat(", parallel => '{0}'", Parallel ? 1 : 0);
+				if (!String.IsNullOrEmpty(Name) && NamePlacement == JobStepNamePlacement.Default)
+				{
+					Args.AppendFormat(", jobStepName => '{0}'", Name);
+				}
+				if (Exclusive != JobStepExclusiveMode.None)
+				{
+					Args.AppendFormat(", exclusiveMode => '{0}'", Exclusive.ToString().ToLower());
+				}
+				if (!String.IsNullOrEmpty(ResourceName))
+				{
+					Args.AppendFormat(", resourceName => '{0}'", ResourceName);
+				}
+				if (Parameters != null)
+				{
+					Args.AppendFormat(", actualParameter => [{0}]", String.Join(", ", Parameters.Select(x => String.Format("{{actualParameterName => '{0}', value => '{1}'}}", x.Item1, x.Item2))));
+				}
+				if (!String.IsNullOrEmpty(PreCondition))
+				{
+					Args.AppendFormat(", precondition => {0}", PreCondition);
+				}
+				if (!String.IsNullOrEmpty(RunCondition))
+				{
+					Args.AppendFormat(", condition => {0}", RunCondition);
+				}
+				if (ReleaseMode == JobStepReleaseMode.Release)
+				{
+					Args.Append(", releaseMode => 'release'");
+				}
+				Args.Append("});");
+				return Args.ToString();
+			}
+
+			public string GetCompletedCondition()
+			{
+				return String.Format("\"\\$\" . \"[/javascript if(getProperty('{0}/jobSteps[{1}]/status') == 'completed') true;]\"", ParentPath, Name);
+			}
+		}
+
+		static void WriteECPerl(List<JobStep> Steps)
+		{
+			List<string> Lines = new List<string>();
+			Lines.Add("use strict;");
+			Lines.Add("use diagnostics;");
+			Lines.Add("use ElectricCommander();");
+			Lines.Add("my $ec = new ElectricCommander;");
+			Lines.Add("$ec->setTimeout(600);");
+			Lines.Add("my $batch = $ec->newBatch(\"serial\");");
+			foreach (JobStep Step in Steps)
+			{
+				Lines.Add(Step.GetCreationStatement());
+			}
+			Lines.Add("$batch->submit();");
+
 			string ECPerlFile = CommandUtils.CombinePaths(CommandUtils.CmdEnv.LogFolder, "jobsteps.pl");
-			CommandUtils.WriteAllLines_NoExceptions(ECPerlFile, Args.ToArray());
+			CommandUtils.WriteAllLines_NoExceptions(ECPerlFile, Lines.ToArray());
 		}
 
 		string GetPropertyFromStep(string PropertyPath)
@@ -211,13 +329,6 @@ namespace AutomationTool
 			BuildNode LastSticky = null;
 			bool HitNonSticky = false;
 			bool bHaveECNodes = false;
-			List<string> StepList = new List<string>();
-			StepList.Add("use strict;");
-			StepList.Add("use diagnostics;");
-			StepList.Add("use ElectricCommander();");
-			StepList.Add("my $ec = new ElectricCommander;");
-			StepList.Add("$ec->setTimeout(600);");
-			StepList.Add("my $batch = $ec->newBatch(\"serial\");");
 			// sticky nodes are ones that we run on the main agent. We run then first and they must not be intermixed with parallel jobs
 			foreach (BuildNode NodeToDo in OrdereredToDo)
 			{
@@ -239,6 +350,7 @@ namespace AutomationTool
 				}
 			}
 
+			List<JobStep> Steps = new List<JobStep>();
 			using(CommandUtils.TelemetryStopwatch PerlOutputStopwatch = new CommandUtils.TelemetryStopwatch("PerlOutput"))
 			{
 				string ParentPath = Command.ParseParamValue("ParentPath");
@@ -248,8 +360,11 @@ namespace AutomationTool
 				if (LastSticky == null && bHaveECNodes)
 				{
 					// if we don't have any sticky nodes and we have other nodes, we run a fake noop just to release the resource 
-					string Args = String.Format("{0}, subprocedure => 'GUBP_UAT_Node', parallel => '0', jobStepName => 'Noop', actualParameter => [{{actualParameterName => 'NodeName', value => 'Noop'}}, {{actualParameterName => 'Sticky', value =>'1' }}], releaseMode => 'release'}});", BaseArgs);
-					StepList.Add(Args);
+					List<Tuple<string, string>> NoopParameters = new List<Tuple<string, string>>();
+					NoopParameters.Add(new Tuple<string, string>("NodeName", "Noop"));
+					NoopParameters.Add(new Tuple<string, string>("Sticky", "1"));
+					Steps.Add(new JobStep(ParentPath, "Noop", "GUBP_UAT_Node", false, null, null, NoopParameters, JobStepReleaseMode.Release));
+
 					bHasNoop = true;
 				}
 
@@ -329,18 +444,28 @@ namespace AutomationTool
 						{
 							if (NodeToDo.IsParallelAgentShareEditor)
 							{
-								Procedure = "GUBP_UAT_Node_Parallel_AgentShare_Editor";
+								if(NodeToDo.AgentSharingGroup == "")
+								{
+									Procedure = "GUBP_UAT_Node_Parallel_AgentShare_Editor";
+								}
+								else
+								{
+									Procedure = "GUBP_UAT_Node_Parallel_AgentShare3_Editor";
+								}
 							}
 							else
 							{
-								Procedure = "GUBP" + ProcedureInfix + "_UAT_Node";
-								if (!NodeToDo.IsSticky)
+								if (NodeToDo.IsSticky)
 								{
-									Procedure += "_Parallel";
+									Procedure = "GUBP" + ProcedureInfix + "_UAT_Node";
 								}
-								if (NodeToDo.AgentSharingGroup != "")
+								else if (NodeToDo.AgentSharingGroup == "")
 								{
-									Procedure += "_AgentShare";
+									Procedure = String.Format("GUBP{0}_UAT_Node_Parallel", ProcedureInfix);
+								}
+								else
+								{
+									Procedure = String.Format("GUBP{0}_UAT_Node_Parallel_AgentShare3", ProcedureInfix);
 								}
 							}
 							if (NodeToDo.IsSticky && NodeToDo == LastSticky)
@@ -369,97 +494,56 @@ namespace AutomationTool
 							}
 						}
 
-						string ActualParameterArgs = String.Join(", ", Parameters.Select(x => String.Format("{{actualParameterName => '{0}', value => '{1}'}}", x.Item1, x.Item2)));
-						string Args = String.Format("{0}, subprocedure => '{1}', parallel => '{2}', jobStepName => '{3}', actualParameter => [{4}]", BaseArgs, Procedure, DoParallel? 1 : 0, NodeToDo.Name, ActualParameterArgs);
-						
 						List<BuildNode> UncompletedEcDeps = new List<BuildNode>();
+						foreach (BuildNode Dep in NodeToDo.AllDirectDependencies)
 						{
-							foreach (BuildNode Dep in NodeToDo.AllDirectDependencies)
+							if (!Dep.IsComplete && OrdereredToDo.Contains(Dep)) // if something is already finished, we don't put it into EC
 							{
-								if (!Dep.IsComplete && OrdereredToDo.Contains(Dep)) // if something is already finished, we don't put it into EC
+								if (OrdereredToDo.IndexOf(Dep) > OrdereredToDo.IndexOf(NodeToDo))
 								{
-									if (OrdereredToDo.IndexOf(Dep) > OrdereredToDo.IndexOf(NodeToDo))
-									{
-										throw new AutomationException("Topological sort error, node {0} has a dependency of {1} which sorted after it.", NodeToDo.Name, Dep.Name);
-									}
-									UncompletedEcDeps.Add(Dep);
+									throw new AutomationException("Topological sort error, node {0} has a dependency of {1} which sorted after it.", NodeToDo.Name, Dep.Name);
 								}
+								UncompletedEcDeps.Add(Dep);
 							}
 						}
 
 						string PreCondition = GetPreConditionForNode(OrdereredToDo, ParentPath, bHasNoop, AgentGroupChains, StickyChain, NodeToDo, UncompletedEcDeps);
 						string RunCondition = GetRunConditionForNode(UncompletedEcDeps, ParentPath);
 
+						// Create the jobs to setup the agent sharing group if necessary
+						string NodeParentPath = ParentPath;
 						if (NodeToDo.AgentSharingGroup != "")
 						{
-							string MyAgentGroup = NodeToDo.AgentSharingGroup;
+							NodeParentPath = String.Format("{0}/jobSteps[{1}]", NodeParentPath, NodeToDo.AgentSharingGroup);
 
-							List<BuildNode> MyChain = AgentGroupChains[MyAgentGroup];
+							List<BuildNode> MyChain = AgentGroupChains[NodeToDo.AgentSharingGroup];
 							if(MyChain.IndexOf(NodeToDo) <= 0)
 							{
-								{
-									string NestArgs = String.Format("$batch->createJobStep({{parentPath => '{0}', jobStepName => '{1}', parallel => '1'",
-										ParentPath, MyAgentGroup);
-									if (!String.IsNullOrEmpty(PreCondition))
-									{
-										NestArgs = NestArgs + ", precondition => " + PreCondition;
-									}
-									NestArgs = NestArgs + "});";
-									StepList.Add(NestArgs);
-								}
-								{
-									string NestArgs = String.Format("$batch->createJobStep({{parentPath => '{0}/jobSteps[{1}]', jobStepName => '{2}_GetPool', subprocedure => 'GUBP{3}_AgentShare_GetPool', parallel => '1', actualParameter => [{{actualParameterName => 'AgentSharingGroup', value => '{4}'}}, {{actualParameterName => 'NodeName', value => '{5}'}}]",
-										ParentPath, MyAgentGroup, MyAgentGroup, ProcedureInfix, MyAgentGroup, NodeToDo.Name);
-									if (!String.IsNullOrEmpty(PreCondition))
-									{
-										NestArgs = NestArgs + ", precondition => " + PreCondition;
-									}
-									NestArgs = NestArgs + "});";
-									StepList.Add(NestArgs);
-								}
-								{
-									string NestArgs = String.Format("$batch->createJobStep({{parentPath => '{0}/jobSteps[{1}]', jobStepName => '{2}_GetAgent', subprocedure => 'GUBP{3}_AgentShare_GetAgent', parallel => '1', exclusiveMode => 'call', resourceName => '{4}', actualParameter => [{{actualParameterName => 'AgentSharingGroup', value => '{5}'}}, {{actualParameterName => 'NodeName', value=> '{6}'}}]",
-										ParentPath, MyAgentGroup, MyAgentGroup, ProcedureInfix,
-										String.Format("$[/myJob/jobSteps[{0}]/ResourcePool]", MyAgentGroup),
-										MyAgentGroup, NodeToDo.Name);
-									{
-										NestArgs = NestArgs + ", precondition  => ";
-										NestArgs = NestArgs + "\"\\$\" . \"[/javascript if(";
-										NestArgs = NestArgs + "getProperty('" + ParentPath + "/jobSteps[" + MyAgentGroup + "]/jobSteps[" + MyAgentGroup + "_GetPool]/status') == 'completed'";
-										NestArgs = NestArgs + ") true;]\"";
-									}
-									NestArgs = NestArgs + "});";
-									StepList.Add(NestArgs);
-								}
-								{
-									PreCondition = "\"\\$\" . \"[/javascript if(";
-									PreCondition = PreCondition + "getProperty('" + ParentPath + "/jobSteps[" + MyAgentGroup + "]/jobSteps[" + MyAgentGroup + "_GetAgent]/status') == 'completed'";
-									PreCondition = PreCondition + ") true;]\"";
-								}
+								// Create the parent job step for this group
+								Steps.Add(new JobStep(ParentPath, NodeToDo.AgentSharingGroup, null, true, PreCondition, null, null, JobStepReleaseMode.Keep, InNamePlacement: JobStepNamePlacement.BeforeParallel));
+
+								// Get the resource pool
+								List<Tuple<string, string>> GetPoolParameters = new List<Tuple<string,string>>();
+								GetPoolParameters.Add(new Tuple<string, string>("AgentSharingGroup", NodeToDo.AgentSharingGroup));
+								GetPoolParameters.Add(new Tuple<string,string>("NodeName", NodeToDo.Name));
+								Steps.Add(new JobStep(NodeParentPath, String.Format("{0}_GetPool", NodeToDo.AgentSharingGroup), String.Format("GUBP{0}_AgentShare_GetPool", ProcedureInfix), true, PreCondition, null, GetPoolParameters, JobStepReleaseMode.Keep, InNamePlacement: JobStepNamePlacement.BeforeProcedure));
+
+								// Get the agent for this sharing group
+								List<Tuple<string, string>> GetAgentParameters = new List<Tuple<string, string>>();
+								GetAgentParameters.Add(new Tuple<string, string>("AgentSharingGroup", NodeToDo.AgentSharingGroup));
+								GetAgentParameters.Add(new Tuple<string, string>("NodeName", NodeToDo.Name));
+								Steps.Add(new JobStep(NodeParentPath, String.Format("{0}_GetAgent", NodeToDo.AgentSharingGroup), String.Format("GUBP{0}_AgentShare_GetAgent", ProcedureInfix), true, Steps.Last().GetCompletedCondition(), null, GetAgentParameters, JobStepReleaseMode.Keep, InExclusive: JobStepExclusiveMode.Call, InResourceName: String.Format("$[/myJob/jobSteps[{0}]/ResourcePool]", NodeToDo.AgentSharingGroup), InNamePlacement: JobStepNamePlacement.BeforeProcedure));
+
+								// Also add the agent being set up to the existing precondition
+								PreCondition = Steps.Last().GetCompletedCondition();
 							}
-							Args = Args.Replace(String.Format("parentPath => '{0}'", ParentPath), String.Format("parentPath => '{0}/jobSteps[{1}]'", ParentPath, MyAgentGroup));
-
-							Args = Args.Replace("UAT_Node_Parallel_AgentShare", "UAT_Node_Parallel_AgentShare3");
-						}
-
-						if (!String.IsNullOrEmpty(PreCondition))
-						{
-							Args = Args + ", precondition => " + PreCondition;
-						}
-						if (!String.IsNullOrEmpty(RunCondition))
-						{
-							Args = Args + ", condition => " + RunCondition;
 						}
 
-						if (Sticky && NodeToDo == LastSticky)
-						{
-							Args = Args + ", releaseMode => 'release'";
-						}
-						Args = Args + "});";
-						StepList.Add(Args);
+						// Build the job step for this node
+						Steps.Add(new JobStep(NodeParentPath, NodeToDo.Name, Procedure, DoParallel, PreCondition, RunCondition, Parameters, (NodeToDo.IsSticky && NodeToDo == LastSticky) ? JobStepReleaseMode.Release : JobStepReleaseMode.Keep));
 					}
 				}
-				WriteECPerl(StepList);
+				WriteECPerl(Steps);
 			}
 			bool bHasTests = OrdereredToDo.Any(x => x.Node.IsTest());
 			RunECTool(String.Format("setProperty \"/myWorkflow/HasTests\" \"{0}\"", bHasTests));
