Fix unused variable warning on Mac.

diff --git a/Engine/Source/Programs/AutomationTool/GUBP/ElectricCommander.cs b/Engine/Source/Programs/AutomationTool/GUBP/ElectricCommander.cs
--- a/Engine/Source/Programs/AutomationTool/GUBP/ElectricCommander.cs
+++ b/Engine/Source/Programs/AutomationTool/GUBP/ElectricCommander.cs
@@ -354,7 +354,6 @@ namespace AutomationTool
 			using(CommandUtils.TelemetryStopwatch PerlOutputStopwatch = new CommandUtils.TelemetryStopwatch("PerlOutput"))
 			{
 				string ParentPath = Command.ParseParamValue("ParentPath");
-				string BaseArgs = String.Format("$batch->createJobStep({{parentPath => '{0}'", ParentPath);
 
 				bool bHasNoop = false;
 				if (LastSticky == null && bHaveECNodes)
