More cleanup of EC procedure setup. Separate out trigger handling from regular nodes.

diff --git a/Engine/Source/Programs/AutomationTool/GUBP/ElectricCommander.cs b/Engine/Source/Programs/AutomationTool/GUBP/ElectricCommander.cs
--- a/Engine/Source/Programs/AutomationTool/GUBP/ElectricCommander.cs
+++ b/Engine/Source/Programs/AutomationTool/GUBP/ElectricCommander.cs
@@ -66,11 +66,11 @@ namespace AutomationTool
 			public string RunCondition;
 			public JobStepExclusiveMode Exclusive;
 			public string ResourceName;
-			public List<Tuple<string, string>> Parameters;
+			public Dictionary<string, string> ActualParameters = new Dictionary<string, string>();
 			public JobStepReleaseMode ReleaseMode;
 			public JobStepNamePlacement NamePlacement;
 
-			public JobStep(string InParentPath, string InName, string InSubProcedure, bool InParallel, string InPreCondition, string InRunCondition, List<Tuple<string, string>> InParameters, JobStepReleaseMode InReleaseMode, JobStepExclusiveMode InExclusive = JobStepExclusiveMode.None, string InResourceName = null, JobStepNamePlacement InNamePlacement = JobStepNamePlacement.Default)
+			public JobStep(string InParentPath, string InName, string InSubProcedure, bool InParallel, string InPreCondition, string InRunCondition, JobStepReleaseMode InReleaseMode, JobStepNamePlacement InNamePlacement = JobStepNamePlacement.Default)
 			{
 				ParentPath = InParentPath;
 				Name = InName;
@@ -78,10 +78,7 @@ namespace AutomationTool
 				Parallel = InParallel;
 				PreCondition = InPreCondition;
 				RunCondition = InRunCondition;
-				Parameters = InParameters;
 				ReleaseMode = InReleaseMode;
-				Exclusive = InExclusive;
-				ResourceName = InResourceName;
 				NamePlacement = InNamePlacement;
 			}
 
@@ -115,9 +112,9 @@ namespace AutomationTool
 				{
 					Args.AppendFormat(", resourceName => '{0}'", ResourceName);
 				}
-				if (Parameters != null)
+				if (ActualParameters.Count > 0)
 				{
-					Args.AppendFormat(", actualParameter => [{0}]", String.Join(", ", Parameters.Select(x => String.Format("{{actualParameterName => '{0}', value => '{1}'}}", x.Item1, x.Item2))));
+					Args.AppendFormat(", actualParameter => [{0}]", String.Join(", ", ActualParameters.Select(x => String.Format("{{actualParameterName => '{0}', value => '{1}'}}", x.Key, x.Value))));
 				}
 				if (!String.IsNullOrEmpty(PreCondition))
 				{
@@ -359,10 +356,10 @@ namespace AutomationTool
 				if (LastSticky == null && bHaveECNodes)
 				{
 					// if we don't have any sticky nodes and we have other nodes, we run a fake noop just to release the resource 
-					List<Tuple<string, string>> NoopParameters = new List<Tuple<string, string>>();
-					NoopParameters.Add(new Tuple<string, string>("NodeName", "Noop"));
-					NoopParameters.Add(new Tuple<string, string>("Sticky", "1"));
-					Steps.Add(new JobStep(ParentPath, "Noop", "GUBP_UAT_Node", false, null, null, NoopParameters, JobStepReleaseMode.Release));
+					JobStep NoopStep = new JobStep(ParentPath, "Noop", "GUBP_UAT_Node", false, null, null, JobStepReleaseMode.Release);
+					NoopStep.ActualParameters.Add("NodeName", "Noop");
+					NoopStep.ActualParameters.Add("Sticky", "1");
+					Steps.Add(NoopStep);
 
 					bHasNoop = true;
 				}
@@ -426,24 +423,63 @@ namespace AutomationTool
 
 						bool DoParallel = !Sticky || NodeToDo.IsParallelAgentShareEditor;
 						
-						TriggerNode TriggerNodeToDo = NodeToDo as TriggerNode;
-
-						List<Tuple<string, string>> Parameters = new List<Tuple<string,string>>();
-
-						Parameters.Add(new Tuple<string, string>("NodeName", NodeToDo.Name));
-						Parameters.Add(new Tuple<string, string>("Sticky", NodeToDo.IsSticky ? "1" : "0"));
-
-						if (NodeToDo.AgentSharingGroup != "")
+						List<BuildNode> UncompletedEcDeps = new List<BuildNode>();
+						foreach (BuildNode Dep in NodeToDo.AllDirectDependencies)
 						{
-							Parameters.Add(new Tuple<string, string>("AgentSharingGroup", NodeToDo.AgentSharingGroup));
+							if (!Dep.IsComplete && OrdereredToDo.Contains(Dep)) // if something is already finished, we don't put it into EC
+							{
+								if (OrdereredToDo.IndexOf(Dep) > OrdereredToDo.IndexOf(NodeToDo))
+								{
+									throw new AutomationException("Topological sort error, node {0} has a dependency of {1} which sorted after it.", NodeToDo.Name, Dep.Name);
+								}
+								UncompletedEcDeps.Add(Dep);
+							}
 						}
 
-						string Procedure;
-						if(TriggerNodeToDo == null || TriggerNodeToDo.IsTriggered)
+						string PreCondition = GetPreConditionForNode(OrdereredToDo, ParentPath, bHasNoop, AgentGroupChains, StickyChain, NodeToDo, UncompletedEcDeps);
+						string RunCondition = GetRunConditionForNode(UncompletedEcDeps, ParentPath);
+
+						// Create the job steps for this node
+						TriggerNode TriggerNodeToDo = NodeToDo as TriggerNode;
+						if (TriggerNodeToDo == null)
 						{
+							// Create the jobs to setup the agent sharing group if necessary
+							string NodeParentPath = ParentPath;
+							if (NodeToDo.AgentSharingGroup != "")
+							{
+								NodeParentPath = String.Format("{0}/jobSteps[{1}]", NodeParentPath, NodeToDo.AgentSharingGroup);
+
+								List<BuildNode> MyChain = AgentGroupChains[NodeToDo.AgentSharingGroup];
+								if(MyChain.IndexOf(NodeToDo) <= 0)
+								{
+									// Create the parent job step for this group
+									JobStep ParentStep = new JobStep(ParentPath, NodeToDo.AgentSharingGroup, null, true, PreCondition, null, JobStepReleaseMode.Keep, InNamePlacement: JobStepNamePlacement.BeforeParallel);
+									Steps.Add(ParentStep);
+
+									// Get the resource pool
+									JobStep GetPoolStep = new JobStep(NodeParentPath, String.Format("{0}_GetPool", NodeToDo.AgentSharingGroup), String.Format("GUBP{0}_AgentShare_GetPool", ProcedureInfix), true, PreCondition, null, JobStepReleaseMode.Keep, InNamePlacement: JobStepNamePlacement.BeforeProcedure);
+									GetPoolStep.ActualParameters.Add("AgentSharingGroup", NodeToDo.AgentSharingGroup);
+									GetPoolStep.ActualParameters.Add("NodeName", NodeToDo.Name);
+									Steps.Add(GetPoolStep);
+
+									// Get the agent for this sharing group
+									JobStep GetAgentStep = new JobStep(NodeParentPath, String.Format("{0}_GetAgent", NodeToDo.AgentSharingGroup), String.Format("GUBP{0}_AgentShare_GetAgent", ProcedureInfix), true, GetPoolStep.GetCompletedCondition(), null, JobStepReleaseMode.Keep, InNamePlacement: JobStepNamePlacement.BeforeProcedure);
+									GetAgentStep.Exclusive = JobStepExclusiveMode.Call;
+									GetAgentStep.ResourceName = String.Format("$[/myJob/jobSteps[{0}]/ResourcePool]", NodeToDo.AgentSharingGroup);
+									GetAgentStep.ActualParameters.Add("AgentSharingGroup", NodeToDo.AgentSharingGroup);
+									GetAgentStep.ActualParameters.Add("NodeName", NodeToDo.Name);
+									Steps.Add(GetAgentStep);
+
+									// Set the precondition from this point onwards to be whether the group was set up, since it can't succeed unless the original precondition succeeded
+									PreCondition = GetAgentStep.GetCompletedCondition();
+								}
+							}
+
+							// Get the procedure name
+							string Procedure;
 							if (NodeToDo.IsParallelAgentShareEditor)
 							{
-								if(NodeToDo.AgentSharingGroup == "")
+								if (NodeToDo.AgentSharingGroup == "")
 								{
 									Procedure = "GUBP_UAT_Node_Parallel_AgentShare_Editor";
 								}
@@ -471,75 +507,50 @@ namespace AutomationTool
 							{
 								Procedure += "_Release";
 							}
+
+							// Build the job step for this node
+							JobStep MainStep = new JobStep(NodeParentPath, NodeToDo.Name, Procedure, DoParallel, PreCondition, RunCondition, (NodeToDo.IsSticky && NodeToDo == LastSticky) ? JobStepReleaseMode.Release : JobStepReleaseMode.Keep);
+							MainStep.ActualParameters.Add("NodeName", NodeToDo.Name);
+							MainStep.ActualParameters.Add("Sticky", NodeToDo.IsSticky ? "1" : "0");
+							if (NodeToDo.AgentSharingGroup != "")
+							{
+								MainStep.ActualParameters.Add("AgentSharingGroup", NodeToDo.AgentSharingGroup);
+							}
+							Steps.Add(MainStep);
 						}
 						else
 						{
-							if(TriggerNodeToDo.RequiresRecursiveWorkflow)
+							// Get the procedure name
+							string Procedure;
+							if (TriggerNodeToDo.IsTriggered)
+							{
+								Procedure = String.Format("GUBP{0}_UAT_Node{1}", ProcedureInfix, (NodeToDo == LastSticky)? "_Release" : "");
+							}
+							else if (TriggerNodeToDo.RequiresRecursiveWorkflow)
 							{
-				                Procedure = "GUBP_UAT_Trigger"; //here we run a recursive workflow to wait for the trigger
+								Procedure = "GUBP_UAT_Trigger"; //here we run a recursive workflow to wait for the trigger
 							}
 							else
 							{
 								Procedure = "GUBP_Hardcoded_Trigger"; //here we advance the state in the hardcoded workflow so folks can approve
 							}
 
-							Parameters.Add(new Tuple<string, string>("TriggerState", TriggerNodeToDo.StateName));
-							Parameters.Add(new Tuple<string, string>("ActionText", TriggerNodeToDo.ActionText));
-							Parameters.Add(new Tuple<string, string>("DescText", TriggerNodeToDo.DescriptionText));
-
-							if (NodeToDo.RecipientsForFailureEmails.Length > 0)
+							// Create the job step
+							JobStep TriggerStep = new JobStep(ParentPath, NodeToDo.Name, Procedure, DoParallel, PreCondition, RunCondition, (NodeToDo.IsSticky && NodeToDo == LastSticky) ? JobStepReleaseMode.Release : JobStepReleaseMode.Keep);
+							TriggerStep.ActualParameters.Add("NodeName", NodeToDo.Name);
+							TriggerStep.ActualParameters.Add("Sticky", NodeToDo.IsSticky ? "1" : "0");
+							if(!TriggerNodeToDo.IsTriggered)
 							{
-								Parameters.Add(new Tuple<string, string>("EmailsForTrigger", String.Join(" ", NodeToDo.RecipientsForFailureEmails)));
-							}
-						}
-
-						List<BuildNode> UncompletedEcDeps = new List<BuildNode>();
-						foreach (BuildNode Dep in NodeToDo.AllDirectDependencies)
-						{
-							if (!Dep.IsComplete && OrdereredToDo.Contains(Dep)) // if something is already finished, we don't put it into EC
-							{
-								if (OrdereredToDo.IndexOf(Dep) > OrdereredToDo.IndexOf(NodeToDo))
+								TriggerStep.ActualParameters.Add("TriggerState", TriggerNodeToDo.StateName);
+								TriggerStep.ActualParameters.Add("ActionText", TriggerNodeToDo.ActionText);
+								TriggerStep.ActualParameters.Add("DescText", TriggerNodeToDo.DescriptionText);
+								if (NodeToDo.RecipientsForFailureEmails.Length > 0)
 								{
-									throw new AutomationException("Topological sort error, node {0} has a dependency of {1} which sorted after it.", NodeToDo.Name, Dep.Name);
+									TriggerStep.ActualParameters.Add("EmailsForTrigger", String.Join(" ", NodeToDo.RecipientsForFailureEmails));
 								}
-								UncompletedEcDeps.Add(Dep);
-							}
-						}
-
-						string PreCondition = GetPreConditionForNode(OrdereredToDo, ParentPath, bHasNoop, AgentGroupChains, StickyChain, NodeToDo, UncompletedEcDeps);
-						string RunCondition = GetRunConditionForNode(UncompletedEcDeps, ParentPath);
-
-						// Create the jobs to setup the agent sharing group if necessary
-						string NodeParentPath = ParentPath;
-						if (NodeToDo.AgentSharingGroup != "")
-						{
-							NodeParentPath = String.Format("{0}/jobSteps[{1}]", NodeParentPath, NodeToDo.AgentSharingGroup);
-
-							List<BuildNode> MyChain = AgentGroupChains[NodeToDo.AgentSharingGroup];
-							if(MyChain.IndexOf(NodeToDo) <= 0)
-							{
-								// Create the parent job step for this group
-								Steps.Add(new JobStep(ParentPath, NodeToDo.AgentSharingGroup, null, true, PreCondition, null, null, JobStepReleaseMode.Keep, InNamePlacement: JobStepNamePlacement.BeforeParallel));
-
-								// Get the resource pool
-								List<Tuple<string, string>> GetPoolParameters = new List<Tuple<string,string>>();
-								GetPoolParameters.Add(new Tuple<string, string>("AgentSharingGroup", NodeToDo.AgentSharingGroup));
-								GetPoolParameters.Add(new Tuple<string,string>("NodeName", NodeToDo.Name));
-								Steps.Add(new JobStep(NodeParentPath, String.Format("{0}_GetPool", NodeToDo.AgentSharingGroup), String.Format("GUBP{0}_AgentShare_GetPool", ProcedureInfix), true, PreCondition, null, GetPoolParameters, JobStepReleaseMode.Keep, InNamePlacement: JobStepNamePlacement.BeforeProcedure));
-
-								// Get the agent for this sharing group
-								List<Tuple<string, string>> GetAgentParameters = new List<Tuple<string, string>>();
-								GetAgentParameters.Add(new Tuple<string, string>("AgentSharingGroup", NodeToDo.AgentSharingGroup));
-								GetAgentParameters.Add(new Tuple<string, string>("NodeName", NodeToDo.Name));
-								Steps.Add(new JobStep(NodeParentPath, String.Format("{0}_GetAgent", NodeToDo.AgentSharingGroup), String.Format("GUBP{0}_AgentShare_GetAgent", ProcedureInfix), true, Steps.Last().GetCompletedCondition(), null, GetAgentParameters, JobStepReleaseMode.Keep, InExclusive: JobStepExclusiveMode.Call, InResourceName: String.Format("$[/myJob/jobSteps[{0}]/ResourcePool]", NodeToDo.AgentSharingGroup), InNamePlacement: JobStepNamePlacement.BeforeProcedure));
-
-								// Also add the agent being set up to the existing precondition
-								PreCondition = Steps.Last().GetCompletedCondition();
 							}
+							Steps.Add(TriggerStep);
 						}
-
-						// Build the job step for this node
-						Steps.Add(new JobStep(NodeParentPath, NodeToDo.Name, Procedure, DoParallel, PreCondition, RunCondition, Parameters, (NodeToDo.IsSticky && NodeToDo == LastSticky) ? JobStepReleaseMode.Release : JobStepReleaseMode.Keep));
 					}
 				}
 				WriteECPerl(Steps);
