Use a consistent ordering of arguments to the createJobStep function.

diff --git a/Engine/Source/Programs/AutomationTool/GUBP/ElectricCommander.cs b/Engine/Source/Programs/AutomationTool/GUBP/ElectricCommander.cs
--- a/Engine/Source/Programs/AutomationTool/GUBP/ElectricCommander.cs
+++ b/Engine/Source/Programs/AutomationTool/GUBP/ElectricCommander.cs
@@ -48,13 +48,6 @@ namespace AutomationTool
 			Release,
 		}
 
-		enum JobStepNamePlacement
-		{
-			Default,
-			BeforeParallel,
-			BeforeProcedure,
-		}
-
 		[DebuggerDisplay("{Name}")]
 		class JobStep
 		{
@@ -68,9 +61,8 @@ namespace AutomationTool
 			public string ResourceName;
 			public Dictionary<string, string> ActualParameters = new Dictionary<string, string>();
 			public JobStepReleaseMode ReleaseMode;
-			public JobStepNamePlacement NamePlacement;
 
-			public JobStep(string InParentPath, string InName, string InSubProcedure, bool InParallel, string InPreCondition, string InRunCondition, JobStepReleaseMode InReleaseMode, JobStepNamePlacement InNamePlacement = JobStepNamePlacement.Default)
+			public JobStep(string InParentPath, string InName, string InSubProcedure, bool InParallel, string InPreCondition, string InRunCondition, JobStepReleaseMode InReleaseMode)
 			{
 				ParentPath = InParentPath;
 				Name = InName;
@@ -79,7 +71,6 @@ namespace AutomationTool
 				PreCondition = InPreCondition;
 				RunCondition = InRunCondition;
 				ReleaseMode = InReleaseMode;
-				NamePlacement = InNamePlacement;
 			}
 
 			public string GetCreationStatement()
@@ -87,7 +78,7 @@ namespace AutomationTool
 				StringBuilder Args = new StringBuilder();
 				Args.Append("$batch->createJobStep({");
 				Args.AppendFormat("parentPath => '{0}'", ParentPath);
-				if (!String.IsNullOrEmpty(Name) && NamePlacement == JobStepNamePlacement.BeforeProcedure)
+				if (!String.IsNullOrEmpty(Name))
 				{
 					Args.AppendFormat(", jobStepName => '{0}'", Name);
 				}
@@ -95,15 +86,7 @@ namespace AutomationTool
 				{
 					Args.AppendFormat(", subprocedure => '{0}'", SubProcedure);
 				}
-				if (!String.IsNullOrEmpty(Name) && NamePlacement == JobStepNamePlacement.BeforeParallel)
-				{
-					Args.AppendFormat(", jobStepName => '{0}'", Name);
-				}
 				Args.AppendFormat(", parallel => '{0}'", Parallel ? 1 : 0);
-				if (!String.IsNullOrEmpty(Name) && NamePlacement == JobStepNamePlacement.Default)
-				{
-					Args.AppendFormat(", jobStepName => '{0}'", Name);
-				}
 				if (Exclusive != JobStepExclusiveMode.None)
 				{
 					Args.AppendFormat(", exclusiveMode => '{0}'", Exclusive.ToString().ToLower());
@@ -453,17 +436,17 @@ namespace AutomationTool
 								if(MyChain.IndexOf(NodeToDo) <= 0)
 								{
 									// Create the parent job step for this group
-									JobStep ParentStep = new JobStep(ParentPath, NodeToDo.AgentSharingGroup, null, true, PreCondition, null, JobStepReleaseMode.Keep, InNamePlacement: JobStepNamePlacement.BeforeParallel);
+									JobStep ParentStep = new JobStep(ParentPath, NodeToDo.AgentSharingGroup, null, true, PreCondition, null, JobStepReleaseMode.Keep);
 									Steps.Add(ParentStep);
 
 									// Get the resource pool
-									JobStep GetPoolStep = new JobStep(NodeParentPath, String.Format("{0}_GetPool", NodeToDo.AgentSharingGroup), String.Format("GUBP{0}_AgentShare_GetPool", ProcedureInfix), true, PreCondition, null, JobStepReleaseMode.Keep, InNamePlacement: JobStepNamePlacement.BeforeProcedure);
+									JobStep GetPoolStep = new JobStep(NodeParentPath, String.Format("{0}_GetPool", NodeToDo.AgentSharingGroup), String.Format("GUBP{0}_AgentShare_GetPool", ProcedureInfix), true, PreCondition, null, JobStepReleaseMode.Keep);
 									GetPoolStep.ActualParameters.Add("AgentSharingGroup", NodeToDo.AgentSharingGroup);
 									GetPoolStep.ActualParameters.Add("NodeName", NodeToDo.Name);
 									Steps.Add(GetPoolStep);
 
 									// Get the agent for this sharing group
-									JobStep GetAgentStep = new JobStep(NodeParentPath, String.Format("{0}_GetAgent", NodeToDo.AgentSharingGroup), String.Format("GUBP{0}_AgentShare_GetAgent", ProcedureInfix), true, GetPoolStep.GetCompletedCondition(), null, JobStepReleaseMode.Keep, InNamePlacement: JobStepNamePlacement.BeforeProcedure);
+									JobStep GetAgentStep = new JobStep(NodeParentPath, String.Format("{0}_GetAgent", NodeToDo.AgentSharingGroup), String.Format("GUBP{0}_AgentShare_GetAgent", ProcedureInfix), true, GetPoolStep.GetCompletedCondition(), null, JobStepReleaseMode.Keep);
 									GetAgentStep.Exclusive = JobStepExclusiveMode.Call;
 									GetAgentStep.ResourceName = String.Format("$[/myJob/jobSteps[{0}]/ResourcePool]", NodeToDo.AgentSharingGroup);
 									GetAgentStep.ActualParameters.Add("AgentSharingGroup", NodeToDo.AgentSharingGroup);
