Separate out the creation of the branch and job definition files from the job step creation code.

diff --git a/Engine/Source/Programs/AutomationTool/GUBP/ElectricCommander.cs b/Engine/Source/Programs/AutomationTool/GUBP/ElectricCommander.cs
--- a/Engine/Source/Programs/AutomationTool/GUBP/ElectricCommander.cs
+++ b/Engine/Source/Programs/AutomationTool/GUBP/ElectricCommander.cs
@@ -256,16 +256,39 @@ namespace AutomationTool
 		}
 
 		public void DoCommanderSetup(IEnumerable<BuildNode> AllNodes, IEnumerable<AggregateNode> AllAggregates, List<BuildNode> OrdereredToDo, List<BuildNode> SortedNodes, int TimeIndex, int TimeQuantum, bool bSkipTriggers, bool bFake, bool bFakeEC, string CLString, TriggerNode ExplicitTrigger, List<TriggerNode> UnfinishedTriggers, string FakeFail)
+		{
+			// Check there's some nodes in the graph
+			if (OrdereredToDo.Count == 0)
+			{
+				throw new AutomationException("No nodes to do!");
+			}
+
+			// Make sure everything before the explicit trigger is completed.
+			if (ExplicitTrigger != null)
+			{
+				foreach (BuildNode NodeToDo in OrdereredToDo)
+				{
+					if (!NodeToDo.IsComplete && NodeToDo != ExplicitTrigger && !NodeToDo.DependsOn(ExplicitTrigger)) // if something is already finished, we don't put it into EC
+					{
+						throw new AutomationException("We are being asked to process node {0}, however, this is an explicit trigger {1}, so everything before it should already be handled. It seems likely that you waited too long to run the trigger. You will have to do a new build from scratch.", NodeToDo.Name, ExplicitTrigger.Name);
+					}
+				}
+			}
+
+			// Update all the EC properties, and the branch definition files
+			WriteProperties(AllNodes, AllAggregates, OrdereredToDo, SortedNodes, TimeIndex, TimeQuantum);
+
+			// Write all the job setup
+			WriteJobSteps(OrdereredToDo, bSkipTriggers);
+		}
+
+		private void WriteProperties(IEnumerable<BuildNode> AllNodes, IEnumerable<AggregateNode> AllAggregates, List<BuildNode> OrdereredToDo, List<BuildNode> SortedNodes, int TimeIndex, int TimeQuantum)
 		{
 			List<AggregateNode> SeparatePromotables = FindPromotables(AllAggregates);
 			Dictionary<BuildNode, List<AggregateNode>> DependentPromotions = FindDependentPromotables(AllNodes, SeparatePromotables);
 
 			Dictionary<BuildNode, int> FullNodeListSortKey = GetDisplayOrder(SortedNodes);
 
-			if (OrdereredToDo.Count == 0)
-			{
-				throw new AutomationException("No nodes to do!");
-			}
 			List<string> ECProps = new List<string>();
 			ECProps.Add(String.Format("TimeIndex={0}", TimeIndex));
 			foreach (BuildNode Node in SortedNodes)
@@ -284,28 +307,32 @@ namespace AutomationTool
 			{
 				ECProps.Add(string.Format("PossiblePromotables/{0}={1}", Node.Name, ""));
 			}
-			List<string> ECJobProps = new List<string>();
-			if (ExplicitTrigger != null)
-			{
-				ECJobProps.Add("IsRoot=0");
-			}
-			else
-			{
-				ECJobProps.Add("IsRoot=1");
-			}
 
-			// here we are just making sure everything before the explicit trigger is completed.
-			if (ExplicitTrigger != null)
+			foreach (BuildNode NodeToDo in OrdereredToDo)
 			{
-				foreach (BuildNode NodeToDo in OrdereredToDo)
+				if (!NodeToDo.IsComplete) // if something is already finished, we don't put it into EC  
 				{
-					if (!NodeToDo.IsComplete && NodeToDo != ExplicitTrigger && !NodeToDo.DependsOn(ExplicitTrigger)) // if something is already finished, we don't put it into EC
-					{
-						throw new AutomationException("We are being asked to process node {0}, however, this is an explicit trigger {1}, so everything before it should already be handled. It seems likely that you waited too long to run the trigger. You will have to do a new build from scratch.", NodeToDo.Name, ExplicitTrigger.Name);
-					}
+					List<string> NodeProps = GetECPropsForNode(NodeToDo);
+					ECProps.AddRange(NodeProps);
 				}
 			}
 
+			ECProps.Add("GUBP_LoadedProps=1");
+			string BranchDefFile = CommandUtils.CombinePaths(CommandUtils.CmdEnv.LogFolder, "BranchDef.properties");
+			CommandUtils.WriteAllLines(BranchDefFile, ECProps.ToArray());
+			RunECTool(String.Format("setProperty \"/myWorkflow/BranchDefFile\" \"{0}\"", BranchDefFile.Replace("\\", "\\\\")));
+
+			ECProps.Add("GUBP_LoadedJobProps=1");
+			string BranchJobDefFile = CommandUtils.CombinePaths(CommandUtils.CmdEnv.LogFolder, "BranchJobDef.properties");
+			CommandUtils.WriteAllLines(BranchJobDefFile, ECProps.ToArray());
+			RunECTool(String.Format("setProperty \"/myJob/BranchJobDefFile\" \"{0}\"", BranchJobDefFile.Replace("\\", "\\\\")));
+
+			bool bHasTests = OrdereredToDo.Any(x => x.Node.IsTest());
+			RunECTool(String.Format("setProperty \"/myWorkflow/HasTests\" \"{0}\"", bHasTests));
+		}
+
+		private void WriteJobSteps(List<BuildNode> OrdereredToDo, bool bSkipTriggers)
+		{
 			BuildNode LastSticky = null;
 			bool HitNonSticky = false;
 			bool bHaveECNodes = false;
@@ -331,7 +358,7 @@ namespace AutomationTool
 			}
 
 			List<JobStep> Steps = new List<JobStep>();
-			using(CommandUtils.TelemetryStopwatch PerlOutputStopwatch = new CommandUtils.TelemetryStopwatch("PerlOutput"))
+			using (CommandUtils.TelemetryStopwatch PerlOutputStopwatch = new CommandUtils.TelemetryStopwatch("PerlOutput"))
 			{
 				string ParentPath = Command.ParseParamValue("ParentPath");
 
@@ -378,34 +405,31 @@ namespace AutomationTool
 				{
 					if (!NodeToDo.IsComplete) // if something is already finished, we don't put it into EC  
 					{
-						List<string> NodeProps = GetECPropsForNode(NodeToDo);
-						ECProps.AddRange(NodeProps);
-
 						bool Sticky = NodeToDo.IsSticky;
-						if(NodeToDo.IsSticky)
+						if (NodeToDo.IsSticky)
 						{
-							if(NodeToDo.AgentSharingGroup != "")
+							if (NodeToDo.AgentSharingGroup != "")
 							{
 								throw new AutomationException("Node {0} is both agent sharing and sitcky.", NodeToDo.Name);
 							}
-							if(NodeToDo.AgentPlatform != UnrealTargetPlatform.Win64)
+							if (NodeToDo.AgentPlatform != UnrealTargetPlatform.Win64)
 							{
 								throw new AutomationException("Node {0} is sticky, but {1} hosted. Sticky nodes must be PC hosted.", NodeToDo.Name, NodeToDo.AgentPlatform);
 							}
-							if(NodeToDo.AgentRequirements != "")
+							if (NodeToDo.AgentRequirements != "")
 							{
 								throw new AutomationException("Node {0} is sticky but has agent requirements.", NodeToDo.Name);
 							}
 						}
 
 						string ProcedureInfix = "";
-						if(NodeToDo.AgentPlatform != UnrealTargetPlatform.Unknown && NodeToDo.AgentPlatform != UnrealTargetPlatform.Win64)
+						if (NodeToDo.AgentPlatform != UnrealTargetPlatform.Unknown && NodeToDo.AgentPlatform != UnrealTargetPlatform.Win64)
 						{
 							ProcedureInfix = "_" + NodeToDo.AgentPlatform.ToString();
 						}
 
 						bool DoParallel = !Sticky || NodeToDo.IsParallelAgentShareEditor;
-						
+
 						List<BuildNode> UncompletedEcDeps = new List<BuildNode>();
 						foreach (BuildNode Dep in NodeToDo.AllDirectDependencies)
 						{
@@ -433,7 +457,7 @@ namespace AutomationTool
 								NodeParentPath = String.Format("{0}/jobSteps[{1}]", NodeParentPath, NodeToDo.AgentSharingGroup);
 
 								List<BuildNode> MyChain = AgentGroupChains[NodeToDo.AgentSharingGroup];
-								if(MyChain.IndexOf(NodeToDo) <= 0)
+								if (MyChain.IndexOf(NodeToDo) <= 0)
 								{
 									// Create the parent job step for this group
 									JobStep ParentStep = new JobStep(ParentPath, NodeToDo.AgentSharingGroup, null, true, PreCondition, null, JobStepReleaseMode.Keep);
@@ -507,7 +531,7 @@ namespace AutomationTool
 							string Procedure;
 							if (TriggerNodeToDo.IsTriggered)
 							{
-								Procedure = String.Format("GUBP{0}_UAT_Node{1}", ProcedureInfix, (NodeToDo == LastSticky)? "_Release" : "");
+								Procedure = String.Format("GUBP{0}_UAT_Node{1}", ProcedureInfix, (NodeToDo == LastSticky) ? "_Release" : "");
 							}
 							else if (TriggerNodeToDo.RequiresRecursiveWorkflow)
 							{
@@ -522,7 +546,7 @@ namespace AutomationTool
 							JobStep TriggerStep = new JobStep(ParentPath, NodeToDo.Name, Procedure, DoParallel, PreCondition, RunCondition, (NodeToDo.IsSticky && NodeToDo == LastSticky) ? JobStepReleaseMode.Release : JobStepReleaseMode.Keep);
 							TriggerStep.ActualParameters.Add("NodeName", NodeToDo.Name);
 							TriggerStep.ActualParameters.Add("Sticky", NodeToDo.IsSticky ? "1" : "0");
-							if(!TriggerNodeToDo.IsTriggered)
+							if (!TriggerNodeToDo.IsTriggered)
 							{
 								TriggerStep.ActualParameters.Add("TriggerState", TriggerNodeToDo.StateName);
 								TriggerStep.ActualParameters.Add("ActionText", TriggerNodeToDo.ActionText);
@@ -538,20 +562,6 @@ namespace AutomationTool
 				}
 				WriteECPerl(Steps);
 			}
-			bool bHasTests = OrdereredToDo.Any(x => x.Node.IsTest());
-			RunECTool(String.Format("setProperty \"/myWorkflow/HasTests\" \"{0}\"", bHasTests));
-			{
-				ECProps.Add("GUBP_LoadedProps=1");
-				string BranchDefFile = CommandUtils.CombinePaths(CommandUtils.CmdEnv.LogFolder, "BranchDef.properties");
-				CommandUtils.WriteAllLines(BranchDefFile, ECProps.ToArray());
-				RunECTool(String.Format("setProperty \"/myWorkflow/BranchDefFile\" \"{0}\"", BranchDefFile.Replace("\\", "\\\\")));
-			}
-			{
-				ECProps.Add("GUBP_LoadedJobProps=1");
-				string BranchJobDefFile = CommandUtils.CombinePaths(CommandUtils.CmdEnv.LogFolder, "BranchJobDef.properties");
-				CommandUtils.WriteAllLines(BranchJobDefFile, ECProps.ToArray());
-				RunECTool(String.Format("setProperty \"/myJob/BranchJobDefFile\" \"{0}\"", BranchJobDefFile.Replace("\\", "\\\\")));
-			}
 		}
 
 		private string GetPreConditionForNode(List<BuildNode> OrdereredToDo, string PreconditionParentPath, bool bHasNoop, Dictionary<string, List<BuildNode>> AgentGroupChains, List<BuildNode> StickyChain, BuildNode NodeToDo, List<BuildNode> UncompletedEcDeps)
