Remove some references to legacy node properties.

diff --git a/Engine/Source/Programs/AutomationTool/GUBP/ElectricCommander.cs b/Engine/Source/Programs/AutomationTool/GUBP/ElectricCommander.cs
--- a/Engine/Source/Programs/AutomationTool/GUBP/ElectricCommander.cs
+++ b/Engine/Source/Programs/AutomationTool/GUBP/ElectricCommander.cs
@@ -327,7 +327,7 @@ namespace AutomationTool
 			CommandUtils.WriteAllLines(BranchJobDefFile, ECProps.ToArray());
 			RunECTool(String.Format("setProperty \"/myJob/BranchJobDefFile\" \"{0}\"", BranchJobDefFile.Replace("\\", "\\\\")));
 
-			bool bHasTests = OrdereredToDo.Any(x => x.Node.IsTest());
+			bool bHasTests = OrdereredToDo.Any(x => x.IsTest);
 			RunECTool(String.Format("setProperty \"/myWorkflow/HasTests\" \"{0}\"", bHasTests));
 		}
 
@@ -749,7 +749,7 @@ namespace AutomationTool
 				Dictionary<string, List<BuildNode>> DisplayGroups = new Dictionary<string,List<BuildNode>>();
 				foreach(BuildNode Node in NodesByThisFrequency)
 				{
-					string GroupName = Node.Node.GetDisplayGroupName();
+					string GroupName = Node.DisplayGroupName;
 					if(!DisplayGroups.ContainsKey(GroupName))
 					{
 						DisplayGroups.Add(GroupName, new List<BuildNode>{ Node });
@@ -778,7 +778,7 @@ namespace AutomationTool
 				// Add nodes for each frequency into the master list, trying to match up different groups along the way
 				foreach(BuildNode FirstNode in NodesByThisFrequency)
 				{
-					List<BuildNode> GroupNodes = DisplayGroups[FirstNode.Node.GetDisplayGroupName()];
+					List<BuildNode> GroupNodes = DisplayGroups[FirstNode.DisplayGroupName];
 					foreach(BuildNode GroupNode in GroupNodes)
 					{
 						AddNodeAndDependencies(GroupNode, NodeDependencies, VisitedNodes, SortedNodes);
diff --git a/Engine/Source/Programs/AutomationTool/GUBP/LegacyBranchSetup.cs b/Engine/Source/Programs/AutomationTool/GUBP/LegacyBranchSetup.cs
--- a/Engine/Source/Programs/AutomationTool/GUBP/LegacyBranchSetup.cs
+++ b/Engine/Source/Programs/AutomationTool/GUBP/LegacyBranchSetup.cs
@@ -295,7 +295,7 @@ partial class GUBP
         }
     }
 
-	void FindFrequenciesForNodes(GUBP.GUBPBranchConfig BranchConfig, IEnumerable<BuildNode> Nodes, Dictionary<string, int> FrequencyOverrides)
+	void FindFrequenciesForNodes(GUBP.GUBPBranchConfig BranchConfig, Dictionary<GUBP.GUBPNode, BuildNode> LegacyToNewNodes, Dictionary<string, int> FrequencyOverrides)
 	{
         List<GUBPFrequencyHacker> FrequencyHackers = new List<GUBPFrequencyHacker>();
 
@@ -316,9 +316,10 @@ partial class GUBP
             }
         }
 
-		foreach(BuildNode Node in Nodes)
+		foreach (KeyValuePair<GUBP.GUBPNode, BuildNode> LegacyToNewNodePair in LegacyToNewNodes)
 		{
-			Node.FrequencyShift = Node.Node.CISFrequencyQuantumShift(BranchConfig);
+			BuildNode Node = LegacyToNewNodePair.Value;
+			Node.FrequencyShift = LegacyToNewNodePair.Key.CISFrequencyQuantumShift(BranchConfig);
 
 			foreach(GUBPFrequencyHacker FrequencyHacker in FrequencyHackers)
 			{
@@ -1212,15 +1213,18 @@ partial class GUBP
 		// Calculate the frequency overrides
 		Dictionary<string, int> FrequencyOverrides = ApplyFrequencyBarriers(BranchConfig.GUBPNodes, BranchConfig.GUBPAggregates, BranchOptions.FrequencyBarriers);
 
+		// Get the legacy node to new node mapping
+		Dictionary<GUBP.GUBPNode, BuildNode> LegacyToNewNodes = BranchConfig.GUBPNodes.Values.ToDictionary(x => x, x => x.GetBuildNode());
+
 		// Convert the GUBPNodes and GUBPAggregates maps into lists
 		AllNodes = new List<BuildNode>();
-		AllNodes.AddRange(BranchConfig.GUBPNodes.Values.Select(x => x.GetBuildNode()));
+		AllNodes.AddRange(LegacyToNewNodes.Values);
 
 		AllAggregates = new List<AggregateNode>();
 		AllAggregates.AddRange(BranchConfig.GUBPAggregates.Values.Select(x => new LegacyAggregateNode(x)));
 
 		// Calculate the frequencies for each node
-		FindFrequenciesForNodes(BranchConfig, AllNodes, FrequencyOverrides);
+		FindFrequenciesForNodes(BranchConfig, LegacyToNewNodes, FrequencyOverrides);
 
 		// Get the email list for each node
 		FindEmailsForNodes(BranchName, AllNodes);
@@ -1233,11 +1237,11 @@ partial class GUBP
 			{
 				NodeNames.Add(GameAggregatePromotableNode.StaticGetFullName(GameProj));
 			}
-			foreach (BuildNode Node in AllNodes)
+			foreach(GUBP.GUBPNode Node in BranchConfig.GUBPNodes.Values)
 			{
-				if (Node.Node.GameNameIfAnyForTempStorage() == GameProj.GameName)
+				if (Node.GameNameIfAnyForTempStorage() == GameProj.GameName)
 				{
-					NodeNames.Add(Node.Name);
+					NodeNames.Add(Node.GetFullName());
 				}
 			}
 			if (NodeNames.Count > 0)
