UE-19136 - fix for trying to sign a mac build when staging on PC
#uat

diff --git a/Engine/Source/Programs/AutomationTool/Mac/MacPlatform.Automation.cs b/Engine/Source/Programs/AutomationTool/Mac/MacPlatform.Automation.cs
--- a/Engine/Source/Programs/AutomationTool/Mac/MacPlatform.Automation.cs
+++ b/Engine/Source/Programs/AutomationTool/Mac/MacPlatform.Automation.cs
@@ -426,22 +426,25 @@ public class MacPlatform : Platform
 
 	public override bool SignExecutables(DeploymentContext SC, ProjectParams Params)
 	{
-		// Sign everything we built
-		List<string> FilesToSign = GetExecutableNames(SC);
-		foreach (var Exe in FilesToSign)
+		if (UnrealBuildTool.BuildHostPlatform.Current.Platform == UnrealTargetPlatform.Mac)
 		{
-			string AppBundlePath = "";
-			if (Exe.StartsWith(CombinePaths(SC.RuntimeProjectRootDir, "Binaries", SC.PlatformDir)))
+			// Sign everything we built
+			List<string> FilesToSign = GetExecutableNames(SC);
+			foreach (var Exe in FilesToSign)
 			{
-				AppBundlePath = CombinePaths(SC.RuntimeProjectRootDir, "Binaries", SC.PlatformDir, Path.GetFileNameWithoutExtension(Exe) + ".app");
-			}
-			else if (Exe.StartsWith(CombinePaths(SC.RuntimeRootDir, "Engine/Binaries", SC.PlatformDir)))
-			{
-				AppBundlePath = CombinePaths("Engine/Binaries", SC.PlatformDir, Path.GetFileNameWithoutExtension(Exe) + ".app");
+				string AppBundlePath = "";
+				if (Exe.StartsWith(CombinePaths(SC.RuntimeProjectRootDir, "Binaries", SC.PlatformDir)))
+				{
+					AppBundlePath = CombinePaths(SC.RuntimeProjectRootDir, "Binaries", SC.PlatformDir, Path.GetFileNameWithoutExtension(Exe) + ".app");
+				}
+				else if (Exe.StartsWith(CombinePaths(SC.RuntimeRootDir, "Engine/Binaries", SC.PlatformDir)))
+				{
+					AppBundlePath = CombinePaths("Engine/Binaries", SC.PlatformDir, Path.GetFileNameWithoutExtension(Exe) + ".app");
+				}
+				CodeSign.SignMacFileOrFolder(AppBundlePath);
 			}
-			CodeSign.SignMacFileOrFolder(AppBundlePath);
-		}
 
+		}
 		return true;
 	}
 }
