UAT: Add support to searching C:\\Windows\\Sysnative when exe can't be found in C:\\Windows\\System32

diff --git a/Engine/Source/Programs/AutomationTool/AutomationUtils/BuildEnvironment.cs b/Engine/Source/Programs/AutomationTool/AutomationUtils/BuildEnvironment.cs
--- a/Engine/Source/Programs/AutomationTool/AutomationUtils/BuildEnvironment.cs
+++ b/Engine/Source/Programs/AutomationTool/AutomationUtils/BuildEnvironment.cs
@@ -8,6 +8,7 @@ using Microsoft.Win32;
 using System.Diagnostics;
 using Tools.DotNETCommon;
 using UnrealBuildTool;
+using System.Text.RegularExpressions;
 
 namespace AutomationTool
 {
@@ -115,9 +116,9 @@ namespace AutomationTool
 			EngineSavedFolder = CommandUtils.GetEnvVar(EnvVarNames.EngineSavedFolder);
             CSVFile = CommandUtils.GetEnvVar(EnvVarNames.CSVFile);            
 			LogFolder = CommandUtils.GetEnvVar(EnvVarNames.LogFolder);
-			RobocopyExe = CommandUtils.CombinePaths(Environment.SystemDirectory, "robocopy.exe");
-			MountExe = CommandUtils.CombinePaths(Environment.SystemDirectory, "mount.exe");
-			CmdExe = Utils.IsRunningOnMono ? "/bin/sh" : CommandUtils.CombinePaths(Environment.SystemDirectory, "cmd.exe");
+			RobocopyExe = GetSystemExePath("robocopy.exe");
+			MountExe = GetSystemExePath("mount.exe");
+			CmdExe = Utils.IsRunningOnMono ? "/bin/sh" : GetSystemExePath("cmd.exe");
 
 			if (String.IsNullOrEmpty(LogFolder))
 			{
@@ -158,6 +159,29 @@ namespace AutomationTool
 			LogSettings();
 		}
 
+		/// <summary>
+		/// Returns the path to an executable in the System Directory.
+		/// To help support running 32-bit assemblies on a 64-bit operating system, if the executable
+		/// can't be found in System32, we also search Sysnative.
+		/// </summary>
+		/// <param name="ExeName">The name of the executable to find</param>
+		/// <returns>The path to the executable within the system folder</returns>
+		string GetSystemExePath(string ExeName)
+		{
+			var Result = CommandUtils.CombinePaths(Environment.SystemDirectory, ExeName);
+			if (!CommandUtils.FileExists(Result))
+			{
+				// Use Regex.Replace so we can do a case-insensitive replacement of System32
+				var SysNativeDirectory = Regex.Replace(Environment.SystemDirectory, "System32", "Sysnative", RegexOptions.IgnoreCase);
+				var SysNativeExe = CommandUtils.CombinePaths(SysNativeDirectory, ExeName);
+				if (CommandUtils.FileExists(SysNativeExe))
+				{
+					Result = SysNativeExe;
+				}
+			}
+			return Result;
+		}
+
 		void LogSettings()
 		{
 			Log.TraceVerbose("Command Environment settings:");
