Get UAT running as a 64-bit process again.
* As of .NET 4.5, the default for Any CPU is to "prefer 32-bit process", and you need to explicitly turn this off to require the process be run as 64-bit.
* See http://blogs.microsoft.co.il/sasha/2012/04/04/what-anycpu-really-means-as-of-net-45-and-visual-studio-11/ for details.
* UAT will now assert when Environment.Is64BitProcess is false on startup, as we essentially require 64-bit for certain scripts, and that's how we ran it before upgrading to .NET4.5.
* Also set default optimization options for the project, which were changed at some point.
* Finally, rolled back CL#2633880, which was trying to work around the fact that UAT was running 32-bit.

diff --git a/Engine/Source/Programs/AutomationTool/AutomationTool.csproj b/Engine/Source/Programs/AutomationTool/AutomationTool.csproj
--- a/Engine/Source/Programs/AutomationTool/AutomationTool.csproj
+++ b/Engine/Source/Programs/AutomationTool/AutomationTool.csproj
@@ -63,17 +63,19 @@
     <PlatformTarget>AnyCPU</PlatformTarget>
     <ErrorReport>prompt</ErrorReport>
     <CodeAnalysisRuleSet>MinimumRecommendedRules.ruleset</CodeAnalysisRuleSet>
+    <Prefer32Bit>false</Prefer32Bit>
   </PropertyGroup>
   <PropertyGroup Condition=" '$(Configuration)|$(Platform)' == 'Development|AnyCPU' ">
     <OutputPath>..\..\..\Binaries\DotNET\</OutputPath>
     <DefineConstants>TRACE</DefineConstants>
-    <Optimize>false</Optimize>
+    <Optimize>true</Optimize>
     <DebugType>pdbonly</DebugType>
     <PlatformTarget>AnyCPU</PlatformTarget>
     <ErrorReport>prompt</ErrorReport>
     <CodeAnalysisRuleSet>MinimumRecommendedRules.ruleset</CodeAnalysisRuleSet>
     <DocumentationFile>
     </DocumentationFile>
+    <Prefer32Bit>false</Prefer32Bit>
   </PropertyGroup>
   <ItemGroup>
     <Reference Include="Ionic.Zip.Reduced">
diff --git a/Engine/Source/Programs/AutomationTool/AutomationUtils/BuildEnvironment.cs b/Engine/Source/Programs/AutomationTool/AutomationUtils/BuildEnvironment.cs
--- a/Engine/Source/Programs/AutomationTool/AutomationUtils/BuildEnvironment.cs
+++ b/Engine/Source/Programs/AutomationTool/AutomationUtils/BuildEnvironment.cs
@@ -8,7 +8,6 @@ using Microsoft.Win32;
 using System.Diagnostics;
 using Tools.DotNETCommon;
 using UnrealBuildTool;
-using System.Text.RegularExpressions;
 
 namespace AutomationTool
 {
@@ -116,9 +115,9 @@ namespace AutomationTool
 			EngineSavedFolder = CommandUtils.GetEnvVar(EnvVarNames.EngineSavedFolder);
             CSVFile = CommandUtils.GetEnvVar(EnvVarNames.CSVFile);            
 			LogFolder = CommandUtils.GetEnvVar(EnvVarNames.LogFolder);
-			RobocopyExe = GetSystemExePath("robocopy.exe");
-			MountExe = GetSystemExePath("mount.exe");
-			CmdExe = Utils.IsRunningOnMono ? "/bin/sh" : GetSystemExePath("cmd.exe");
+			RobocopyExe = CommandUtils.CombinePaths(Environment.SystemDirectory, "robocopy.exe");
+			MountExe = CommandUtils.CombinePaths(Environment.SystemDirectory, "mount.exe");
+			CmdExe = Utils.IsRunningOnMono ? "/bin/sh" : CommandUtils.CombinePaths(Environment.SystemDirectory, "cmd.exe");
 
 			if (String.IsNullOrEmpty(LogFolder))
 			{
@@ -159,29 +158,6 @@ namespace AutomationTool
 			LogSettings();
 		}
 
-		/// <summary>
-		/// Returns the path to an executable in the System Directory.
-		/// To help support running 32-bit assemblies on a 64-bit operating system, if the executable
-		/// can't be found in System32, we also search Sysnative.
-		/// </summary>
-		/// <param name="ExeName">The name of the executable to find</param>
-		/// <returns>The path to the executable within the system folder</returns>
-		string GetSystemExePath(string ExeName)
-		{
-			var Result = CommandUtils.CombinePaths(Environment.SystemDirectory, ExeName);
-			if (!CommandUtils.FileExists(Result))
-			{
-				// Use Regex.Replace so we can do a case-insensitive replacement of System32
-				var SysNativeDirectory = Regex.Replace(Environment.SystemDirectory, "System32", "Sysnative", RegexOptions.IgnoreCase);
-				var SysNativeExe = CommandUtils.CombinePaths(SysNativeDirectory, ExeName);
-				if (CommandUtils.FileExists(SysNativeExe))
-				{
-					Result = SysNativeExe;
-				}
-			}
-			return Result;
-		}
-
 		void LogSettings()
 		{
 			Log.TraceVerbose("Command Environment settings:");
diff --git a/Engine/Source/Programs/AutomationTool/Program.cs b/Engine/Source/Programs/AutomationTool/Program.cs
--- a/Engine/Source/Programs/AutomationTool/Program.cs
+++ b/Engine/Source/Programs/AutomationTool/Program.cs
@@ -27,6 +27,10 @@ namespace AutomationTool
 
             try
             {
+                if (!Environment.Is64BitProcess)
+                {
+                    throw new AutomationException(ErrorCodes.Error_UATLaunchFailure, "Not running as a 64-bit process. Something must be wrong with the build settings.");
+                }
                 HostPlatform.Initialize();
 
                 Log.TraceVerbose("Running on {0}", HostPlatform.Current.GetType().Name);
diff --git a/Engine/Source/Programs/AutomationToolLauncher/AutomationToolLauncher.csproj b/Engine/Source/Programs/AutomationToolLauncher/AutomationToolLauncher.csproj
--- a/Engine/Source/Programs/AutomationToolLauncher/AutomationToolLauncher.csproj
+++ b/Engine/Source/Programs/AutomationToolLauncher/AutomationToolLauncher.csproj
@@ -21,6 +21,8 @@
     <DefineConstants>DEBUG;TRACE</DefineConstants>
     <ErrorReport>prompt</ErrorReport>
     <WarningLevel>4</WarningLevel>
+    <Prefer32Bit>false</Prefer32Bit>
+    <AllowUnsafeBlocks>false</AllowUnsafeBlocks>
   </PropertyGroup>
   <PropertyGroup Condition=" '$(Configuration)|$(Platform)' == 'Development|AnyCPU' ">
     <PlatformTarget>AnyCPU</PlatformTarget>
@@ -30,6 +32,8 @@
     <DefineConstants>TRACE</DefineConstants>
     <ErrorReport>prompt</ErrorReport>
     <WarningLevel>4</WarningLevel>
+    <Prefer32Bit>false</Prefer32Bit>
+    <AllowUnsafeBlocks>false</AllowUnsafeBlocks>
   </PropertyGroup>
   <ItemGroup>
     <Reference Include="System" />
