real fix for logging verbosity.
#uat

diff --git a/Engine/Source/Programs/AutomationTool/AutomationUtils/UE4Build.cs b/Engine/Source/Programs/AutomationTool/AutomationUtils/UE4Build.cs
--- a/Engine/Source/Programs/AutomationTool/AutomationUtils/UE4Build.cs
+++ b/Engine/Source/Programs/AutomationTool/AutomationUtils/UE4Build.cs
@@ -1023,8 +1023,7 @@ namespace AutomationTool
 											if (PCHFileName.Contains(@"\SharedPCHs\") && PCHFileName.Contains(@"\UE4Editor\"))
 											{
 												Key = "SharedEditorPCH$ " + PCHFileName;
-												LogWarning("Hack: detected Shared PCH, which will use a different key {0}", Key);
-
+												Log("Hack: detected Shared PCH, which will use a different key {0}", Key);
 											}
 										}
 									}
