Fixed issue with paking DLC when there is no engine directory.

diff --git a/Engine/Source/Programs/AutomationTool/Scripts/CopyBuildToStagingDirectory.Automation.cs b/Engine/Source/Programs/AutomationTool/Scripts/CopyBuildToStagingDirectory.Automation.cs
--- a/Engine/Source/Programs/AutomationTool/Scripts/CopyBuildToStagingDirectory.Automation.cs
+++ b/Engine/Source/Programs/AutomationTool/Scripts/CopyBuildToStagingDirectory.Automation.cs
@@ -247,17 +247,18 @@ public partial class Project : CommandUtils
             SC.StageFiles(StagedFileType.NonUFS, CombinePaths(SC.ProjectRoot, "Plugins", DLCName, "Binaries"), "UE4-*.dll", true, null, null, true);
             SC.StageFiles(StagedFileType.NonUFS, CombinePaths(SC.ProjectRoot, "Plugins", DLCName, "Binaries"), "libUE4Server-*.so", true, null, null, true);
             SC.StageFiles(StagedFileType.NonUFS, CombinePaths(SC.ProjectRoot, "Plugins", DLCName, "Binaries"), "UE4Server-*.dll", true, null, null, true);
-            
+
             // Put all of the cooked dir into the staged dir
-            if (SC.DedicatedServer || Params.DLCIncludeEngineContent)
-            {
-                // Dedicated server cook doesn't save shaders so no Engine dir is created
-                SC.StageFiles(StagedFileType.UFS, CombinePaths(SC.ProjectRoot, "Plugins", DLCName, "Saved", "Cooked", SC.CookPlatform), "*", true, new[] { "AssetRegistry.bin" }, "", true, !Params.UsePak(SC.StageTargetPlatform));
-            }
-            else
+            // Dedicated server cook doesn't save shaders so no Engine dir is created
+            if ((!SC.DedicatedServer) && (!Params.DLCIncludeEngineContent))
             {
-                SC.StageFiles(StagedFileType.UFS, CombinePaths(SC.ProjectRoot, "Plugins", DLCName, "Saved", "Cooked", SC.CookPlatform), "*", true, new[] { "AssetRegistry.bin", CommandUtils.CombinePaths("Engine", "*") }, "", true, !Params.UsePak(SC.StageTargetPlatform));
+                if (Directory.Exists(CombinePaths(SC.ProjectRoot, "Plugins", DLCName, "Saved", "Cooked", SC.CookPlatform, "Engine")))
+                {
+                    SC.StageFiles(StagedFileType.UFS, CombinePaths(SC.ProjectRoot, "Plugins", DLCName, "Saved", "Cooked", SC.CookPlatform), "*", true, new[] { CommandUtils.CombinePaths("Engine", "*") }, "", true, !Params.UsePak(SC.StageTargetPlatform));
+                }
             }
+            
+            SC.StageFiles(StagedFileType.UFS, CombinePaths(SC.ProjectRoot, "Plugins", DLCName, "Saved", "Cooked", SC.CookPlatform), "*", true, new[] { "AssetRegistry.bin" }, "", true, !Params.UsePak(SC.StageTargetPlatform));
 
             return;
         }
