Convert to a better exception class and error message. Most tools assume NotImplementedException means something very specific.

diff --git a/Engine/Source/Programs/AutomationTool/AutomationUtils/CommandletUtils.cs b/Engine/Source/Programs/AutomationTool/AutomationUtils/CommandletUtils.cs
--- a/Engine/Source/Programs/AutomationTool/AutomationUtils/CommandletUtils.cs
+++ b/Engine/Source/Programs/AutomationTool/AutomationUtils/CommandletUtils.cs
@@ -224,7 +224,7 @@ namespace AutomationTool
 				case UnrealBuildTool.UnrealTargetPlatform.Linux:
 					return CommandUtils.CombinePaths(BuildRoot, "Engine/Binaries/Linux/UE4Editor");
 				default:
-					throw new NotImplementedException();
+					throw new AutomationException("EditorCommandlet is not supported for platform {0}", HostPlatform);
 			}
 		}
 
