UE-17454 - fix for cook hang on mac
we now don't write the same message twice to the output window
#mac

diff --git a/Engine/Source/Programs/AutomationTool/Scripts/CookCommand.Automation.cs b/Engine/Source/Programs/AutomationTool/Scripts/CookCommand.Automation.cs
--- a/Engine/Source/Programs/AutomationTool/Scripts/CookCommand.Automation.cs
+++ b/Engine/Source/Programs/AutomationTool/Scripts/CookCommand.Automation.cs
@@ -137,7 +137,7 @@ public partial class Project : CommandUtils
 
             try
             {
-                var CommandletParams = "-buildmachine -fileopenlog";
+                var CommandletParams = IsBuildMachine ? "-buildmachine -fileopenlog" : "-fileopenlog";
                 if (Params.UnversionedCookedContent)
                 {
                     CommandletParams += " -unversioned";
