Add a -ScriptsForProject parameter to UAT, which causes it to only compile and load scripts for a certain project. Fixes issues when only syncing source for one project, as is typical in UGS workflows. Always pass that parameter from the editor, and always compile UAT when running from a non-installed, non-promoted build.

diff --git a/Engine/Source/Developer/LauncherServices/Private/Launcher/LauncherUATTask.h b/Engine/Source/Developer/LauncherServices/Private/Launcher/LauncherUATTask.h
--- a/Engine/Source/Developer/LauncherServices/Private/Launcher/LauncherUATTask.h
+++ b/Engine/Source/Developer/LauncherServices/Private/Launcher/LauncherUATTask.h
@@ -59,7 +59,8 @@ protected:
 		FString UATCommandLine;
 		FString ProjectPath = *ChainState.Profile->GetProjectPath();
 		ProjectPath = FPaths::ConvertRelativePathToFull(ProjectPath);
-		UATCommandLine = FString::Printf(TEXT("BuildCookRun -project=\"%s\" -noP4 -clientconfig=%s -serverconfig=%s"),
+		UATCommandLine = FString::Printf(TEXT("-ScriptsForProject=\"%s\" BuildCookRun -project=\"%s\" -noP4 -clientconfig=%s -serverconfig=%s"),
+			*ProjectPath,
 			*ProjectPath,
 			*ConfigStrings[ChainState.Profile->GetBuildConfiguration()],
 			*ConfigStrings[ChainState.Profile->GetBuildConfiguration()]);
diff --git a/Engine/Source/Developer/LauncherServices/Private/Profiles/LauncherProfile.h b/Engine/Source/Developer/LauncherServices/Private/Profiles/LauncherProfile.h
--- a/Engine/Source/Developer/LauncherServices/Private/Profiles/LauncherProfile.h
+++ b/Engine/Source/Developer/LauncherServices/Private/Profiles/LauncherProfile.h
@@ -845,7 +845,7 @@ public:
 
 		// default build settings
 		BuildGame = false;
-		BuildUAT = false;
+		BuildUAT = !FApp::GetEngineIsPromotedBuild() && !FApp::IsEngineInstalled();
 
 		// default cook settings
 		CookConfiguration = FApp::GetBuildConfiguration();
diff --git a/Engine/Source/Editor/MainFrame/Private/Frame/MainFrameActions.cpp b/Engine/Source/Editor/MainFrame/Private/Frame/MainFrameActions.cpp
--- a/Engine/Source/Editor/MainFrame/Private/Frame/MainFrameActions.cpp
+++ b/Engine/Source/Editor/MainFrame/Private/Frame/MainFrameActions.cpp
@@ -288,7 +288,7 @@ const TCHAR* GetUATCompilationFlags()
 {
 	// We never want to compile editor targets when invoking UAT in this context.
 	// If we are rocket or don't have a compiler, we must assume we have a precompiled UAT.
-	return FRocketSupport::IsRocket() || !FSourceCodeNavigation::IsCompilerAvailable()
+	return (FApp::GetEngineIsPromotedBuild() || FApp::IsEngineInstalled())
 		? TEXT("-nocompile -nocompileeditor")
 		: TEXT("-nocompileeditor");
 }
@@ -579,19 +579,6 @@ void FMainFrameActionCallbacks::PackageProject( const FName InPlatformInfoName )
 		OptionalParams += FString::Printf(TEXT(" -manifests -createchunkinstall -chunkinstalldirectory=\"%s\" -chunkinstallversion=%s"), *(PackagingSettings->HttpChunkInstallDataDirectory.Path), *(PackagingSettings->HttpChunkInstallDataVersion));
 	}
 
-	FString ExecutableName = FPlatformProcess::ExecutableName(false);
-#if PLATFORM_WINDOWS
-	// turn UE4editor into UE4editor.cmd
-	if(ExecutableName.EndsWith(".exe", ESearchCase::IgnoreCase) && !FPaths::GetBaseFilename(ExecutableName).EndsWith("-cmd", ESearchCase::IgnoreCase))
-	{
-		FString NewExeName = ExecutableName.Left(ExecutableName.Len() - 4) + "-Cmd.exe";
-		if (FPaths::FileExists(NewExeName))
-		{
-			ExecutableName = NewExeName;
-		}
-	}
-#endif
-
 	int32 NumCookers = GetDefault<UEditorExperimentalSettings>()->MultiProcessCooking;
 	if (NumCookers > 0 )
 	{
@@ -609,14 +596,15 @@ void FMainFrameActionCallbacks::PackageProject( const FName InPlatformInfoName )
 	Configuration = Configuration.Replace(TEXT("PPBC_"), TEXT(""));
 
 	FString ProjectPath = FPaths::IsProjectFilePathSet() ? FPaths::ConvertRelativePathToFull(FPaths::GetProjectFilePath()) : FPaths::RootDir() / FApp::GetGameName() / FApp::GetGameName() + TEXT(".uproject");
-	FString CommandLine = FString::Printf(TEXT("BuildCookRun %s%s%s -nop4 -project=\"%s\" -cook -stage -archive -archivedirectory=\"%s\" -package -clientconfig=%s -ue4exe=%s %s -utf8output"),
+	FString CommandLine = FString::Printf(TEXT("-ScriptsForProject=\"%s\" BuildCookRun %s%s%s -nop4 -project=\"%s\" -cook -stage -archive -archivedirectory=\"%s\" -package -clientconfig=%s -ue4exe=%s %s -utf8output"),
+		*ProjectPath,
 		FRocketSupport::IsRocket() ? TEXT("-rocket ") : TEXT(""),
 		GetUATCompilationFlags(),
 		FApp::IsEngineInstalled() ? TEXT(" -installed") : TEXT(""),
 		*ProjectPath,
 		*PackagingSettings->StagingDirectory.Path,
 		*Configuration,
-		*ExecutableName,
+		*FUnrealEdMisc::Get().GetExecutableForCommandlets(),
 		*OptionalParams
 	);
 
diff --git a/Engine/Source/Programs/AutomationTool/AutomationUtils/Automation.cs b/Engine/Source/Programs/AutomationTool/AutomationUtils/Automation.cs
--- a/Engine/Source/Programs/AutomationTool/AutomationUtils/Automation.cs
+++ b/Engine/Source/Programs/AutomationTool/AutomationUtils/Automation.cs
@@ -159,9 +159,10 @@ AutomationTool.exe [-verbose] [-compileonly] [-p4] Command0 [-Arg0 -Arg1 -Arg2 
 		/// <param name="ParamIndex">Parameter index</param>
 		/// <param name="CommandLine">Command line</param>
 		/// <param name="CurrentCommand">Recently parsed command</param>
+		/// <param name="OutScriptsForProjectFileName">The only project to build scripts for</param>
 		/// <param name="OutAdditionalScriptsFolders">Additional script locations</param>
 		/// <returns>True if the parameter has been successfully parsed.</returns>
-		private static void ParseParam(string CurrentParam, CommandInfo CurrentCommand, List<string> OutAdditionalScriptsFolders)
+		private static void ParseParam(string CurrentParam, CommandInfo CurrentCommand, ref string OutScriptsForProjectFileName, List<string> OutAdditionalScriptsFolders)
 		{
 			bool bGlobalParam = false;
 			foreach (var RegisteredParam in GlobalCommandLine.RegisteredArgs)
@@ -176,7 +177,20 @@ AutomationTool.exe [-verbose] [-compileonly] [-p4] Command0 [-Arg0 -Arg1 -Arg2 
 			}
 
 			// The parameter was not found in the list of global parameters, continue looking...
-			if (CurrentParam.StartsWith("-ScriptDir=", StringComparison.InvariantCultureIgnoreCase))
+			if (CurrentParam.StartsWith("-ScriptsForProject=", StringComparison.InvariantCultureIgnoreCase))
+			{
+				if(OutScriptsForProjectFileName != null)
+				{
+					throw new AutomationException("The -ProjectScripts argument may only be specified once");
+				}
+				var ProjectFileName = CurrentParam.Substring(CurrentParam.IndexOf('=') + 1);
+				if(!File.Exists(ProjectFileName))
+				{
+					throw new AutomationException("Project '{0}' does not exist", ProjectFileName);
+				}
+				OutScriptsForProjectFileName = Path.GetFullPath(ProjectFileName);
+			}
+			else if (CurrentParam.StartsWith("-ScriptDir=", StringComparison.InvariantCultureIgnoreCase))
 			{
 				var ScriptDir = CurrentParam.Substring(CurrentParam.IndexOf('=') + 1);
 				if (Directory.Exists(ScriptDir))
@@ -220,20 +234,22 @@ AutomationTool.exe [-verbose] [-compileonly] [-p4] Command0 [-Arg0 -Arg1 -Arg2 
 		/// <param name="CommandLine">Command line</param>
 		/// <param name="OutCommandsToExecute">List containing the names of commands to execute</param>
 		/// <param name="OutAdditionalScriptsFolders">Optional list of additional paths to look for scripts in</param>
-		private static void ParseCommandLine(string[] CommandLine, List<CommandInfo> OutCommandsToExecute, List<string> OutAdditionalScriptsFolders)
+		private static void ParseCommandLine(string[] CommandLine, List<CommandInfo> OutCommandsToExecute, out string OutScriptsForProjectFileName, List<string> OutAdditionalScriptsFolders)
 		{
 			// Initialize global command line parameters
 			GlobalCommandLine.Init();
 
 			Log.TraceInformation("Parsing command line: {0}", String.Join(" ", CommandLine));
 
+			OutScriptsForProjectFileName = null;
+
 			CommandInfo CurrentCommand = null;
 			for (int Index = 0; Index < CommandLine.Length; ++Index)
 			{
 				var Param = CommandLine[Index];
 				if (Param.StartsWith("-") || Param.Contains("="))
 				{
-					ParseParam(CommandLine[Index], CurrentCommand, OutAdditionalScriptsFolders);
+					ParseParam(CommandLine[Index], CurrentCommand, ref OutScriptsForProjectFileName, OutAdditionalScriptsFolders);
 				}
 				else
 				{
@@ -288,8 +304,9 @@ AutomationTool.exe [-verbose] [-compileonly] [-p4] Command0 [-Arg0 -Arg1 -Arg2 
 
 			// Scan the command line for commands to execute.
 			var CommandsToExecute = new List<CommandInfo>();
+			string OutScriptsForProjectFileName;
 			var AdditionalScriptsFolders = new List<string>();
-			ParseCommandLine(CommandLine, CommandsToExecute, AdditionalScriptsFolders);
+			ParseCommandLine(CommandLine, CommandsToExecute, out OutScriptsForProjectFileName, AdditionalScriptsFolders);
 
 			// Check for build machine override (force local)
 			IsBuildMachine = GlobalCommandLine.ForceLocal ? false : IsBuildMachine;
@@ -330,7 +347,7 @@ AutomationTool.exe [-verbose] [-compileonly] [-p4] Command0 [-Arg0 -Arg1 -Arg2 
 			ScriptCompiler Compiler = new ScriptCompiler();
 			using(CommandUtils.TelemetryStopwatch ScriptCompileStopwatch = new CommandUtils.TelemetryStopwatch("ScriptCompile"))
 			{
-				Compiler.FindAndCompileAllScripts(AdditionalScriptsFolders: AdditionalScriptsFolders);
+				Compiler.FindAndCompileAllScripts(OutScriptsForProjectFileName, AdditionalScriptsFolders);
 			}
 
 			if (GlobalCommandLine.CompileOnly)
diff --git a/Engine/Source/Programs/AutomationTool/AutomationUtils/ScriptCompiler.cs b/Engine/Source/Programs/AutomationTool/AutomationUtils/ScriptCompiler.cs
--- a/Engine/Source/Programs/AutomationTool/AutomationUtils/ScriptCompiler.cs
+++ b/Engine/Source/Programs/AutomationTool/AutomationUtils/ScriptCompiler.cs
@@ -53,8 +53,9 @@ namespace AutomationTool
 		/// <summary>
 		/// Finds and/or compiles all script files and assemblies.
 		/// </summary>
+		/// <param name="ScriptsForProjectFileName">Path to the current project. May be null, in which case we compile scripts for all projects.</param>
 		/// <param name="AdditionalScriptsFolders">Additional script fodlers to look for source files in.</param>
-		public void FindAndCompileAllScripts(List<string> AdditionalScriptsFolders)
+		public void FindAndCompileAllScripts(string ScriptsForProjectFileName, List<string> AdditionalScriptsFolders)
 		{
 			bool DoCompile = false;
 			if (GlobalCommandLine.Compile)
@@ -78,7 +79,7 @@ namespace AutomationTool
 			if (DoCompile && !String.IsNullOrEmpty(CommandUtils.CmdEnv.MsBuildExe))
 			{
 				CleanupScriptsAssemblies();
-				FindAndCompileScriptModules(AdditionalScriptsFolders);
+				FindAndCompileScriptModules(ScriptsForProjectFileName, AdditionalScriptsFolders);
 			}
 
 			var ScriptAssemblies = new List<Assembly>();
@@ -110,7 +111,7 @@ namespace AutomationTool
 			}
 		}
 
-		private static void FindAndCompileScriptModules(List<string> AdditionalScriptsFolders)
+		private static void FindAndCompileScriptModules(string ScriptsForProjectFileName, List<string> AdditionalScriptsFolders)
 		{
 			var OldCWD = Environment.CurrentDirectory;
 			var UnrealBuildToolCWD = CommandUtils.CombinePaths(CommandUtils.CmdEnv.LocalRoot, "Engine", "Source");
@@ -126,7 +127,15 @@ namespace AutomationTool
 
 			// Configure the rules compiler
 			// Get all game folders and convert them to build subfolders.
-			var AllGameFolders = UnrealBuildTool.UEBuildTarget.DiscoverAllGameFolders();
+			List<string> AllGameFolders;
+			if(ScriptsForProjectFileName == null)
+			{
+				AllGameFolders = UnrealBuildTool.UEBuildTarget.DiscoverAllGameFolders();
+			}
+			else
+			{
+				AllGameFolders = new List<string>{ Path.GetDirectoryName(ScriptsForProjectFileName) };
+			}
 			var BuildFolders = new List<string>(AllGameFolders.Count);
 			foreach (var Folder in AllGameFolders)
 			{
