Remove build.properties. Applications can now use Build.version for structured version information; timestamp was only used by the launcher, so it now uses its own.

diff --git a/Engine/Source/Programs/AutomationTool/AutomationUtils/Utils.cs b/Engine/Source/Programs/AutomationTool/AutomationUtils/Utils.cs
--- a/Engine/Source/Programs/AutomationTool/AutomationUtils/Utils.cs
+++ b/Engine/Source/Programs/AutomationTool/AutomationUtils/Utils.cs
@@ -994,31 +994,6 @@ namespace AutomationTool
                 throw new AutomationException(string.Format("Failed to parse {0} as an FEngineVersion compatible string", versionString), ex);
             }
         }
-
-        public static DateTime BuildTime()
-        {
-            string VerFile = CommandUtils.CombinePaths(CommandUtils.CmdEnv.LocalRoot, "Engine", "Build", "build.properties");
-
-            var VerLines = CommandUtils.ReadAllText(VerFile);
-
-            var SearchFor = "TimestampForBVT=";
-
-            int Index = VerLines.IndexOf(SearchFor);
-
-            if (Index < 0)
-            {
-                throw new AutomationException("Could not find {0} in {1} for file {2}", SearchFor, VerLines, VerFile);
-            }
-            Index = Index + SearchFor.Length;
-
-
-            var Parts = VerLines.Substring(Index).Split('.', '_', '-', '\n', '\r', '\t');
-
-            DateTime Result = new DateTime(int.Parse(Parts[0]), int.Parse(Parts[1]), int.Parse(Parts[2]), int.Parse(Parts[3]), int.Parse(Parts[4]), int.Parse(Parts[5]));
-
-			CommandUtils.LogConsole("Current Build Time is {0}", Result);
-            return Result;
-        }
     }
 
     #endregion
diff --git a/Engine/Source/Programs/AutomationTool/Scripts/ArchiveCommand.Automation.cs b/Engine/Source/Programs/AutomationTool/Scripts/ArchiveCommand.Automation.cs
--- a/Engine/Source/Programs/AutomationTool/Scripts/ArchiveCommand.Automation.cs
+++ b/Engine/Source/Programs/AutomationTool/Scripts/ArchiveCommand.Automation.cs
@@ -26,9 +26,9 @@ public partial class Project : CommandUtils
 
 		if (Params.ArchiveMetaData)
 		{
-			// archive the build.properties for extra info for testing, etc
-			string BuildPropertiesFile = CombinePaths(SC.LocalRoot, "Engine", "Build", "build.properties");
-			SC.ArchiveFiles(Path.GetDirectoryName(BuildPropertiesFile), Path.GetFileName(BuildPropertiesFile));
+			// archive the build.version file for extra info for testing, etc
+			string BuildVersionFile = CombinePaths(SC.LocalRoot, "Engine", "Build", "Build.version");
+			SC.ArchiveFiles(Path.GetDirectoryName(BuildVersionFile), Path.GetFileName(BuildVersionFile));
 		}
 	}
 
diff --git a/Engine/Source/Programs/AutomationTool/Scripts/CopyBuildToStagingDirectory.Automation.cs b/Engine/Source/Programs/AutomationTool/Scripts/CopyBuildToStagingDirectory.Automation.cs
--- a/Engine/Source/Programs/AutomationTool/Scripts/CopyBuildToStagingDirectory.Automation.cs
+++ b/Engine/Source/Programs/AutomationTool/Scripts/CopyBuildToStagingDirectory.Automation.cs
@@ -281,7 +281,7 @@ public partial class Project : CommandUtils
 		{
 			BuildPropertiesPath = BuildPropertiesPath.ToLowerInvariant();
 		}
-		SC.StageFiles(StagedFileType.NonUFS, BuildPropertiesPath, "build.properties", false, null, null, true);
+		SC.StageFiles(StagedFileType.NonUFS, BuildPropertiesPath, "Build.version", false, null, null, true);
 
 		// move the UE4Commandline.txt file to the root of the stage
 		// this file needs to be treated as a UFS file for casing, but NonUFS for being put into the .pak file
diff --git a/Engine/Source/Programs/AutomationTool/Scripts/Tests.Automation.cs b/Engine/Source/Programs/AutomationTool/Scripts/Tests.Automation.cs
--- a/Engine/Source/Programs/AutomationTool/Scripts/Tests.Automation.cs
+++ b/Engine/Source/Programs/AutomationTool/Scripts/Tests.Automation.cs
@@ -476,15 +476,6 @@ class TestMacZip : BuildCommand
 	}
 }
 
-[Help("Reads the build time from build.properties")]
-class TestBuildTime : BuildCommand
-{
-	public override void ExecuteBuild()
-	{
-		FEngineVersionSupport.BuildTime();
-	}
-}
-
 [Help("Tests P4 functionality. Creates a new changelist under the workspace %P4CLIENT%")]
 [RequireP4]
 class TestP4_CreateChangelist : BuildCommand
