fix for invalid option when code signing multiple executables at once
#uat

diff --git a/Engine/Source/Programs/AutomationTool/AutomationUtils/CommandUtils.cs b/Engine/Source/Programs/AutomationTool/AutomationUtils/CommandUtils.cs
--- a/Engine/Source/Programs/AutomationTool/AutomationUtils/CommandUtils.cs
+++ b/Engine/Source/Programs/AutomationTool/AutomationUtils/CommandUtils.cs
@@ -2480,7 +2480,7 @@ namespace AutomationTool
 			//@TODO: Verbosity choosing
 			//  /v will spew lots of info
 			//  /q does nothing on success and minimal output on failure
-			string CodeSignArgs = String.Format("sign{0} /a /n \"{1}\" /t {2} /d /v {3}", SpecificStoreArg, SigningIdentity, TimestampServer, FilesToSign);
+			string CodeSignArgs = String.Format("sign{0} /a /n \"{1}\" /t {2} /v {3}", SpecificStoreArg, SigningIdentity, TimestampServer, FilesToSign);
 
 			DateTime StartTime = DateTime.Now;
 
