Remove unnecessary var

diff --git a/Engine/Source/Programs/AutomationTool/Program.cs b/Engine/Source/Programs/AutomationTool/Program.cs
--- a/Engine/Source/Programs/AutomationTool/Program.cs
+++ b/Engine/Source/Programs/AutomationTool/Program.cs
@@ -14,9 +14,6 @@ namespace AutomationTool
 {
 	public class Program
 	{
-		// This needs to be static, otherwise SetConsoleCtrlHandler will result in a crash on exit.
-		static ProcessManager.CtrlHandlerDelegate ProgramCtrlHandler = new ProcessManager.CtrlHandlerDelegate(CtrlHandler);
-
 		[STAThread]
 		public static int Main()
 		{
@@ -45,7 +42,7 @@ namespace AutomationTool
                 var Domain = AppDomain.CurrentDomain;
                 Domain.ProcessExit += Domain_ProcessExit;
                 Domain.DomainUnload += Domain_ProcessExit;
-                HostPlatform.Current.SetConsoleCtrlHandler(ProgramCtrlHandler);
+                HostPlatform.Current.SetConsoleCtrlHandler(CtrlHandler);
 
                 var Version = AssemblyUtils.ExecutableVersion;
                 Log.TraceVerbose("{0} ver. {1}", Version.ProductName, Version.ProductVersion);
