Added additional build telemetry when storing to shared temp storage. Outputs the total file count and total size in bytes.
* ie: UAT.StoreToTempStorage.5.1024
* this can be pulled about by downstream telemetry to measure temp storage performance across the system, and how it's being used in general.

Misc
* Fixed CommanUtils.WriteToFile to add a Newline to match old behavior.
* Added a new feature to the telemetry stopwatch to allow setting a dynamically generated string when finishing the timing.

diff --git a/Engine/Source/Programs/AutomationTool/AutomationUtils/CommandUtils.cs b/Engine/Source/Programs/AutomationTool/AutomationUtils/CommandUtils.cs
--- a/Engine/Source/Programs/AutomationTool/AutomationUtils/CommandUtils.cs
+++ b/Engine/Source/Programs/AutomationTool/AutomationUtils/CommandUtils.cs
@@ -1026,7 +1026,7 @@ namespace AutomationTool
             Filename = ConvertSeparators(PathSeparator.Default, Filename);
             try
 			{
-                File.AppendAllText(Filename, Text);
+                File.AppendAllText(Filename, Text + Environment.NewLine);
 			}
 			catch (Exception Ex)
 			{
@@ -1928,20 +1928,35 @@ namespace AutomationTool
 				bFinished = true;
 			}
 
-			public void Finish()
+            /// <summary>
+            /// Flushes the time to <see cref="CmdEnv.CSVFile"/> if we are the build machine and that environment variable is specified.
+            /// Call Finish manually with an alternate name to use that one instead. Useful for dynamically generated names that you can't specify at construction.
+            /// </summary>
+            /// <param name="AlternateName">Used in place of the Name specified during construction.</param>
+			public void Finish(string AlternateName = null)
 			{
-				if(!bFinished && IsBuildMachine && !String.IsNullOrEmpty(CmdEnv.CSVFile) && CmdEnv.CSVFile != "nul")
+				if(!bFinished)
 				{
-					try
-					{
-						File.AppendAllText(CmdEnv.CSVFile, String.Format("UAT,{0},{1},{2}" + Environment.NewLine, Name, StartTime.ToString(), DateTime.Now.ToString()));
-					}
-					catch (Exception Ex)
-					{
-						LogWarning("Could not append to csv file ({0}) : {1}", CmdEnv.CSVFile, Ex.ToString());
-					}
-					bFinished = true;
+                    if (!string.IsNullOrEmpty(AlternateName))
+                    {
+                        Name = AlternateName;
+                    }
+
+                    var OutputStr = String.Format("UAT,{0},{1},{2}" + Environment.NewLine, Name, StartTime, DateTime.Now);
+                    LogVerbose(OutputStr);
+                    if (IsBuildMachine && !String.IsNullOrEmpty(CmdEnv.CSVFile) && CmdEnv.CSVFile != "nul")
+                    {
+					    try
+					    {
+                            File.AppendAllText(CmdEnv.CSVFile, OutputStr);
+                        }
+					    catch (Exception Ex)
+					    {
+						    LogWarning("Could not append to csv file ({0}) : {1}", CmdEnv.CSVFile, Ex.ToString());
+					    }
+                    }
 				}
+                bFinished = true;
 			}
 
 			public void Dispose()
diff --git a/Engine/Source/Programs/AutomationTool/GUBP/GUBP.cs b/Engine/Source/Programs/AutomationTool/GUBP/GUBP.cs
--- a/Engine/Source/Programs/AutomationTool/GUBP/GUBP.cs
+++ b/Engine/Source/Programs/AutomationTool/GUBP/GUBP.cs
@@ -1729,18 +1729,8 @@ public partial class GUBP : BuildCommand
 						
                     }
 
-					using(TelemetryStopwatch StoreBuildProductsStopwatch = new TelemetryStopwatch("StoreBuildProducts"))
-					{
-						double StoreDuration = 0.0;
-						DateTime StartTime = DateTime.UtcNow;
-						TempStorage.StoreToTempStorage(CmdEnv, NodeStoreName, NodeToDo.Node.BuildProducts, !bSaveSharedTempStorage, GameNameIfAny, StorageRootIfAny);
-						StoreDuration = (DateTime.UtcNow - StartTime).TotalMilliseconds / 1000;
-						LogConsole("Took {0} seconds to store build products", StoreDuration);
-						if (IsBuildMachine)
-						{
-							EC.RunECTool(String.Format("setProperty \"/myJobStep/StoreDuration\" \"{0}\"", StoreDuration.ToString()));
-						}
-					}
+                    TempStorage.StoreToTempStorage(CmdEnv, NodeStoreName, NodeToDo.Node.BuildProducts, !bSaveSharedTempStorage, GameNameIfAny, StorageRootIfAny);
+
                     if (ParseParam("StompCheck"))
                     {
                         foreach (string Dep in NodeToDo.Node.AllDependencies)
diff --git a/Engine/Source/Programs/AutomationTool/GUBP/TempStorage.cs b/Engine/Source/Programs/AutomationTool/GUBP/TempStorage.cs
--- a/Engine/Source/Programs/AutomationTool/GUBP/TempStorage.cs
+++ b/Engine/Source/Programs/AutomationTool/GUBP/TempStorage.cs
@@ -753,83 +753,87 @@ namespace AutomationTool
             var Local = SaveLocalTempStorageManifest(Env, BaseFolder, StorageBlockName, Files); 
             if (!bLocalOnly)
             {
-                var StartTime = DateTime.UtcNow;
-
-                var BlockPath = SharedTempStorageDirectory(StorageBlockName, GameFolder);
-                CommandUtils.LogConsole("Storing to {0}", BlockPath);
-                if (CommandUtils.DirectoryExists_NoExceptions(BlockPath))
-                {
-                    throw new AutomationException("Storage Block Already Exists! {0}", BlockPath);
-                }
-                CommandUtils.CreateDirectory(true, BlockPath);
-                if (!CommandUtils.DirectoryExists_NoExceptions(BlockPath))
+                using (CommandUtils.TelemetryStopwatch TelemetryStopwatch = new CommandUtils.TelemetryStopwatch("StoreToTempStorage"))
                 {
-                    throw new AutomationException("Storage Block Could Not Be Created! {0}", BlockPath);
-                }
+                    var StartTime = DateTime.UtcNow;
 
-                var DestFiles = new List<string>();
-                if (ThreadsToCopyWith() < 2)
-                {
-                    foreach (string InFilename in Files)
+                    var BlockPath = SharedTempStorageDirectory(StorageBlockName, GameFolder);
+                    CommandUtils.LogConsole("Storing to {0}", BlockPath);
+                    if (CommandUtils.DirectoryExists_NoExceptions(BlockPath))
                     {
-                        var Filename = CommandUtils.CombinePaths(InFilename);
-                        InternalUtils.Robust_FileExists_NoExceptions(false, Filename, "Could not add {0} to manifest because it does not exist");
+                        throw new AutomationException("Storage Block Already Exists! {0}", BlockPath);
+                    }
+                    CommandUtils.CreateDirectory(true, BlockPath);
+                    if (!CommandUtils.DirectoryExists_NoExceptions(BlockPath))
+                    {
+                        throw new AutomationException("Storage Block Could Not Be Created! {0}", BlockPath);
+                    }
 
-                        if (!Filename.StartsWith(BaseFolder, StringComparison.InvariantCultureIgnoreCase))
-                        {
-                            throw new AutomationException("Could not add {0} to manifest because it does not start with the base folder {1}", Filename, BaseFolder);
-                        }
-                        var RelativeFile = Filename.Substring(BaseFolder.Length);
-                        var DestFile = CommandUtils.CombinePaths(BlockPath, RelativeFile);
-                        if (CommandUtils.FileExists_NoExceptions(true, DestFile))
+                    var DestFiles = new List<string>();
+                    if (ThreadsToCopyWith() < 2)
+                    {
+                        foreach (string InFilename in Files)
                         {
-                            throw new AutomationException("Dest file {0} already exists.", DestFile);
-                        }
-                        CommandUtils.CopyFile(Filename, DestFile, true);
-                        InternalUtils.Robust_FileExists_NoExceptions(true, DestFile, "Could not copy to {0}");
+                            var Filename = CommandUtils.CombinePaths(InFilename);
+                            InternalUtils.Robust_FileExists_NoExceptions(false, Filename, "Could not add {0} to manifest because it does not exist");
+
+                            if (!Filename.StartsWith(BaseFolder, StringComparison.InvariantCultureIgnoreCase))
+                            {
+                                throw new AutomationException("Could not add {0} to manifest because it does not start with the base folder {1}", Filename, BaseFolder);
+                            }
+                            var RelativeFile = Filename.Substring(BaseFolder.Length);
+                            var DestFile = CommandUtils.CombinePaths(BlockPath, RelativeFile);
+                            if (CommandUtils.FileExists_NoExceptions(true, DestFile))
+                            {
+                                throw new AutomationException("Dest file {0} already exists.", DestFile);
+                            }
+                            CommandUtils.CopyFile(Filename, DestFile, true);
+                            InternalUtils.Robust_FileExists_NoExceptions(true, DestFile, "Could not copy to {0}");
 
-                        DestFiles.Add(DestFile);
+                            DestFiles.Add(DestFile);
+                        }
                     }
-                }
-                else
-                {
-                    var SrcFiles = new List<string>();
-                    foreach (string InFilename in Files)
+                    else
                     {
-                        var Filename = CommandUtils.CombinePaths(InFilename);
-                        InternalUtils.Robust_FileExists_NoExceptions(false, Filename, "Could not add {0} to manifest because it does not exist");
-
-                        if (!Filename.StartsWith(BaseFolder, StringComparison.InvariantCultureIgnoreCase))
+                        var SrcFiles = new List<string>();
+                        foreach (string InFilename in Files)
                         {
-                            throw new AutomationException("Could not add {0} to manifest because it does not start with the base folder {1}", Filename, BaseFolder);
+                            var Filename = CommandUtils.CombinePaths(InFilename);
+                            InternalUtils.Robust_FileExists_NoExceptions(false, Filename, "Could not add {0} to manifest because it does not exist");
+
+                            if (!Filename.StartsWith(BaseFolder, StringComparison.InvariantCultureIgnoreCase))
+                            {
+                                throw new AutomationException("Could not add {0} to manifest because it does not start with the base folder {1}", Filename, BaseFolder);
+                            }
+                            var RelativeFile = Filename.Substring(BaseFolder.Length);
+                            var DestFile = CommandUtils.CombinePaths(BlockPath, RelativeFile);
+                            if (CommandUtils.FileExists_NoExceptions(true, DestFile))
+                            {
+                                throw new AutomationException("Dest file {0} already exists.", DestFile);
+                            }
+                            SrcFiles.Add(Filename);
+                            DestFiles.Add(DestFile);
                         }
-                        var RelativeFile = Filename.Substring(BaseFolder.Length);
-                        var DestFile = CommandUtils.CombinePaths(BlockPath, RelativeFile);
-                        if (CommandUtils.FileExists_NoExceptions(true, DestFile))
+                        CommandUtils.ThreadedCopyFiles(SrcFiles.ToArray(), DestFiles.ToArray(), ThreadsToCopyWith());
+                        foreach (string DestFile in DestFiles)
                         {
-                            throw new AutomationException("Dest file {0} already exists.", DestFile);
+                            InternalUtils.Robust_FileExists_NoExceptions(true, DestFile, "Could not copy to {0}");
                         }
-                        SrcFiles.Add(Filename);
-                        DestFiles.Add(DestFile);
                     }
-                    CommandUtils.ThreadedCopyFiles(SrcFiles.ToArray(), DestFiles.ToArray(), ThreadsToCopyWith());
-                    foreach (string DestFile in DestFiles)
+                    var Shared = SaveSharedTempStorageManifest(Env, StorageBlockName, GameFolder, DestFiles);
+                    if (!Local.Compare(Shared))
                     {
-                        InternalUtils.Robust_FileExists_NoExceptions(true, DestFile, "Could not copy to {0}");
+                        // we will rename this so it can't be used, but leave it around for inspection
+                        CommandUtils.RenameFile_NoExceptions(SharedTempStorageManifestFilename(Env, StorageBlockName, GameFolder), SharedTempStorageManifestFilename(Env, StorageBlockName, GameFolder) + ".broken");
+                        throw new AutomationException("Shared and Local manifest mismatch.");
                     }
-                }
-                var Shared = SaveSharedTempStorageManifest(Env, StorageBlockName, GameFolder, DestFiles);
-                if (!Local.Compare(Shared))
-                {
-                    // we will rename this so it can't be used, but leave it around for inspection
-                    CommandUtils.RenameFile_NoExceptions(SharedTempStorageManifestFilename(Env, StorageBlockName, GameFolder), SharedTempStorageManifestFilename(Env, StorageBlockName, GameFolder) + ".broken");
-                    throw new AutomationException("Shared and Local manifest mismatch.");
-                }
-                float BuildDuration = (float)((DateTime.UtcNow - StartTime).TotalSeconds);
-                if (BuildDuration > 60.0f && Shared.GetTotalSize() > 0)
-                {
-                    var MBSec = (((float)(Shared.GetTotalSize())) / (1024.0f * 1024.0f)) / BuildDuration;
-                    CommandUtils.LogConsole("Wrote to shared temp storage at {0} MB/s    {1}B {2}s", MBSec, Shared.GetTotalSize(), BuildDuration);
+                    float BuildDuration = (float)((DateTime.UtcNow - StartTime).TotalSeconds);
+                    if (BuildDuration > 10.0f && Shared.GetTotalSize() > 0)
+                    {
+                        var MBSec = (((float)(Shared.GetTotalSize())) / (1024.0f * 1024.0f)) / BuildDuration;
+                        CommandUtils.LogConsole("Wrote to shared temp storage at {0} MB/s    {1}B {2}s", MBSec, Shared.GetTotalSize(), BuildDuration);
+                    }
+                    TelemetryStopwatch.Finish(string.Format("StoreToTempStorage.{0}.{1}", Shared.GetFileCount(), Shared.GetTotalSize()));
                 }
             }
 
