Remove some GUBPNode dependencies from GUBP.cs.

diff --git a/Engine/Source/Programs/AutomationTool/GUBP/GUBP.cs b/Engine/Source/Programs/AutomationTool/GUBP/GUBP.cs
--- a/Engine/Source/Programs/AutomationTool/GUBP/GUBP.cs
+++ b/Engine/Source/Programs/AutomationTool/GUBP/GUBP.cs
@@ -378,7 +378,7 @@ public partial class GUBP : BuildCommand
 		foreach(BuildNode NodeToDo in NodesToDo)
 		{
 			string NodeStoreName = StoreName + "-" + NodeToDo.Name;
-			string GameNameIfAny = NodeToDo.Node.GameNameIfAnyForTempStorage();
+			string GameNameIfAny = NodeToDo.GameNameIfAnyForTempStorage;
 
 			if (LocalOnly)
 			{
@@ -561,7 +561,7 @@ public partial class GUBP : BuildCommand
 			}
 
             string Agent = NodeToDo.AgentRequirements;
-			if(ParseParamValue("AgentOverride") != "" && !NodeToDo.Node.GetFullName().Contains("Mac"))
+			if(ParseParamValue("AgentOverride") != "" && !NodeToDo.Name.Contains("Mac"))
 			{
 				Agent = ParseParamValue("AgentOverride");
 			}
@@ -935,7 +935,7 @@ public partial class GUBP : BuildCommand
 
         if (!(NodeToDo is TriggerNode) && CLString != "")
         {
-            string GameNameIfAny = NodeToDo.Node.GameNameIfAnyForTempStorage();
+            string GameNameIfAny = NodeToDo.GameNameIfAnyForTempStorage;
             string NodeStoreWildCard = StoreName.Replace(CLString, "*") + "-" + NodeToDo.Name;
             History = new NodeHistory();
 
@@ -1084,7 +1084,7 @@ public partial class GUBP : BuildCommand
     {
         int Result = P4Env.Changelist;
 
-        string GameNameIfAny = NodeToDo.Node.GameNameIfAnyForTempStorage();
+        string GameNameIfAny = NodeToDo.GameNameIfAnyForTempStorage;
         string NodeStore = StoreName + "-" + NodeToDo.Name + FailedTempStorageSuffix;
 
         List<int> BackwardsFails = new List<int>(History.AllFailed);
@@ -1225,7 +1225,7 @@ public partial class GUBP : BuildCommand
 		{
 			// Find all the dependencies
 			HashSet<BuildNode> Dependencies = new HashSet<BuildNode>();
-			foreach (string DependencyName in Node.Node.FullNamesOfDependencies)
+			foreach (string DependencyName in Node.DependencyNames)
 			{
 				if (!ResolveDependencies(DependencyName, AggregateNameToInfo, NodeNameToInfo, Dependencies))
 				{
@@ -1237,7 +1237,7 @@ public partial class GUBP : BuildCommand
 
 			// Find all the pseudo-dependencies
 			HashSet<BuildNode> PseudoDependencies = new HashSet<BuildNode>();
-			foreach (string PseudoDependencyName in Node.Node.FullNamesOfPseudosependencies)
+			foreach (string PseudoDependencyName in Node.PseudoDependencyNames)
 			{
 				if (!ResolveDependencies(PseudoDependencyName, AggregateNameToInfo, NodeNameToInfo, PseudoDependencies))
 				{
@@ -1653,8 +1653,8 @@ public partial class GUBP : BuildCommand
 
             string NodeStoreName = StoreName + "-" + NodeToDo.Name;
             
-            string GameNameIfAny = NodeToDo.Node.GameNameIfAnyForTempStorage();
-            string StorageRootIfAny = NodeToDo.Node.RootIfAnyForTempStorage();
+            string GameNameIfAny = NodeToDo.GameNameIfAnyForTempStorage;
+            string StorageRootIfAny = NodeToDo.RootIfAnyForTempStorage;
 					
             if (bFake)
             {
@@ -1715,7 +1715,7 @@ public partial class GUBP : BuildCommand
                     if (bFake)
                     {
 						LogConsole("***** FAKE!! Building GUBP Node {0} for {1}", NodeToDo.Name, NodeStoreName);
-                        NodeToDo.Node.DoFakeBuild(this);
+                        NodeToDo.DoFakeBuild();
                     }
                     else
                     {
@@ -1723,10 +1723,9 @@ public partial class GUBP : BuildCommand
 						DateTime StartTime = DateTime.UtcNow;
 						using(TelemetryStopwatch DoBuildStopwatch = new TelemetryStopwatch("DoBuild.{0}", NodeToDo.Name))
 						{
-							NodeToDo.Node.DoBuild(this);
+							NodeToDo.DoBuild();
 						}
 						BuildDuration = (DateTime.UtcNow - StartTime).TotalMilliseconds / 1000;
-						
                     }
 
                     TempStorage.StoreToTempStorage(CmdEnv, NodeStoreName, NodeToDo.Node.BuildProducts, !bSaveSharedTempStorage, GameNameIfAny, StorageRootIfAny);
diff --git a/Engine/Source/Programs/AutomationTool/GUBP/LegacyBranchSetup.cs b/Engine/Source/Programs/AutomationTool/GUBP/LegacyBranchSetup.cs
--- a/Engine/Source/Programs/AutomationTool/GUBP/LegacyBranchSetup.cs
+++ b/Engine/Source/Programs/AutomationTool/GUBP/LegacyBranchSetup.cs
@@ -1235,7 +1235,7 @@ partial class GUBP
 		Dictionary<string, int> FrequencyOverrides = ApplyFrequencyBarriers(BranchConfig.GUBPNodes, BranchConfig.GUBPAggregates, BranchOptions.FrequencyBarriers);
 
 		// Get the legacy node to new node mapping
-		Dictionary<GUBP.GUBPNode, BuildNode> LegacyToNewNodes = BranchConfig.GUBPNodes.Values.ToDictionary(x => x, x => x.GetBuildNode());
+		Dictionary<GUBP.GUBPNode, BuildNode> LegacyToNewNodes = BranchConfig.GUBPNodes.Values.ToDictionary(x => x, x => x.GetBuildNode(this));
 
 		// Convert the GUBPNodes and GUBPAggregates maps into lists
 		AllNodes = new List<BuildNode>();
diff --git a/Engine/Source/Programs/AutomationTool/GUBP/LegacyNodes.cs b/Engine/Source/Programs/AutomationTool/GUBP/LegacyNodes.cs
--- a/Engine/Source/Programs/AutomationTool/GUBP/LegacyNodes.cs
+++ b/Engine/Source/Programs/AutomationTool/GUBP/LegacyNodes.cs
@@ -37,9 +37,9 @@ partial class GUBP
         {
             return "";
         }
-		public virtual BuildNode GetBuildNode()
+		public virtual BuildNode GetBuildNode(GUBP bp)
 		{
-			return new LegacyBuildNode(this);
+			return new LegacyBuildNode(bp, this);
 		}
         public virtual void DoBuild(GUBP bp)
         {
@@ -245,9 +245,9 @@ partial class GUBP
 			}
             return "";
         }
-		public override BuildNode GetBuildNode()
+		public override BuildNode GetBuildNode(GUBP bp)
 		{
-			BuildNode Node = base.GetBuildNode();
+			BuildNode Node = base.GetBuildNode(bp);
 			Node.AgentPlatform = GetAgentPlatform();
 			return Node;
 		}
@@ -365,9 +365,9 @@ partial class GUBP
             }
             return false;
         }
-		public override BuildNode GetBuildNode()
+		public override BuildNode GetBuildNode(GUBP bp)
 		{
-			BuildNode Node = base.GetBuildNode();
+			BuildNode Node = base.GetBuildNode(bp);
 			if(HostPlatform == UnrealTargetPlatform.Win64)
 			{
 				Node.IsParallelAgentShareEditor = true;
@@ -467,9 +467,9 @@ partial class GUBP
             }
             return false;
         }
-		public override BuildNode GetBuildNode()
+		public override BuildNode GetBuildNode(GUBP bp)
 		{
-			BuildNode Node = base.GetBuildNode();
+			BuildNode Node = base.GetBuildNode(bp);
 			if (HostPlatform == UnrealTargetPlatform.Win64)
 			{
 				Node.IsParallelAgentShareEditor = true;
@@ -1061,9 +1061,9 @@ partial class GUBP
             }
             return false;
         }
-		public override BuildNode GetBuildNode()
+		public override BuildNode GetBuildNode(GUBP bp)
 		{
-			BuildNode Node = base.GetBuildNode();
+			BuildNode Node = base.GetBuildNode(bp);
 			if (HostPlatform == UnrealTargetPlatform.Win64)
 			{
 				Node.IsParallelAgentShareEditor = true;
@@ -1781,9 +1781,9 @@ partial class GUBP
             BuildProducts = new List<string>();
             SaveRecordOfSuccessAndAddToBuildProducts();
         }
-        public override BuildNode GetBuildNode()
+        public override BuildNode GetBuildNode(GUBP bp)
         {
-			return new TriggerNode(this);
+			return new TriggerNode(bp, this);
         }
         public virtual string GetTriggerStateName()
         {
diff --git a/Engine/Source/Programs/AutomationTool/GUBP/Nodes/BuildNode.cs b/Engine/Source/Programs/AutomationTool/GUBP/Nodes/BuildNode.cs
--- a/Engine/Source/Programs/AutomationTool/GUBP/Nodes/BuildNode.cs
+++ b/Engine/Source/Programs/AutomationTool/GUBP/Nodes/BuildNode.cs
@@ -10,7 +10,7 @@ namespace AutomationTool
 	[DebuggerDisplay("{Name}")]
 	public abstract class BuildNode
 	{
-		public string Name;
+		public readonly string Name;
 		public UnrealTargetPlatform AgentPlatform = UnrealTargetPlatform.Win64;
 		public string AgentRequirements;
 		public string AgentSharingGroup;
@@ -31,11 +31,30 @@ namespace AutomationTool
 		public bool SendSuccessEmail;
 		public bool IsParallelAgentShareEditor;
 
+		public BuildNode(string InName)
+		{
+			Name = InName;
+		}
+
 		public abstract bool IsSticky
 		{
 			get;
 		}
 
+		public abstract IEnumerable<string> DependencyNames
+		{
+			get;
+		}
+
+		public abstract IEnumerable<string> PseudoDependencyNames
+		{
+			get;
+		}
+
+		public abstract void DoBuild();
+
+		public abstract void DoFakeBuild();
+
 		public string ControllingTriggerDotName
 		{
 			get { return String.Join(".", ControllingTriggers.Select(x => x.Name)); }
@@ -63,13 +82,25 @@ namespace AutomationTool
 		{
 			get;
 		}
+
+		public abstract string GameNameIfAnyForTempStorage
+		{
+			get;
+		}
+
+		public abstract string RootIfAnyForTempStorage
+		{
+			get;
+		}
 	}
 
 	public class LegacyBuildNode : BuildNode
 	{
-		public LegacyBuildNode(GUBP.GUBPNode InNode)
+		GUBP Owner;
+
+		public LegacyBuildNode(GUBP InOwner, GUBP.GUBPNode InNode) : base(InNode.GetFullName())
 		{
-			Name = InNode.GetFullName();
+			Owner = InOwner;
 			Node = InNode;
 			AgentRequirements = Node.ECAgentString();
 			AgentSharingGroup = Node.AgentSharingGroup;
@@ -84,6 +115,26 @@ namespace AutomationTool
 			get { return Node.IsSticky(); }
 		}
 
+		public override IEnumerable<string> DependencyNames
+		{
+			get { return Node.FullNamesOfDependencies; }
+		}
+
+		public override IEnumerable<string> PseudoDependencyNames
+		{
+			get { return Node.FullNamesOfPseudosependencies; }
+		}
+
+		public override void DoBuild()
+		{
+			Node.DoBuild(Owner);
+		}
+
+		public override void DoFakeBuild()
+		{
+			Node.DoFakeBuild(Owner);
+		}
+
 		public override bool IsTest
 		{
 			get { return Node.IsTest(); }
@@ -93,5 +144,15 @@ namespace AutomationTool
 		{
 			get { return Node.GetDisplayGroupName(); }
 		}
+
+		public override string GameNameIfAnyForTempStorage
+		{
+			get { return Node.GameNameIfAnyForTempStorage(); }
+		}
+
+		public override string RootIfAnyForTempStorage
+		{
+			get { return Node.RootIfAnyForTempStorage(); }
+		}
 	}
 }
diff --git a/Engine/Source/Programs/AutomationTool/GUBP/Nodes/TriggerNode.cs b/Engine/Source/Programs/AutomationTool/GUBP/Nodes/TriggerNode.cs
--- a/Engine/Source/Programs/AutomationTool/GUBP/Nodes/TriggerNode.cs
+++ b/Engine/Source/Programs/AutomationTool/GUBP/Nodes/TriggerNode.cs
@@ -14,7 +14,7 @@ namespace AutomationTool
 		public string DescriptionText;
 		public string ActionText;
 
-		public TriggerNode(GUBP.WaitForUserInput InNode) : base(InNode)
+		public TriggerNode(GUBP bp, GUBP.WaitForUserInput InNode) : base(bp, InNode)
 		{
 			AddSubmittersToFailureEmails = false;
 
