Fix a couple of annoying variable names with spelling mistakes ("FullNamesOfPseudosependencies" and "OrdereredToDo")

diff --git a/Engine/Source/Programs/AutomationTool/GUBP/ElectricCommander.cs b/Engine/Source/Programs/AutomationTool/GUBP/ElectricCommander.cs
--- a/Engine/Source/Programs/AutomationTool/GUBP/ElectricCommander.cs
+++ b/Engine/Source/Programs/AutomationTool/GUBP/ElectricCommander.cs
@@ -255,10 +255,10 @@ namespace AutomationTool
 			}
 		}
 
-		public void DoCommanderSetup(IEnumerable<BuildNode> AllNodes, IEnumerable<AggregateNode> AllAggregates, List<BuildNode> OrdereredToDo, List<BuildNode> SortedNodes, int TimeIndex, int TimeQuantum, bool bSkipTriggers, bool bFake, bool bFakeEC, string CLString, TriggerNode ExplicitTrigger, List<TriggerNode> UnfinishedTriggers, string FakeFail)
+		public void DoCommanderSetup(IEnumerable<BuildNode> AllNodes, IEnumerable<AggregateNode> AllAggregates, List<BuildNode> OrderedToDo, List<BuildNode> SortedNodes, int TimeIndex, int TimeQuantum, bool bSkipTriggers, bool bFake, bool bFakeEC, string CLString, TriggerNode ExplicitTrigger, List<TriggerNode> UnfinishedTriggers, string FakeFail)
 		{
 			// Check there's some nodes in the graph
-			if (OrdereredToDo.Count == 0)
+			if (OrderedToDo.Count == 0)
 			{
 				throw new AutomationException("No nodes to do!");
 			}
@@ -266,7 +266,7 @@ namespace AutomationTool
 			// Make sure everything before the explicit trigger is completed.
 			if (ExplicitTrigger != null)
 			{
-				foreach (BuildNode NodeToDo in OrdereredToDo)
+				foreach (BuildNode NodeToDo in OrderedToDo)
 				{
 					if (!NodeToDo.IsComplete && NodeToDo != ExplicitTrigger && !NodeToDo.DependsOn(ExplicitTrigger)) // if something is already finished, we don't put it into EC
 					{
@@ -276,13 +276,13 @@ namespace AutomationTool
 			}
 
 			// Update all the EC properties, and the branch definition files
-			WriteProperties(AllNodes, AllAggregates, OrdereredToDo, SortedNodes, TimeIndex, TimeQuantum);
+			WriteProperties(AllNodes, AllAggregates, OrderedToDo, SortedNodes, TimeIndex, TimeQuantum);
 
 			// Write all the job setup
-			WriteJobSteps(OrdereredToDo, bSkipTriggers);
+			WriteJobSteps(OrderedToDo, bSkipTriggers);
 		}
 
-		private void WriteProperties(IEnumerable<BuildNode> AllNodes, IEnumerable<AggregateNode> AllAggregates, List<BuildNode> OrdereredToDo, List<BuildNode> SortedNodes, int TimeIndex, int TimeQuantum)
+		private void WriteProperties(IEnumerable<BuildNode> AllNodes, IEnumerable<AggregateNode> AllAggregates, List<BuildNode> OrderedToDo, List<BuildNode> SortedNodes, int TimeIndex, int TimeQuantum)
 		{
 			List<AggregateNode> SeparatePromotables = FindPromotables(AllAggregates);
 			Dictionary<BuildNode, List<AggregateNode>> DependentPromotions = FindDependentPromotables(AllNodes, SeparatePromotables);
@@ -308,7 +308,7 @@ namespace AutomationTool
 				ECProps.Add(string.Format("PossiblePromotables/{0}={1}", Node.Name, ""));
 			}
 
-			foreach (BuildNode NodeToDo in OrdereredToDo)
+			foreach (BuildNode NodeToDo in OrderedToDo)
 			{
 				if (!NodeToDo.IsComplete) // if something is already finished, we don't put it into EC  
 				{
@@ -327,17 +327,17 @@ namespace AutomationTool
 			CommandUtils.WriteAllLines(BranchJobDefFile, ECProps.ToArray());
 			RunECTool(String.Format("setProperty \"/myJob/BranchJobDefFile\" \"{0}\"", BranchJobDefFile.Replace("\\", "\\\\")));
 
-			bool bHasTests = OrdereredToDo.Any(x => x.IsTest);
+			bool bHasTests = OrderedToDo.Any(x => x.IsTest);
 			RunECTool(String.Format("setProperty \"/myWorkflow/HasTests\" \"{0}\"", bHasTests));
 		}
 
-		private void WriteJobSteps(List<BuildNode> OrdereredToDo, bool bSkipTriggers)
+		private void WriteJobSteps(List<BuildNode> OrderedToDo, bool bSkipTriggers)
 		{
 			BuildNode LastSticky = null;
 			bool HitNonSticky = false;
 			bool bHaveECNodes = false;
 			// sticky nodes are ones that we run on the main agent. We run then first and they must not be intermixed with parallel jobs
-			foreach (BuildNode NodeToDo in OrdereredToDo)
+			foreach (BuildNode NodeToDo in OrderedToDo)
 			{
 				if (!NodeToDo.IsComplete) // if something is already finished, we don't put it into EC
 				{
@@ -376,7 +376,7 @@ namespace AutomationTool
 
 				Dictionary<string, List<BuildNode>> AgentGroupChains = new Dictionary<string, List<BuildNode>>();
 				List<BuildNode> StickyChain = new List<BuildNode>();
-				foreach (BuildNode NodeToDo in OrdereredToDo)
+				foreach (BuildNode NodeToDo in OrderedToDo)
 				{
 					if (!NodeToDo.IsComplete) // if something is already finished, we don't put it into EC  
 					{
@@ -401,7 +401,7 @@ namespace AutomationTool
 						}
 					}
 				}
-				foreach (BuildNode NodeToDo in OrdereredToDo)
+				foreach (BuildNode NodeToDo in OrderedToDo)
 				{
 					if (!NodeToDo.IsComplete) // if something is already finished, we don't put it into EC  
 					{
@@ -433,9 +433,9 @@ namespace AutomationTool
 						List<BuildNode> UncompletedEcDeps = new List<BuildNode>();
 						foreach (BuildNode Dep in NodeToDo.AllDirectDependencies)
 						{
-							if (!Dep.IsComplete && OrdereredToDo.Contains(Dep)) // if something is already finished, we don't put it into EC
+							if (!Dep.IsComplete && OrderedToDo.Contains(Dep)) // if something is already finished, we don't put it into EC
 							{
-								if (OrdereredToDo.IndexOf(Dep) > OrdereredToDo.IndexOf(NodeToDo))
+								if (OrderedToDo.IndexOf(Dep) > OrderedToDo.IndexOf(NodeToDo))
 								{
 									throw new AutomationException("Topological sort error, node {0} has a dependency of {1} which sorted after it.", NodeToDo.Name, Dep.Name);
 								}
@@ -443,7 +443,7 @@ namespace AutomationTool
 							}
 						}
 
-						string PreCondition = GetPreConditionForNode(OrdereredToDo, ParentPath, bHasNoop, AgentGroupChains, StickyChain, NodeToDo, UncompletedEcDeps);
+						string PreCondition = GetPreConditionForNode(OrderedToDo, ParentPath, bHasNoop, AgentGroupChains, StickyChain, NodeToDo, UncompletedEcDeps);
 						string RunCondition = GetRunConditionForNode(UncompletedEcDeps, ParentPath);
 
 						// Create the job steps for this node
@@ -564,7 +564,7 @@ namespace AutomationTool
 			}
 		}
 
-		private string GetPreConditionForNode(List<BuildNode> OrdereredToDo, string PreconditionParentPath, bool bHasNoop, Dictionary<string, List<BuildNode>> AgentGroupChains, List<BuildNode> StickyChain, BuildNode NodeToDo, List<BuildNode> UncompletedEcDeps)
+		private string GetPreConditionForNode(List<BuildNode> OrderedToDo, string PreconditionParentPath, bool bHasNoop, Dictionary<string, List<BuildNode>> AgentGroupChains, List<BuildNode> StickyChain, BuildNode NodeToDo, List<BuildNode> UncompletedEcDeps)
 		{
 			List<BuildNode> PreConditionUncompletedEcDeps = new List<BuildNode>();
 			if(NodeToDo.AgentSharingGroup == "")
@@ -586,9 +586,9 @@ namespace AutomationTool
 					{
 						foreach (BuildNode Dep in Chain.AllDirectDependencies)
 						{
-							if (!Dep.IsComplete && OrdereredToDo.Contains(Dep)) // if something is already finished, we don't put it into EC
+							if (!Dep.IsComplete && OrderedToDo.Contains(Dep)) // if something is already finished, we don't put it into EC
 							{
-								if (OrdereredToDo.IndexOf(Dep) > OrdereredToDo.IndexOf(Chain))
+								if (OrderedToDo.IndexOf(Dep) > OrderedToDo.IndexOf(Chain))
 								{
 									throw new AutomationException("Topological sort error, node {0} has a dependency of {1} which sorted after it.", Chain.Name, Dep.Name);
 								}
@@ -616,9 +616,9 @@ namespace AutomationTool
 				{
 					foreach (BuildNode Dep in NodeToDo.AllDirectDependencies)
 					{
-						if (!Dep.IsComplete && OrdereredToDo.Contains(Dep)) // if something is already finished, we don't put it into EC
+						if (!Dep.IsComplete && OrderedToDo.Contains(Dep)) // if something is already finished, we don't put it into EC
 						{
-							if (OrdereredToDo.IndexOf(Dep) > OrdereredToDo.IndexOf(NodeToDo))
+							if (OrderedToDo.IndexOf(Dep) > OrderedToDo.IndexOf(NodeToDo))
 							{
 								throw new AutomationException("Topological sort error, node {0} has a dependency of {1} which sorted after it.", NodeToDo.Name, Dep.Name);
 							}
diff --git a/Engine/Source/Programs/AutomationTool/GUBP/GUBP.cs b/Engine/Source/Programs/AutomationTool/GUBP/GUBP.cs
--- a/Engine/Source/Programs/AutomationTool/GUBP/GUBP.cs
+++ b/Engine/Source/Programs/AutomationTool/GUBP/GUBP.cs
@@ -698,7 +698,7 @@ public partial class GUBP : BuildCommand
     {
         DateTime StartTime = DateTime.UtcNow;
 
-        List<BuildNode> OrdereredToDo = new List<BuildNode>();
+        List<BuildNode> OrderedToDo = new List<BuildNode>();
 
         Dictionary<string, List<BuildNode>> SortedAgentGroupChains = new Dictionary<string, List<BuildNode>>();
         if (!SubSort)
@@ -871,13 +871,13 @@ public partial class GUBP : BuildCommand
                 {
                     foreach (BuildNode ChainNode in SortedAgentGroupChains[BestNode.AgentSharingGroup])
                     {
-                        OrdereredToDo.Add(ChainNode);
+                        OrderedToDo.Add(ChainNode);
                         NodesToDo.Remove(ChainNode);
                     }
                 }
                 else
                 {
-                    OrdereredToDo.Add(BestNode);
+                    OrderedToDo.Add(BestNode);
                     NodesToDo.Remove(BestNode);
                 }
             }
@@ -923,10 +923,10 @@ public partial class GUBP : BuildCommand
         if (!SubSort)
         {
             double BuildDuration = (DateTime.UtcNow - StartTime).TotalMilliseconds;
-			LogConsole("Took {0}s to sort {1} nodes", BuildDuration / 1000, OrdereredToDo.Count);
+			LogConsole("Took {0}s to sort {1} nodes", BuildDuration / 1000, OrderedToDo.Count);
         }
 
-        return OrdereredToDo;
+        return OrderedToDo;
     }
 
     static NodeHistory FindNodeHistory(BuildNode NodeToDo, string CLString, string StoreName)
@@ -1517,13 +1517,13 @@ public partial class GUBP : BuildCommand
 		throw new AutomationException("Failed to update the CIS counter after 20 tries.");
 	}
 
-	private List<TriggerNode> FindUnfinishedTriggers(bool bSkipTriggers, BuildNode ExplicitTrigger, List<BuildNode> OrdereredToDo)
+	private List<TriggerNode> FindUnfinishedTriggers(bool bSkipTriggers, BuildNode ExplicitTrigger, List<BuildNode> OrderedToDo)
 	{
 		// find all unfinished triggers, excepting the one we are triggering right now
 		List<TriggerNode> UnfinishedTriggers = new List<TriggerNode>();
 		if (!bSkipTriggers)
 		{
-			foreach (TriggerNode NodeToDo in OrdereredToDo.OfType<TriggerNode>())
+			foreach (TriggerNode NodeToDo in OrderedToDo.OfType<TriggerNode>())
 			{
 				if (!NodeToDo.IsComplete && ExplicitTrigger != NodeToDo)
 				{
@@ -1537,19 +1537,19 @@ public partial class GUBP : BuildCommand
 	/// <summary>
 	/// Validates that the given nodes are sorted correctly, so that all dependencies are met first
 	/// </summary>
-	/// <param name="OrdereredToDo">The sorted list of nodes</param>
-	private void CheckSortOrder(List<BuildNode> OrdereredToDo)
+	/// <param name="OrderedToDo">The sorted list of nodes</param>
+	private void CheckSortOrder(List<BuildNode> OrderedToDo)
 	{
-		foreach (BuildNode NodeToDo in OrdereredToDo)
+		foreach (BuildNode NodeToDo in OrderedToDo)
 		{
 			if ((NodeToDo is TriggerNode) && (NodeToDo.IsSticky || NodeToDo.IsComplete)) // these sticky triggers are ok, everything is already completed anyway
 			{
 				continue;
 			}
-			int MyIndex = OrdereredToDo.IndexOf(NodeToDo);
+			int MyIndex = OrderedToDo.IndexOf(NodeToDo);
 			foreach (BuildNode Dep in NodeToDo.Dependencies)
 			{
-				int DepIndex = OrdereredToDo.IndexOf(Dep);
+				int DepIndex = OrderedToDo.IndexOf(Dep);
 				if (DepIndex >= MyIndex)
 				{
 					throw new AutomationException("Topological sort error, node {0} has a dependency of {1} which sorted after it.", NodeToDo.Name, Dep.Name);
@@ -1557,7 +1557,7 @@ public partial class GUBP : BuildCommand
 			}
 			foreach (BuildNode Dep in NodeToDo.PseudoDependencies)
 			{
-				int DepIndex = OrdereredToDo.IndexOf(Dep);
+				int DepIndex = OrderedToDo.IndexOf(Dep);
 				if (DepIndex >= MyIndex)
 				{
 					throw new AutomationException("Topological sort error, node {0} has a pseduodependency of {1} which sorted after it.", NodeToDo.Name, Dep.Name);
@@ -1566,16 +1566,16 @@ public partial class GUBP : BuildCommand
 		}
 	}
 
-	private void DoCommanderSetup(ElectricCommander EC, IEnumerable<BuildNode> AllNodes, IEnumerable<AggregateNode> AllAggregates, List<BuildNode> OrdereredToDo, int TimeIndex, int TimeQuantum, bool bSkipTriggers, bool bFake, bool bFakeEC, string CLString, TriggerNode ExplicitTrigger, List<TriggerNode> UnfinishedTriggers, string FakeFail, bool bPreflightBuild)
+	private void DoCommanderSetup(ElectricCommander EC, IEnumerable<BuildNode> AllNodes, IEnumerable<AggregateNode> AllAggregates, List<BuildNode> OrderedToDo, int TimeIndex, int TimeQuantum, bool bSkipTriggers, bool bFake, bool bFakeEC, string CLString, TriggerNode ExplicitTrigger, List<TriggerNode> UnfinishedTriggers, string FakeFail, bool bPreflightBuild)
 	{
 		List<BuildNode> SortedNodes = TopologicalSort(new HashSet<BuildNode>(AllNodes), null, SubSort: false, DoNotConsiderCompletion: true);
 		LogConsole("******* {0} GUBP Nodes", SortedNodes.Count);
 
-		List<BuildNode> FilteredOrdereredToDo = new List<BuildNode>();
+		List<BuildNode> FilteredOrderedToDo = new List<BuildNode>();
 		using(TelemetryStopwatch StartFilterTimer = new TelemetryStopwatch("FilterNodes"))
 		{
 			// remove nodes that have unfinished triggers
-			foreach (BuildNode NodeToDo in OrdereredToDo)
+			foreach (BuildNode NodeToDo in OrderedToDo)
 			{
 				if (NodeToDo.ControllingTriggers.Length == 0 || !UnfinishedTriggers.Contains(NodeToDo.ControllingTriggers.Last()))
 				{
@@ -1591,23 +1591,23 @@ public partial class GUBP : BuildCommand
 						continue;
 					}
 
-					FilteredOrdereredToDo.Add(NodeToDo);
+					FilteredOrderedToDo.Add(NodeToDo);
 				}
 			}
 		}
 		using(TelemetryStopwatch PrintNodesTimer = new TelemetryStopwatch("SetupCommanderPrint"))
 		{
 			LogConsole("*********** EC Nodes, in order.");
-			PrintNodes(this, FilteredOrdereredToDo, AllAggregates, UnfinishedTriggers, TimeQuantum);
+			PrintNodes(this, FilteredOrderedToDo, AllAggregates, UnfinishedTriggers, TimeQuantum);
 		}
 
-		EC.DoCommanderSetup(AllNodes, AllAggregates, FilteredOrdereredToDo, SortedNodes, TimeIndex, TimeQuantum, bSkipTriggers, bFake, bFakeEC, CLString, ExplicitTrigger, UnfinishedTriggers, FakeFail);
+		EC.DoCommanderSetup(AllNodes, AllAggregates, FilteredOrderedToDo, SortedNodes, TimeIndex, TimeQuantum, bSkipTriggers, bFake, bFakeEC, CLString, ExplicitTrigger, UnfinishedTriggers, FakeFail);
 	}
 
-	void ExecuteNodes(ElectricCommander EC, List<BuildNode> OrdereredToDo, bool bFake, bool bFakeEC, bool bSaveSharedTempStorage, string CLString, string StoreName, string FakeFail)
+	void ExecuteNodes(ElectricCommander EC, List<BuildNode> OrderedToDo, bool bFake, bool bFakeEC, bool bSaveSharedTempStorage, string CLString, string StoreName, string FakeFail)
 	{
         Dictionary<string, BuildNode> BuildProductToNodeMap = new Dictionary<string, BuildNode>();
-		foreach (BuildNode NodeToDo in OrdereredToDo)
+		foreach (BuildNode NodeToDo in OrderedToDo)
         {
             if (NodeToDo.Node.BuildProducts != null || NodeToDo.Node.AllDependencyBuildProducts != null)
             {
diff --git a/Engine/Source/Programs/AutomationTool/GUBP/LegacyBranchSetup.cs b/Engine/Source/Programs/AutomationTool/GUBP/LegacyBranchSetup.cs
--- a/Engine/Source/Programs/AutomationTool/GUBP/LegacyBranchSetup.cs
+++ b/Engine/Source/Programs/AutomationTool/GUBP/LegacyBranchSetup.cs
@@ -1206,7 +1206,7 @@ partial class GUBP
 			GUBPNode Node = BranchConfig.TryFindNode(NodeToRemovePseudoDependencies);
 			if(Node != null)
 			{
-				Node.FullNamesOfPseudosependencies.Clear();
+				Node.FullNamesOfPseudodependencies.Clear();
 			}
 		}
 
@@ -1270,7 +1270,7 @@ partial class GUBP
 					if(GUBPNodes.TryGetValue(SearchNodes[Idx], out Node))
 					{
 						SearchNodes.AddRange(Node.FullNamesOfDependencies);
-						SearchNodes.AddRange(Node.FullNamesOfPseudosependencies);
+						SearchNodes.AddRange(Node.FullNamesOfPseudodependencies);
 					}
 					else
 					{
diff --git a/Engine/Source/Programs/AutomationTool/GUBP/LegacyNodes.cs b/Engine/Source/Programs/AutomationTool/GUBP/LegacyNodes.cs
--- a/Engine/Source/Programs/AutomationTool/GUBP/LegacyNodes.cs
+++ b/Engine/Source/Programs/AutomationTool/GUBP/LegacyNodes.cs
@@ -15,7 +15,7 @@ partial class GUBP
     public abstract class GUBPNode
     {
         public List<string> FullNamesOfDependencies = new List<string>();
-        public List<string> FullNamesOfPseudosependencies = new List<string>(); //these are really only used for sorting. We want the editor to fail before the monolithics. Think of it as "can't possibly be useful without".
+        public List<string> FullNamesOfPseudodependencies = new List<string>(); //these are really only used for sorting. We want the editor to fail before the monolithics. Think of it as "can't possibly be useful without".
         public List<string> BuildProducts = null;
         public List<string> AllDependencyBuildProducts = null;
         public List<string> AllDependencies = null;		
@@ -142,16 +142,16 @@ partial class GUBP
         }
         public void AddPseudodependency(string Node)
         {
-            if (!FullNamesOfPseudosependencies.Contains(Node))
+            if (!FullNamesOfPseudodependencies.Contains(Node))
             {
-                FullNamesOfPseudosependencies.Add(Node);
+                FullNamesOfPseudodependencies.Add(Node);
             }
         }
         public void RemovePseudodependency(string Node)
         {
-            if (FullNamesOfPseudosependencies.Contains(Node))
+            if (FullNamesOfPseudodependencies.Contains(Node))
             {
-                FullNamesOfPseudosependencies.Remove(Node);
+                FullNamesOfPseudodependencies.Remove(Node);
             }
         }
         public void AddBuildProduct(string Filename)
diff --git a/Engine/Source/Programs/AutomationTool/GUBP/Nodes/BuildNode.cs b/Engine/Source/Programs/AutomationTool/GUBP/Nodes/BuildNode.cs
--- a/Engine/Source/Programs/AutomationTool/GUBP/Nodes/BuildNode.cs
+++ b/Engine/Source/Programs/AutomationTool/GUBP/Nodes/BuildNode.cs
@@ -122,7 +122,7 @@ namespace AutomationTool
 
 		public override IEnumerable<string> PseudoDependencyNames
 		{
-			get { return Node.FullNamesOfPseudosependencies; }
+			get { return Node.FullNamesOfPseudodependencies; }
 		}
 
 		public override void DoBuild()
