Store the full list of input and ordering dependenices for each node, rather than just the ones which are referenced explicitly.

diff --git a/Engine/Source/Programs/AutomationTool/GUBP/ElectricCommander.cs b/Engine/Source/Programs/AutomationTool/GUBP/ElectricCommander.cs
--- a/Engine/Source/Programs/AutomationTool/GUBP/ElectricCommander.cs
+++ b/Engine/Source/Programs/AutomationTool/GUBP/ElectricCommander.cs
@@ -431,7 +431,7 @@ namespace AutomationTool
 						bool DoParallel = !Sticky || NodeToDo.IsParallelAgentShareEditor;
 
 						List<BuildNode> UncompletedEcDeps = new List<BuildNode>();
-						foreach (BuildNode Dep in NodeToDo.AllDirectDependencies)
+						foreach (BuildNode Dep in FindDirectOrderDependencies(NodeToDo))
 						{
 							if (!Dep.IsComplete && OrderedToDo.Contains(Dep)) // if something is already finished, we don't put it into EC
 							{
@@ -584,7 +584,7 @@ namespace AutomationTool
 					// to avoid idle agents (and also EC doesn't actually reserve our agent!), we promote all dependencies to the first one
 					foreach (BuildNode Chain in MyChain)
 					{
-						foreach (BuildNode Dep in Chain.AllDirectDependencies)
+						foreach (BuildNode Dep in FindDirectOrderDependencies(Chain))
 						{
 							if (!Dep.IsComplete && OrderedToDo.Contains(Dep)) // if something is already finished, we don't put it into EC
 							{
@@ -614,7 +614,7 @@ namespace AutomationTool
 				}
 				else
 				{
-					foreach (BuildNode Dep in NodeToDo.AllDirectDependencies)
+					foreach (BuildNode Dep in FindDirectOrderDependencies(NodeToDo))
 					{
 						if (!Dep.IsComplete && OrderedToDo.Contains(Dep)) // if something is already finished, we don't put it into EC
 						{
@@ -649,6 +649,16 @@ namespace AutomationTool
 			return PreCondition;
 		}
 
+		private static HashSet<BuildNode> FindDirectOrderDependencies(BuildNode Node)
+		{
+			HashSet<BuildNode> DirectDependencies = new HashSet<BuildNode>(Node.OrderDependencies);
+			foreach(BuildNode OrderDependency in Node.OrderDependencies)
+			{
+				DirectDependencies.ExceptWith(OrderDependency.OrderDependencies);
+			}
+			return DirectDependencies;
+		}
+
 		private static string GetNodeForAllNodesProperty(BuildNode Node, int TimeQuantum)
 		{
 			string Note = Node.ControllingTriggerDotName;
@@ -712,7 +722,7 @@ namespace AutomationTool
 			Dictionary<BuildNode, List<AggregateNode>> DependentPromotions = NodesToDo.ToDictionary(x => x, x => new List<AggregateNode>());
 			foreach (AggregateNode SeparatePromotion in SeparatePromotions)
 			{
-				BuildNode[] Dependencies = SeparatePromotion.Dependencies.SelectMany(x => x.AllIndirectDependencies).Distinct().ToArray();
+				BuildNode[] Dependencies = SeparatePromotion.Dependencies.SelectMany(x => x.OrderDependencies).Distinct().ToArray();
 				foreach (BuildNode Dependency in Dependencies)
 				{
 					DependentPromotions[Dependency].Add(SeparatePromotion);
@@ -765,7 +775,7 @@ namespace AutomationTool
 				}
 
 				// Build a list of ordering dependencies, putting all Mac nodes after Windows nodes with the same names.
-				Dictionary<BuildNode, List<BuildNode>> NodeDependencies = new Dictionary<BuildNode,List<BuildNode>>(Nodes.ToDictionary(x => x, x => x.AllDirectDependencies.ToList()));
+				Dictionary<BuildNode, List<BuildNode>> NodeDependencies = new Dictionary<BuildNode,List<BuildNode>>(Nodes.ToDictionary(x => x, x => x.OrderDependencies.ToList()));
 				foreach(KeyValuePair<string, List<BuildNode>> DisplayGroup in DisplayGroups)
 				{
 					List<BuildNode> GroupNodes = DisplayGroup.Value;
diff --git a/Engine/Source/Programs/AutomationTool/GUBP/GUBP.cs b/Engine/Source/Programs/AutomationTool/GUBP/GUBP.cs
--- a/Engine/Source/Programs/AutomationTool/GUBP/GUBP.cs
+++ b/Engine/Source/Programs/AutomationTool/GUBP/GUBP.cs
@@ -26,7 +26,6 @@ using System.Diagnostics;
 [Help("TriggerNode=", "Trigger Nodes to process, -triggernode=Node.")]
 [Help("Game=", "Games to process, -game=Game1+Game2+Game3, if no games or nodes are specified, defaults to all nodes.")]
 [Help("ListOnly", "List Nodes in this branch")]
-[Help("SaveGraph", "Save graph as an xml file")]
 [Help("CommanderJobSetupOnly", "Set up the EC branch info via ectool and quit")]
 [Help("FakeEC", "don't run ectool, rather just do it locally, emulating what EC would have done.")]
 [Help("Fake", "Don't actually build anything, just store a record of success as the build product for each node.")]
@@ -264,10 +263,6 @@ public partial class GUBP : BuildCommand
 			{
 				DoCommanderSetup(EC, AllNodes, AllAggregates, OrderedToDo, TimeIndex, TimeQuantum, bSkipTriggers, bFake, bFakeEC, CLString, ExplicitTrigger, UnfinishedTriggers, FakeFail, bPreflightBuild);
 			}
-			else if(ParseParam("SaveGraph"))
-			{
-				SaveGraphVisualization(OrderedToDo);
-			}
 			else if(!ParseParam("ListOnly"))
 			{
 				ExecuteNodes(EC, OrderedToDo, bFake, bFakeEC, bSaveSharedTempStorage, CLString, StoreName, FakeFail);
@@ -298,14 +293,9 @@ public partial class GUBP : BuildCommand
 		{
 			NodeToDo.ControllingTriggers = new TriggerNode[0];
 
-			// Find all the dependencies of this node
-			List<BuildNode> AllDependencies = new List<BuildNode>();
-			AllDependencies.AddRange(NodeToDo.Dependencies);
-			AllDependencies.AddRange(NodeToDo.PseudoDependencies);
-
 			// Find the immediate trigger controlling this one
 			List<TriggerNode> PreviousTriggers = new List<TriggerNode>();
-			foreach (BuildNode Dependency in AllDependencies)
+			foreach (BuildNode Dependency in NodeToDo.OrderDependencies)
 			{
 				FindControllingTriggers(Dependency);
 
@@ -366,9 +356,9 @@ public partial class GUBP : BuildCommand
 	{
 		foreach(BuildNode Node in Nodes)
 		{
-			foreach(BuildNode IndirectDependency in Node.AllIndirectDependencies)
+			foreach(BuildNode OrderDependency in Node.OrderDependencies)
 			{
-				Node.FrequencyShift = Math.Max(Node.FrequencyShift, IndirectDependency.FrequencyShift);
+				Node.FrequencyShift = Math.Max(Node.FrequencyShift, OrderDependency.FrequencyShift);
 			}
 		}
 	}
@@ -581,84 +571,18 @@ public partial class GUBP : BuildCommand
 
             if (bShowDependencies)
             {
-                foreach (BuildNode Dep in NodeToDo.Dependencies)
+				foreach (BuildNode Dep in NodeToDo.InputDependencies.OrderBy(x => x.Name))
                 {
 					LogConsole("            dep> {0}", Dep.Name);
                 }
-                foreach (BuildNode Dep in NodeToDo.PseudoDependencies)
-                {
+				foreach (BuildNode Dep in NodeToDo.OrderDependencies.OrderBy(x => x.Name))
+				{
 					LogConsole("           pdep> {0}", Dep.Name);
-                }
+				}
             }
         }
     }
 
-	/// <summary>
-	/// Exports the build graph as a GEXF file, for visualization in an external tool (eg. Gephi).
-	/// </summary>
-	/// <param name="Nodes">The nodes in the graph</param>
-	static void SaveGraphVisualization(IEnumerable<BuildNode> Nodes)
-	{
-		// Create a graph node for each GUBP node in the graph
-		List<GraphNode> GraphNodes = new List<GraphNode>();
-		Dictionary<BuildNode, GraphNode> NodeToGraphNodeMap = new Dictionary<BuildNode, GraphNode>();
-		foreach(BuildNode Node in Nodes)
-		{
-			GraphNode GraphNode = new GraphNode();
-			GraphNode.Id = NodeToGraphNodeMap.Count;
-			GraphNode.Label = Node.Name;
-
-			NodeToGraphNodeMap.Add(Node, GraphNode);
-			GraphNodes.Add(GraphNode);
-		}
-
-		// Connect everything together
-		List<GraphEdge> GraphEdges = new List<GraphEdge>();
-		foreach(KeyValuePair<BuildNode, GraphNode> NodeToGraphNodePair in NodeToGraphNodeMap)
-		{
-			foreach (BuildNode Dependency in NodeToGraphNodePair.Key.Dependencies)
-			{
-				GraphNode PrerequisiteFileGraphNode;
-				if (NodeToGraphNodeMap.TryGetValue(Dependency, out PrerequisiteFileGraphNode))
-				{
-					// Connect a file our action is dependent on, to our action itself
-					GraphEdge NewGraphEdge = new GraphEdge();
-
-					NewGraphEdge.Id = GraphEdges.Count;
-					NewGraphEdge.Source = PrerequisiteFileGraphNode;
-					NewGraphEdge.Target = NodeToGraphNodePair.Value;
-					NewGraphEdge.Color = new GraphColor() { R = 0.0f, G = 0.0f, B = 0.0f, A = 0.75f };
-
-					GraphEdges.Add(NewGraphEdge);
-				}
-
-			}
-			foreach (BuildNode Dependency in NodeToGraphNodePair.Key.PseudoDependencies)
-			{
-				GraphNode PrerequisiteFileGraphNode;
-				if (NodeToGraphNodeMap.TryGetValue(Dependency, out PrerequisiteFileGraphNode))
-				{
-					// Connect a file our action is dependent on, to our action itself
-					GraphEdge NewGraphEdge = new GraphEdge();
-
-					NewGraphEdge.Id = GraphEdges.Count;
-					NewGraphEdge.Source = PrerequisiteFileGraphNode;
-					NewGraphEdge.Target = NodeToGraphNodePair.Value;
-					NewGraphEdge.Color = new GraphColor() { R = 0.0f, G = 0.0f, B = 0.0f, A = 0.25f };
-
-					GraphEdges.Add(NewGraphEdge);
-				}
-
-			}
-		}
-
-		// Export the graph definition
-		string Filename = CommandUtils.CombinePaths(CommandUtils.CmdEnv.LogFolder, "GubpGraph.gexf");
-		LogConsole("Writing graph to {0}", Filename);
-		GraphVisualization.WriteGraphFile(Filename, "GUBP Nodes", NodeToGraphNodeMap.Values.ToList(), GraphEdges);
-		LogConsole("Wrote graph to {0}", Filename);
-	 }
-
     static List<int> ConvertCLToIntList(List<string> Strings)
     {
         List<int> Result = new List<int>();
@@ -762,7 +686,7 @@ public partial class GUBP : BuildCommand
                         ExaminedAgentGroups.Add(NodeToDo.AgentSharingGroup);
                         foreach (BuildNode ChainNode in SortedAgentGroupChains[NodeToDo.AgentSharingGroup])
                         {
-                            foreach (BuildNode Dep in ChainNode.Dependencies)
+                            foreach (BuildNode Dep in ChainNode.InputDependencies)
                             {
                                 if (!SortedAgentGroupChains[NodeToDo.AgentSharingGroup].Contains(Dep) && NodesToDo.Contains(Dep))
                                 {
@@ -775,7 +699,7 @@ public partial class GUBP : BuildCommand
                                 NonReadyAgentGroups.Add(NodeToDo.AgentSharingGroup);
                                 break;
                             }
-                            foreach (BuildNode Dep in ChainNode.PseudoDependencies)
+                            foreach (BuildNode Dep in ChainNode.OrderDependencies)
                             {
                                 if (!SortedAgentGroupChains[NodeToDo.AgentSharingGroup].Contains(Dep) && NodesToDo.Contains(Dep))
                                 {
@@ -789,7 +713,7 @@ public partial class GUBP : BuildCommand
                 }
                 else
                 {
-                    foreach (BuildNode Dep in NodeToDo.Dependencies)
+                    foreach (BuildNode Dep in NodeToDo.InputDependencies)
                     {
                         if (NodesToDo.Contains(Dep))
                         {
@@ -797,7 +721,7 @@ public partial class GUBP : BuildCommand
                             break;
                         }
                     }
-                    foreach (BuildNode Dep in NodeToDo.PseudoDependencies)
+                    foreach (BuildNode Dep in NodeToDo.OrderDependencies)
                     {
                         if (NodesToDo.Contains(Dep))
                         {
@@ -892,7 +816,7 @@ public partial class GUBP : BuildCommand
                     {
                         foreach (BuildNode ChainNode in SortedAgentGroupChains[NodeToDo.AgentSharingGroup])
                         {
-                            foreach (BuildNode Dep in ChainNode.Dependencies)
+                            foreach (BuildNode Dep in ChainNode.InputDependencies)
                             {
                                 if (!SortedAgentGroupChains[NodeToDo.AgentSharingGroup].Contains(Dep) && NodesToDo.Contains(Dep))
                                 {
@@ -901,14 +825,14 @@ public partial class GUBP : BuildCommand
                             }
                         }
                     }
-                    foreach (BuildNode Dep in NodeToDo.Dependencies)
+                    foreach (BuildNode Dep in NodeToDo.InputDependencies)
                     {
                         if (NodesToDo.Contains(Dep))
                         {
                             Deps = Deps + Dep.Name + " ";
                         }
                     }
-                    foreach (BuildNode Dep in NodeToDo.PseudoDependencies)
+                    foreach (BuildNode Dep in NodeToDo.OrderDependencies)
                     {
                         if (NodesToDo.Contains(Dep))
                         {
@@ -1166,7 +1090,7 @@ public partial class GUBP : BuildCommand
 		}
 		foreach (BuildNode BuildNode in NodeNameToInfo.Values)
 		{
-			LinkNode(BuildNode, AggregateNameToInfo, NodeNameToInfo, ref NumErrors);
+			LinkNode(BuildNode, new List<BuildNode>(), AggregateNameToInfo, NodeNameToInfo, ref NumErrors);
 		}
 		if(NumErrors > 0)
 		{
@@ -1212,59 +1136,67 @@ public partial class GUBP : BuildCommand
 		}
 	}
 
-	/// <summary>
-	/// Resolve a node's dependency names to arrays of NodeInfo instances, filling in the appropriate fields in the NodeInfo object. 
-	/// </summary>
-	/// <param name="Node"></param>
-	/// <param name="AggregateNameToInfo">Map of other aggregate names to their corresponding instance.</param>
-	/// <param name="NodeNameToInfo">Map from node names to their corresponding instance.</param>
-	/// <param name="NumErrors">The number of errors output so far. Incremented if resolving this aggregate fails.</param>
-	private static void LinkNode(BuildNode Node, Dictionary<string, AggregateNode> AggregateNameToInfo, Dictionary<string, BuildNode> NodeNameToInfo, ref int NumErrors)
+	private static void LinkNode(BuildNode Node, List<BuildNode> NodeStack, Dictionary<string, AggregateNode> NameToAggregateNode, Dictionary<string, BuildNode> NameToBuildNode, ref int NumErrors)
 	{
-		if(Node.Dependencies == null)
+		if(Node.OrderDependencies == null)
 		{
-			// Find all the dependencies
-			HashSet<BuildNode> Dependencies = new HashSet<BuildNode>();
-			foreach (string DependencyName in Node.DependencyNames)
+			NodeStack.Add(Node);
+
+			int NodeStackIdx = NodeStack.IndexOf(Node);
+			if (NodeStackIdx != NodeStack.Count - 1)
+			{
+				// There's a cycle in the build graph. Print out the chain.
+				CommandUtils.LogError("Build graph contains a cycle ({0})", String.Join(" -> ", NodeStack.Skip(NodeStackIdx).Select(x => x.Name)));
+				NumErrors++;
+				Node.OrderDependencies = new HashSet<BuildNode>();
+				Node.InputDependencies = new HashSet<BuildNode>();
+			}
+			else
 			{
-				if (!ResolveDependencies(DependencyName, AggregateNameToInfo, NodeNameToInfo, Dependencies))
+				// Find all the direct input dependencies
+				HashSet<BuildNode> DirectInputDependencies = new HashSet<BuildNode>();
+				foreach (string InputDependencyName in Node.InputDependencyNames)
 				{
-					CommandUtils.LogError("Node {0} is not in the graph. It is a dependency of {1}.", DependencyName, Node.Name);
-					NumErrors++;
+					if (!ResolveDependencies(InputDependencyName, NameToAggregateNode, NameToBuildNode, DirectInputDependencies))
+					{
+						CommandUtils.LogError("Node {0} is not in the graph. It is an input dependency of {1}.", InputDependencyName, Node.Name);
+						NumErrors++;
+					}
 				}
-			}
-			Node.Dependencies = Dependencies.ToArray();
 
-			// Find all the pseudo-dependencies
-			HashSet<BuildNode> PseudoDependencies = new HashSet<BuildNode>();
-			foreach (string PseudoDependencyName in Node.PseudoDependencyNames)
-			{
-				if (!ResolveDependencies(PseudoDependencyName, AggregateNameToInfo, NodeNameToInfo, PseudoDependencies))
+				// Find all the direct order dependencies. All the input dependencies are also order dependencies.
+				HashSet<BuildNode> DirectOrderDependencies = new HashSet<BuildNode>(DirectInputDependencies);
+				foreach (string OrderDependencyName in Node.OrderDependencyNames)
 				{
-					CommandUtils.LogError("Node {0} is not in the graph. It is a pseudodependency of {1}.", PseudoDependencyName, Node.Name);
-					NumErrors++;
+					if (!ResolveDependencies(OrderDependencyName, NameToAggregateNode, NameToBuildNode, DirectOrderDependencies))
+					{
+						CommandUtils.LogError("Node {0} is not in the graph. It is an order dependency of {1}.", OrderDependencyName, Node.Name);
+						NumErrors++;
+					}
 				}
-			}
-			Node.PseudoDependencies = PseudoDependencies.ToArray();
 
-			// Set the direct dependencies list
-			Node.AllDirectDependencies = Node.Dependencies.Union(Node.PseudoDependencies).ToArray();
+				// Link everything
+				foreach(BuildNode DirectOrderDependency in DirectOrderDependencies)
+				{
+					LinkNode(DirectOrderDependency, NodeStack, NameToAggregateNode, NameToBuildNode, ref NumErrors);
+				}
 
-			// Recursively find the dependencies for all the dependencies
-			HashSet<BuildNode> IndirectDependenices = new HashSet<BuildNode>(Node.AllDirectDependencies);
-			foreach(BuildNode DirectDependency in Node.AllDirectDependencies)
-			{
-				LinkNode(DirectDependency, AggregateNameToInfo, NodeNameToInfo, ref NumErrors);
-				IndirectDependenices.UnionWith(DirectDependency.AllIndirectDependencies);
-			}
-			Node.AllIndirectDependencies = IndirectDependenices.ToArray();
+				// Recursively include all the input dependencies
+				Node.InputDependencies = new HashSet<BuildNode>(DirectInputDependencies);
+				foreach (BuildNode DirectInputDependency in DirectInputDependencies)
+				{
+					Node.InputDependencies.UnionWith(DirectInputDependency.InputDependencies);
+				}
 
-			// Check the node doesn't reference itself
-			if(Node.AllIndirectDependencies.Contains(Node))
-			{
-				CommandUtils.LogError("Node {0} has a dependency on itself.", Node.Name);
-				NumErrors++;
+				// Same for the order dependencies
+				Node.OrderDependencies = new HashSet<BuildNode>(DirectOrderDependencies);
+				foreach (BuildNode DirectOrderDependency in DirectOrderDependencies)
+				{
+					Node.OrderDependencies.UnionWith(DirectOrderDependency.OrderDependencies);
+				}
 			}
+
+			NodeStack.RemoveAt(NodeStack.Count - 1);
 		}
 	}
 
@@ -1355,7 +1287,7 @@ public partial class GUBP : BuildCommand
 		HashSet<BuildNode> RecursiveNodesToDo = new HashSet<BuildNode>(NodesToDo);
 		foreach(BuildNode NodeToDo in NodesToDo)
 		{
-			RecursiveNodesToDo.UnionWith(NodeToDo.AllIndirectDependencies);
+			RecursiveNodesToDo.UnionWith(NodeToDo.OrderDependencies);
 		}
 
 		return RecursiveNodesToDo;
@@ -1546,21 +1478,14 @@ public partial class GUBP : BuildCommand
 			{
 				continue;
 			}
+
 			int MyIndex = OrderedToDo.IndexOf(NodeToDo);
-			foreach (BuildNode Dep in NodeToDo.Dependencies)
-			{
-				int DepIndex = OrderedToDo.IndexOf(Dep);
-				if (DepIndex >= MyIndex)
-				{
-					throw new AutomationException("Topological sort error, node {0} has a dependency of {1} which sorted after it.", NodeToDo.Name, Dep.Name);
-				}
-			}
-			foreach (BuildNode Dep in NodeToDo.PseudoDependencies)
+			foreach (BuildNode OrderDependency in NodeToDo.OrderDependencies)
 			{
-				int DepIndex = OrderedToDo.IndexOf(Dep);
-				if (DepIndex >= MyIndex)
+				int OrderDependencyIdx = OrderedToDo.IndexOf(OrderDependency);
+				if(OrderDependencyIdx >= MyIndex)
 				{
-					throw new AutomationException("Topological sort error, node {0} has a pseduodependency of {1} which sorted after it.", NodeToDo.Name, Dep.Name);
+					throw new AutomationException("Topological sort error, node {0} has an order dependency on {1} which sorted after it.", NodeToDo.Name, OrderDependency.Name);
 				}
 			}
 		}
@@ -1609,48 +1534,11 @@ public partial class GUBP : BuildCommand
         Dictionary<string, BuildNode> BuildProductToNodeMap = new Dictionary<string, BuildNode>();
 		foreach (BuildNode NodeToDo in OrderedToDo)
         {
-            if (NodeToDo.Node.BuildProducts != null || NodeToDo.Node.AllDependencyBuildProducts != null)
+            if (NodeToDo.Node.BuildProducts != null)
             {
                 throw new AutomationException("topological sort error");
             }
 
-            NodeToDo.Node.AllDependencyBuildProducts = new List<string>();
-            NodeToDo.Node.AllDependencies = new List<string>();
-            foreach (BuildNode Dep in NodeToDo.Dependencies)
-            {
-                NodeToDo.Node.AddAllDependent(Dep.Name);
-
-                if (Dep.Node.AllDependencies == null)
-                {
-					throw new AutomationException("Node {0} was not processed yet?  Processing {1}", Dep, NodeToDo.Name);
-                }
-
-				foreach (string DepDep in Dep.Node.AllDependencies)
-				{
-					NodeToDo.Node.AddAllDependent(DepDep);
-				}
-				
-                if (Dep.Node.BuildProducts == null)
-                {
-                    throw new AutomationException("Node {0} was not processed yet? Processing {1}", Dep, NodeToDo.Name);
-                }
-
-				foreach (string Prod in Dep.Node.BuildProducts)
-                {
-                    NodeToDo.Node.AddDependentBuildProduct(Prod);
-                }
-
-                if (Dep.Node.AllDependencyBuildProducts == null)
-                {
-                    throw new AutomationException("Node {0} was not processed yet2?  Processing {1}", Dep.Name, NodeToDo.Name);
-                }
-
-                foreach (string Prod in Dep.Node.AllDependencyBuildProducts)
-                {
-                    NodeToDo.Node.AddDependentBuildProduct(Prod);
-                }
-            }
-
             string NodeStoreName = StoreName + "-" + NodeToDo.Name;
             
             string GameNameIfAny = NodeToDo.GameNameIfAnyForTempStorage;
@@ -1729,30 +1617,6 @@ public partial class GUBP : BuildCommand
                     }
 
                     TempStorage.StoreToTempStorage(CmdEnv, NodeStoreName, NodeToDo.Node.BuildProducts, !bSaveSharedTempStorage, GameNameIfAny, StorageRootIfAny);
-
-                    if (ParseParam("StompCheck"))
-                    {
-                        foreach (string Dep in NodeToDo.Node.AllDependencies)
-                        {
-                            try
-                            {
-                                bool WasLocal;
-								using(TelemetryStopwatch RetrieveBuildProductsStopwatch = new TelemetryStopwatch("RetrieveBuildProducts"))
-								{
-									TempStorage.RetrieveFromTempStorage(CmdEnv, NodeStoreName, out WasLocal, GameNameIfAny, StorageRootIfAny);
-								}
-								if (!WasLocal)
-								{
-									throw new AutomationException("Retrieve was not local?");
-								}
-																	                                    
-                            }
-                            catch(Exception Ex)
-                            {
-                                throw new AutomationException("Node {0} stomped Node {1}   Ex: {2}", NodeToDo.Name, Dep, LogUtils.FormatException(Ex));
-                            }
-                        }
-                    }
                 }
                 catch (Exception Ex)
                 {
diff --git a/Engine/Source/Programs/AutomationTool/GUBP/LegacyNodes.cs b/Engine/Source/Programs/AutomationTool/GUBP/LegacyNodes.cs
--- a/Engine/Source/Programs/AutomationTool/GUBP/LegacyNodes.cs
+++ b/Engine/Source/Programs/AutomationTool/GUBP/LegacyNodes.cs
@@ -18,7 +18,6 @@ partial class GUBP
         public List<string> FullNamesOfPseudodependencies = new List<string>(); //these are really only used for sorting. We want the editor to fail before the monolithics. Think of it as "can't possibly be useful without".
         public List<string> BuildProducts = null;
         public List<string> AllDependencyBuildProducts = null;
-        public List<string> AllDependencies = null;		
         public string AgentSharingGroup = "";
 
         public virtual string GetFullName()
@@ -178,13 +177,6 @@ partial class GUBP
                 AllDependencyBuildProducts.Add(Info.FullName);
             }
         }
-        public void AddAllDependent(string Node)
-        {
-            if (!AllDependencies.Contains(Node))
-            {
-                AllDependencies.Add(Node);
-            }
-        }
         public void RemoveOveralppingBuildProducts()
         {
             foreach (var ToRemove in AllDependencyBuildProducts)
diff --git a/Engine/Source/Programs/AutomationTool/GUBP/Nodes/BuildNode.cs b/Engine/Source/Programs/AutomationTool/GUBP/Nodes/BuildNode.cs
--- a/Engine/Source/Programs/AutomationTool/GUBP/Nodes/BuildNode.cs
+++ b/Engine/Source/Programs/AutomationTool/GUBP/Nodes/BuildNode.cs
@@ -20,10 +20,8 @@ namespace AutomationTool
 		public float Priority;
 
 		public GUBP.GUBPNode Node;
-		public BuildNode[] Dependencies;
-		public BuildNode[] PseudoDependencies;
-		public BuildNode[] AllDirectDependencies;
-		public BuildNode[] AllIndirectDependencies;
+		public HashSet<BuildNode> InputDependencies;
+		public HashSet<BuildNode> OrderDependencies;
 		public TriggerNode[] ControllingTriggers;
 		public bool IsComplete;
 		public string[] RecipientsForFailureEmails;
@@ -41,12 +39,12 @@ namespace AutomationTool
 			get;
 		}
 
-		public abstract IEnumerable<string> DependencyNames
+		public abstract IEnumerable<string> InputDependencyNames
 		{
 			get;
 		}
 
-		public abstract IEnumerable<string> PseudoDependencyNames
+		public abstract IEnumerable<string> OrderDependencyNames
 		{
 			get;
 		}
@@ -62,7 +60,7 @@ namespace AutomationTool
 
 		public bool DependsOn(BuildNode Node)
 		{
-			return AllIndirectDependencies.Contains(Node);
+			return OrderDependencies.Contains(Node);
 		}
 
 		public override string ToString()
@@ -115,18 +113,27 @@ namespace AutomationTool
 			get { return Node.IsSticky(); }
 		}
 
-		public override IEnumerable<string> DependencyNames
+		public override IEnumerable<string> InputDependencyNames
 		{
 			get { return Node.FullNamesOfDependencies; }
 		}
 
-		public override IEnumerable<string> PseudoDependencyNames
+		public override IEnumerable<string> OrderDependencyNames
 		{
 			get { return Node.FullNamesOfPseudodependencies; }
 		}
 
 		public override void DoBuild()
 		{
+			Node.AllDependencyBuildProducts = new List<string>();
+			foreach (BuildNode InputDependency in InputDependencies)
+			{
+				foreach (string BuildProduct in InputDependency.Node.BuildProducts)
+				{
+					Node.AddDependentBuildProduct(BuildProduct);
+				}
+			}
+
 			Node.DoBuild(Owner);
 		}
 
