Prevent using the dedicated network share for ShooterGame from now on. Will be removed in the next few days.

diff --git a/Engine/Source/Programs/AutomationTool/GUBP/TempStorage.cs b/Engine/Source/Programs/AutomationTool/GUBP/TempStorage.cs
--- a/Engine/Source/Programs/AutomationTool/GUBP/TempStorage.cs
+++ b/Engine/Source/Programs/AutomationTool/GUBP/TempStorage.cs
@@ -577,7 +577,7 @@ namespace AutomationTool
             }
             string Root = CommandUtils.RootSharedTempStorageDirectory();
             string Result = CommandUtils.CombinePaths(Root, GameFolder);
-            if (String.IsNullOrEmpty(GameFolder) || !InternalUtils.Robust_DirectoryExistsAndIsWritable_NoExceptions(Result))
+            if (String.IsNullOrEmpty(GameFolder) || String.Compare(GameFolder, "ShooterGame", true) == 0 || !InternalUtils.Robust_DirectoryExistsAndIsWritable_NoExceptions(Result))
             {
                 string GameStr = "Game";
                 bool HadGame = false;
