Fix up calls to construct AutomationException(string Message, Exception Ex), as this overloaded constructor has been removed.

The overloaded constructor was removed in CL 2605826, so all existing calls are being interpreted as AutomationException(string Message, object[] Params), causing runtime failures when Message does not contain a replacement placeholder "{0}", and potential unwanted behavior if it does.

diff --git a/Engine/Source/Programs/AutomationTool/AutomationUtils/CommandUtils.cs b/Engine/Source/Programs/AutomationTool/AutomationUtils/CommandUtils.cs
--- a/Engine/Source/Programs/AutomationTool/AutomationUtils/CommandUtils.cs
+++ b/Engine/Source/Programs/AutomationTool/AutomationUtils/CommandUtils.cs
@@ -816,7 +816,7 @@ namespace AutomationTool
 			}
 			catch (Exception Ex)
 			{
-				throw new AutomationException(String.Format("Unable to change current directory to {0}", WorkingDirectory), Ex);
+				throw new AutomationException(Ex, "Unable to change current directory to {0}", WorkingDirectory);
 			}
 
 			WorkingDirectoryStack.Push(OrigCurrentDirectory);
@@ -937,7 +937,7 @@ namespace AutomationTool
 			Filename = ConvertSeparators(PathSeparator.Default, Filename);
 			if (!File.Exists(Filename))
 			{
-				throw new AutomationException("Unable to set attributes for a non-exisiting file.", new FileNotFoundException("File not found.", Filename));
+				throw new AutomationException(new FileNotFoundException("File not found.", Filename), "Unable to set attributes for a non-existing file.");
 			}
 
 			FileAttributes Attributes = File.GetAttributes(Filename);
@@ -1615,7 +1615,7 @@ namespace AutomationTool
 					}
 					if (Request.Result != null)
 					{
-						throw new AutomationException(String.Format("Failed to thread-copy files: {0}", Request.Result.Message), Request.Result);
+						throw new AutomationException(Request.Result, "Failed to thread-copy files: {0}", Request.Result.Message);
 					}
 				}
 			}
@@ -1708,7 +1708,7 @@ namespace AutomationTool
 			}
 			catch (Exception Ex)
 			{
-				throw new AutomationException(String.Format("Failed to set environment variable {0} to {1}", Name, Value), Ex);
+				throw new AutomationException(Ex, "Failed to set environment variable {0} to {1}", Name, Value);
 			}
 		}
 
@@ -1740,7 +1740,7 @@ namespace AutomationTool
 			}
 			catch (Exception Ex)
 			{
-				throw new AutomationException("Failed to harvest environment variables", Ex);
+				throw new AutomationException(Ex, "Failed to harvest environment variables");
 			}
 
 			if (AlsoSet)
diff --git a/Engine/Source/Programs/AutomationTool/AutomationUtils/P4Utils.cs b/Engine/Source/Programs/AutomationTool/AutomationUtils/P4Utils.cs
--- a/Engine/Source/Programs/AutomationTool/AutomationUtils/P4Utils.cs
+++ b/Engine/Source/Programs/AutomationTool/AutomationUtils/P4Utils.cs
@@ -25,7 +25,7 @@ namespace AutomationTool
 		public P4Exception(string Msg)
 			: base(Msg) { }
 		public P4Exception(string Msg, Exception InnerException)
-			: base(Msg, InnerException) { }
+			: base(InnerException, Msg) { }
 
 		public P4Exception(string Format, params object[] Args)
 			: base(Format, Args) { }
diff --git a/Engine/Source/Programs/AutomationTool/AutomationUtils/Utils.cs b/Engine/Source/Programs/AutomationTool/AutomationUtils/Utils.cs
--- a/Engine/Source/Programs/AutomationTool/AutomationUtils/Utils.cs
+++ b/Engine/Source/Programs/AutomationTool/AutomationUtils/Utils.cs
@@ -874,7 +874,7 @@ namespace AutomationTool
                 }
                 catch (Exception ex)
                 {
-                    throw new AutomationException(string.Format("Failed to parse line {0} in version file {1}", line, Filename), ex);
+                    throw new AutomationException(ex, "Failed to parse line {0} in version file {1}", line, Filename);
                 }
             }
 
@@ -991,7 +991,7 @@ namespace AutomationTool
             }
             catch (Exception ex)
             {
-                throw new AutomationException(string.Format("Failed to parse {0} as an FEngineVersion compatible string", versionString), ex);
+                throw new AutomationException(ex, "Failed to parse {0} as an FEngineVersion compatible string", versionString);
             }
         }
     }
