Re-added ambiguous Automation ctor that was removed in CL2605826. It is now marked with [Obsolete], and will be removed in the future. CL2643334 already fixed the instances that were mistakenly being routed to the varargs version.

diff --git a/Engine/Source/Programs/AutomationTool/AutomationUtils/AutomationExpection.cs b/Engine/Source/Programs/AutomationTool/AutomationUtils/AutomationExpection.cs
--- a/Engine/Source/Programs/AutomationTool/AutomationUtils/AutomationExpection.cs
+++ b/Engine/Source/Programs/AutomationTool/AutomationUtils/AutomationExpection.cs
@@ -82,6 +82,12 @@ namespace AutomationTool
             this.ErrorCode = ErrorCode;
         }
 
+        [Obsolete]
+        public AutomationException(string Msg, System.Exception InnerException)
+            : base(Msg, InnerException)
+        {
+        }
+
         public AutomationException(Exception InnerException, string Format, params object[] Args)
 			:base(string.Format(Format, Args), InnerException)
 		{
