Call out segmentation fault errors (and other signals) when the editor commandlet terminates abnormally, as well as printing out the exit code.

diff --git a/Engine/Source/Programs/AutomationTool/AutomationUtils/CommandletUtils.cs b/Engine/Source/Programs/AutomationTool/AutomationUtils/CommandletUtils.cs
--- a/Engine/Source/Programs/AutomationTool/AutomationUtils/CommandletUtils.cs
+++ b/Engine/Source/Programs/AutomationTool/AutomationUtils/CommandletUtils.cs
@@ -140,6 +140,7 @@ namespace AutomationTool
         {
             RunCommandlet(ProjectName, UE4Exe, "UpdateGameProject", Parameters);
         }
+
 		/// <summary>
 		/// Runs a commandlet using Engine/Binaries/Win64/UE4Editor-Cmd.exe.
 		/// </summary>
@@ -183,6 +184,19 @@ namespace AutomationTool
 			var RunResult = Run(EditorExe, Args, Options: Opts);
 			PopDir();
 
+			// Draw attention to signal exit codes on Posix systems, rather than just printing the exit code
+			if(RunResult.ExitCode > 128 && RunResult.ExitCode < 128 + 32)
+			{
+				if(RunResult.ExitCode == 139)
+				{
+					CommandUtils.LogError("Editor terminated abnormally due to a segmentation fault");
+				}
+				else
+				{
+					CommandUtils.LogError("Editor terminated abnormally with signal {0}", RunResult.ExitCode - 128);
+				}
+			}
+
 			// Copy the local commandlet log to the destination folder.
 			string DestLogFile = LogUtils.GetUniqueLogName(CombinePaths(CmdEnv.LogFolder, Commandlet));
 			if (!CommandUtils.CopyFile_NoExceptions(LocalLogFile, DestLogFile))
