Added missing Source folder to zipping filter list.

UE-19560

diff --git a/Engine/Source/Programs/AutomationTool/Scripts/BuildCommonTools.Automation.cs b/Engine/Source/Programs/AutomationTool/Scripts/BuildCommonTools.Automation.cs
--- a/Engine/Source/Programs/AutomationTool/Scripts/BuildCommonTools.Automation.cs
+++ b/Engine/Source/Programs/AutomationTool/Scripts/BuildCommonTools.Automation.cs
@@ -177,6 +177,7 @@ public class ZipProjectUp : BuildCommand
         FileFilter Filter = new FileFilter();
         Filter.Include("/Config/...");
         Filter.Include("/Content/...");
+        Filter.Include("/Source/...");
         Filter.Include("*.uproject");
 
         ZipFiles(InstallDirectory, ProjectDirectory, Filter);
