Always copy web server during packaging itself too for html5. #jira UE-19556

diff --git a/Engine/Source/Programs/AutomationTool/HTML5/HTML5Platform.Automation.cs b/Engine/Source/Programs/AutomationTool/HTML5/HTML5Platform.Automation.cs
--- a/Engine/Source/Programs/AutomationTool/HTML5/HTML5Platform.Automation.cs
+++ b/Engine/Source/Programs/AutomationTool/HTML5/HTML5Platform.Automation.cs
@@ -171,6 +171,7 @@ public class HTML5Platform : Platform
 		// Utility 
 		CompressionTasks[5] = Task.Factory.StartNew(() => CompressFile(OutDir + "/Utility.js", OutDir + "/Utility.js.gz"));
 
+		File.Copy(CombinePaths(CmdEnv.LocalRoot, "Engine/Binaries/DotNET/HTML5LaunchHelper.exe"),CombinePaths(OutDir, "HTML5LaunchHelper.exe"),true);
 		Task.WaitAll(CompressionTasks);
 		PrintRunTime();
 	}
