TempStorage Refactor
GUBP High Level
* Temp Storage is zipped into a single archive per node now. This results in ~75% reduction in temp storage usage and network traffic, not to mention the per-file overhead.
* Temp Storage is in P:\\Builds\\{Game}\\TmpStore instead of P:\\Builds\\{Game}\\GUBP (to facilitate easier cleaning of this new structure).
* Temp Storage nodes are in subdirectories of {Branch}\\{CL}\\{NodeName} now instead of a flat directory structure that was hard to manually sift through.

GUBP Mid Level
* Removed -Store= and -StoreSuffix= test parameters.
* Added -NoZipTempStorage parameter to turn off temp storage zipping if necessary.
* Created GUBP.JobInfo class that collects info about the job as a whole to be passed around by GUBP. Mostly used by any code that need to interact with TempStorage.
* Created TempStorageNodeInfo that describes the necessary parameters to find the temp storage location for a node.
* Fully XML commented TempStorage.cs, and commented internals all major functions.
* Added a bunch of telemetry data for storing, retrieving, and cleaning shared temp storage.

UAT Mid Level
* Fixed a bug in Ionic.Zip that make ExtractAll() not work on Mono, checked in new DLLs.
* Added UAT parameter -UseLocalBuildStorage that allows you to test build storage stuff completely locally. Writes to Engine\\Saved\\LocalBuilds\\...

GUBP Low Level
* Refactored some GUBP startup code so temp vars would be limited in scope. Makes it easier to track the impact of refactoring these things.
* CullNodesForPreflight is only called for preflight builds.
* Refactored TempStorage.FindTempStorageManifests to use new TmpStore structure and harden the brittle string/path parsing it was doing. See the new TempStorage FindMatchingSharedTempStorageNodeCLs().
* Refactored TempStorage Saving and Loading to use XDocument instead of older XmlDocument. Removed a bunch of redundant checks.
* Use StripBaseDirectory and MakeRerootedFilePath to remove the brittle directory manipulation code. Directories no longer require a '/' at the end.
* Removed a few redundant caching layers in cleaning temp storage that try to ensure we don't clean a folder twice. None of them were necessary.
* Removed unused single-threaded copy code from temp storage.
* Updated Temp Storage unit test, and fully commented the logic behind it.

UAT Low Level
* UAT top level exception handler is now a single log line now to help parsers find the error.
* Removed several uses of FormatException as it doesn't display the entire exception chain, and is not as good as the default exception formatter.
* Removed ExceptionToString as it used FormatException, which was not a good precedent.
* Fixed several cases of exception propagation that was not properly chaining the inner exception.
* Refactored ThreadedCopyFiles to use Parallel.For because it was just as fast (if not faster) and much simpler to maintain.
* Removed the suffix from Robust_FileExists_NoExceptions because it's sole purpose in life WAS to throw exceptions!
* Added a bunch of XML doc comments to CommandUtils.
* Modernized some container manipulation and iteration to use IEnumerable and extension methods more appropriately.
* Added several [REDACTED_HANDLE] for other minor cleanup stuff that should happen eventually.
* Fixed some uses of String.Compare to use invariant culture.

diff --git a/Engine/Source/Programs/AutomationTool/AutomationUtils/Automation.cs b/Engine/Source/Programs/AutomationTool/AutomationUtils/Automation.cs
--- a/Engine/Source/Programs/AutomationTool/AutomationUtils/Automation.cs
+++ b/Engine/Source/Programs/AutomationTool/AutomationUtils/Automation.cs
@@ -121,6 +121,11 @@ namespace AutomationTool
 		public static CommandLineArg UTF8Output = new CommandLineArg("-UTF8Output");
 		public static CommandLineArg NoAutoSDK = new CommandLineArg("-NoAutoSDK");
 		public static CommandLineArg IgnoreJunk = new CommandLineArg("-ignorejunk");
+        /// <summary>
+        /// Allows you to use local storage for your root build storage dir (default of P:\Builds (on PC) is changed to Engine\Saved\LocalBuilds). Used for local testing.
+        /// </summary>
+        public static CommandLineArg UseLocalBuildStorage = new CommandLineArg("-UseLocalBuildStorage");
+
 		/// <summary>
 		/// Force initialize static members by calling this.
 		/// </summary>
@@ -149,7 +154,8 @@ AutomationTool.exe [-verbose] [-compileonly] [-p4] Command0 [-Arg0 -Arg1 -Arg2 
 	[Help("nosubmit", "Prevents any submit attempts")]
 	[Help("nokill", "Does not kill any spawned processes on exit")]
 	[Help("ignorejunk", "Prevents UBT from cleaning junk files")]
-	public static class Automation
+    [Help("UseLocalBuildStorage", @"Allows you to use local storage for your root build storage dir (default of P:\Builds (on PC) is changed to Engine\Saved\LocalBuilds). Used for local testing.")]
+    public static class Automation
 	{
 		#region Command line parsing
 
@@ -167,7 +173,7 @@ AutomationTool.exe [-verbose] [-compileonly] [-p4] Command0 [-Arg0 -Arg1 -Arg2 
 			bool bGlobalParam = false;
 			foreach (var RegisteredParam in GlobalCommandLine.RegisteredArgs)
 			{
-				if (String.Compare(CurrentParam, RegisteredParam.Key, true) == 0)
+				if (String.Compare(CurrentParam, RegisteredParam.Key, StringComparison.InvariantCultureIgnoreCase) == 0)
 				{
 					// Register and exit, we're done here.
 					RegisteredParam.Value.Set();
diff --git a/Engine/Source/Programs/AutomationTool/AutomationUtils/CommandUtils.cs b/Engine/Source/Programs/AutomationTool/AutomationUtils/CommandUtils.cs
--- a/Engine/Source/Programs/AutomationTool/AutomationUtils/CommandUtils.cs
+++ b/Engine/Source/Programs/AutomationTool/AutomationUtils/CommandUtils.cs
@@ -14,6 +14,7 @@ using UnrealBuildTool;
 using System.Runtime.CompilerServices;
 using System.Linq;
 using System.Reflection;
+using System.Threading.Tasks;
 using Tools.DotNETCommon;
 using Tools.DotNETCommon.CaselessDictionary;
 using Tools.DotNETCommon.HarvestEnvVars;
@@ -108,16 +109,6 @@ namespace AutomationTool
 
 		#region Logging
 
-		/// <summary>
-		/// Returns a formatted string for the specified exception (including inner exception if any).
-		/// </summary>
-		/// <param name="Ex"></param>
-		/// <returns>Formatted exception string.</returns>
-		public static string ExceptionToString(Exception Ex)
-		{
-			return LogUtils.FormatException(Ex);
-		}
-
 		/// <summary>
 		/// Writes formatted text to log (with LogEventType.Console).
 		/// </summary>
@@ -379,6 +370,7 @@ namespace AutomationTool
         /// <summary>
         /// Finds files in specified paths. 
         /// </summary>
+        /// <param name="bQuiet">When true, logging is suppressed.</param>
         /// <param name="SearchPattern">Pattern</param>
         /// <param name="Recursive">Recursive search</param>
         /// <param name="Paths">Paths to search</param>
@@ -403,7 +395,8 @@ namespace AutomationTool
 		/// <summary>
 		/// Finds files in specified paths. 
 		/// </summary>
-		/// <param name="SearchPattern">Pattern</param>
+        /// <param name="bQuiet">When true, logging is suppressed.</param>
+        /// <param name="SearchPattern">Pattern</param>
 		/// <param name="Recursive">Recursive search</param>
 		/// <param name="Paths">Paths to search</param>
 		/// <returns>An array of files found in the specified paths</returns>
@@ -428,7 +421,8 @@ namespace AutomationTool
 		/// <summary>
 		/// Finds Directories in specified paths. 
 		/// </summary>
-		/// <param name="SearchPattern">Pattern</param>
+        /// <param name="bQuiet">When true, logging is suppressed.</param>
+        /// <param name="SearchPattern">Pattern</param>
 		/// <param name="Recursive">Recursive search</param>
 		/// <param name="Paths">Paths to search</param>
 		/// <returns>An array of files found in the specified paths</returns>
@@ -454,7 +448,7 @@ namespace AutomationTool
 		/// If the file does not exist, silently succeeds.
 		/// If the deletion of the file fails, this function throws an Exception.
 		/// </summary>
-		/// <param name="Filename">Filename</param>
+		/// <param name="Filenames">Filename</param>
 		public static void DeleteFile(params string[] Filenames)
 		{
 			foreach (var Filename in Filenames)
@@ -471,7 +465,8 @@ namespace AutomationTool
         /// If the file does not exist, silently succeeds.
         /// If the deletion of the file fails, this function throws an Exception.
         /// </summary>
-        /// <param name="Filename">Filename</param>
+        /// <param name="bQuiet">When true, logging is suppressed.</param>
+        /// <param name="Filenames">Filename</param>
         public static void DeleteFile(bool bQuiet, params string[] Filenames)
         {
             foreach (var Filename in Filenames)
@@ -488,7 +483,7 @@ namespace AutomationTool
 		/// Deletes a file(s). 
 		/// If the deletion of the file fails, prints a warning.
 		/// </summary>
-		/// <param name="Filename">Filename</param>
+		/// <param name="Filenames">Filename</param>
         public static bool DeleteFile_NoExceptions(params string[] Filenames)
 		{
 			bool Result = true;
@@ -526,7 +521,7 @@ namespace AutomationTool
 		/// If the deletion of the directory fails, this function throws an Exception.
 		/// </summary>
 		/// <param name="bQuiet">Suppresses log output if true</param>
-		/// <param name="Directory">Directory</param>
+		/// <param name="Directories">Directories</param>
 		public static void DeleteDirectory(bool bQuiet, params string[] Directories)
 		{
 			foreach (var Directory in Directories)
@@ -543,8 +538,8 @@ namespace AutomationTool
 		/// Deletes a directory(or directories) including its contents (recursively, will delete read-only files).
 		/// If the deletion of the directory fails, this function throws an Exception.
 		/// </summary>
-		/// <param name="Directory">Directory</param>
-		public static void DeleteDirectory(params string[] Directories)
+        /// <param name="Directories">Directories</param>
+        public static void DeleteDirectory(params string[] Directories)
 		{
 			DeleteDirectory(false, Directories);
 		}
@@ -554,8 +549,8 @@ namespace AutomationTool
 		/// If the deletion of the directory fails, prints a warning.
 		/// </summary>
 		/// <param name="bQuiet">Suppresses log output if true</param>
-		/// <param name="Directory">Directory</param>
-		public static bool DeleteDirectory_NoExceptions(bool bQuiet, params string[] Directories)
+        /// <param name="Directories">Directories</param>
+        public static bool DeleteDirectory_NoExceptions(bool bQuiet, params string[] Directories)
 		{
 			bool Result = true;
             foreach (var Directory in Directories)
@@ -588,8 +583,8 @@ namespace AutomationTool
 		/// Deletes a directory(or directories) including its contents (recursively, will delete read-only files).
 		/// If the deletion of the directory fails, prints a warning.
 		/// </summary>
-		/// <param name="Directory">Directory</param>
-		public static bool DeleteDirectory_NoExceptions(params string[] Directories)
+        /// <param name="Directories">Directories</param>
+        public static bool DeleteDirectory_NoExceptions(params string[] Directories)
 		{
 			return DeleteDirectory_NoExceptions(false, Directories);
 		}
@@ -620,8 +615,8 @@ namespace AutomationTool
 		/// <summary>
 		/// Checks if a directory(or directories) exists.
 		/// </summary>
-		/// <param name="Directory">Directory</param>
-		/// <returns>True if the directory exists, false otherwise.</returns>
+        /// <param name="Directories">Directories</param>
+        /// <returns>True if the directory exists, false otherwise.</returns>
 		public static bool DirectoryExists(params string[] Directories)
 		{
 			bool bExists = Directories.Length > 0;
@@ -636,8 +631,8 @@ namespace AutomationTool
 		/// <summary>
 		/// Checks if a directory(or directories) exists.
 		/// </summary>
-		/// <param name="Directory">Directory</param>
-		/// <returns>True if the directory exists, false otherwise.</returns>
+        /// <param name="Directories">Directories</param>
+        /// <returns>True if the directory exists, false otherwise.</returns>
 		public static bool DirectoryExists_NoExceptions(params string[] Directories)
 		{
 			bool bExists = Directories.Length > 0;
@@ -663,8 +658,8 @@ namespace AutomationTool
 		/// Creates a directory(or directories).
 		/// If the creation of the directory fails, this function throws an Exception.
 		/// </summary>
-		/// <param name="Directory">Directory</param>
-		public static void CreateDirectory(params string[] Directories)
+        /// <param name="Directories">Directories</param>
+        public static void CreateDirectory(params string[] Directories)
 		{
 			foreach (var DirectoryName in Directories)
 			{
@@ -680,7 +675,8 @@ namespace AutomationTool
         /// Creates a directory(or directories).
         /// If the creation of the directory fails, this function throws an Exception.
         /// </summary>
-        /// <param name="Directory">Directory</param>
+        /// <param name="bQuiet">When true, logging is suppressed.</param>
+        /// <param name="Directories">Directories</param>
         public static void CreateDirectory(bool bQuiet, params string[] Directories)
         {
             foreach (var DirectoryName in Directories)
@@ -697,8 +693,8 @@ namespace AutomationTool
 		/// Creates a directory (or directories).
 		/// If the creation of the directory fails, this function prints a warning.
 		/// </summary>
-		/// <param name="Directory">Directory</param>
-		public static bool CreateDirectory_NoExceptions(params string[] Directories)
+        /// <param name="Directories">Directories</param>
+        public static bool CreateDirectory_NoExceptions(params string[] Directories)
 		{
 			bool Result = true;
 			foreach (var DirectoryName in Directories)
@@ -717,7 +713,8 @@ namespace AutomationTool
 		/// Renames/moves a file.
 		/// If the rename of the file fails, this function throws an Exception.
 		/// </summary>
-		/// <param name="OldName">Old name</param>
+        /// <param name="bQuiet">When true, logging is suppressed.</param>
+        /// <param name="OldName">Old name</param>
 		/// <param name="NewName">new name</param>
         public static void RenameFile(string OldName, string NewName, bool bQuiet = false)
 		{
@@ -750,7 +747,7 @@ namespace AutomationTool
 		/// <summary>
 		/// Checks if a file(s) exists.
 		/// </summary>
-		/// <param name="Filename">Filename.</param>
+		/// <param name="Filenames">Filename.</param>
 		/// <returns>True if the file exists, false otherwise.</returns>
 		public static bool FileExists(params string[] Filenames)
 		{
@@ -766,7 +763,7 @@ namespace AutomationTool
 		/// <summary>
 		/// Checks if a file(s) exists.
 		/// </summary>
-		/// <param name="Filename">Filename.</param>
+		/// <param name="Filenames">Filename.</param>
 		/// <returns>True if the file exists, false otherwise.</returns>
 		public static bool FileExists_NoExceptions(params string[] Filenames)
 		{
@@ -776,7 +773,8 @@ namespace AutomationTool
         /// <summary>
         /// Checks if a file(s) exists.
         /// </summary>
-        /// <param name="Filename">Filename.</param>
+        /// <param name="bQuiet">When true, logging is suppressed.</param>
+        /// <param name="Filenames">Filename.</param>
         /// <returns>True if the file exists, false otherwise.</returns>
         public static bool FileExists(bool bQuiet, params string[] Filenames)
         {
@@ -792,7 +790,8 @@ namespace AutomationTool
         /// <summary>
         /// Checks if a file(s) exists.
         /// </summary>
-        /// <param name="Filename">Filename.</param>
+        /// <param name="bQuiet">When true, logging is suppressed.</param>
+        /// <param name="Filenames">Filename.</param>
         /// <returns>True if the file exists, false otherwise.</returns>
         public static bool FileExists_NoExceptions(bool bQuiet, params string[] Filenames)
         {
@@ -1085,7 +1084,7 @@ namespace AutomationTool
 		/// Writes text to a file.
 		/// </summary>
 		/// <param name="Filename">Filename</param>
-		/// <param name="Lines">Text</param>
+		/// <param name="Text">Text</param>
 		public static void WriteAllText(string Filename, string Text)
 		{
 			Filename = ConvertSeparators(PathSeparator.Default, Filename);
@@ -1099,7 +1098,7 @@ namespace AutomationTool
 		/// Writes text to a file.
 		/// </summary>
 		/// <param name="Filename">Filename</param>
-		/// <param name="Lines">Text</param>
+		/// <param name="Text">Text</param>
 		public static bool WriteAllText_NoExceptions(string Filename, string Text)
 		{
 			Filename = ConvertSeparators(PathSeparator.Default, Filename);
@@ -1246,7 +1245,8 @@ namespace AutomationTool
 		/// </summary>
 		/// <param name="Source"></param>
 		/// <param name="Dest"></param>
-		/// <returns>True if the operation was successful, false otherwise.</returns>
+        /// <param name="bQuiet">When true, logging is suppressed.</param>
+        /// <returns>True if the operation was successful, false otherwise.</returns>
         public static void CopyFile(string Source, string Dest, bool bQuiet = false)
 		{
 			Source = ConvertSeparators(PathSeparator.Default, Source);
@@ -1277,7 +1277,8 @@ namespace AutomationTool
 		/// </summary>
 		/// <param name="Source"></param>
 		/// <param name="Dest"></param>
-		/// <returns>True if the operation was successful, false otherwise.</returns>
+        /// <param name="bQuiet">When true, logging is suppressed.</param>
+        /// <returns>True if the operation was successful, false otherwise.</returns>
         public static bool CopyFile_NoExceptions(string Source, string Dest, bool bQuiet = false)
 		{
 			Source = ConvertSeparators(PathSeparator.Default, Source);
@@ -1354,7 +1355,8 @@ namespace AutomationTool
 		/// </summary>
 		/// <param name="Source"></param>
 		/// <param name="Dest"></param>
-		/// <returns>True if the operation was successful, false otherwise.</returns>
+        /// <param name="bQuiet">When true, logging is suppressed.</param>
+        /// <returns>True if the operation was successful, false otherwise.</returns>
 		public static bool CopyDirectory_NoExceptions(string Source, string Dest, bool bQuiet = false)
 		{
 			Source = ConvertSeparators(PathSeparator.Default, Source);
@@ -1511,39 +1513,8 @@ namespace AutomationTool
 
 		#region Threaded Copy
 
-		/// <summary>
-		/// Single batch of files to copy on a thread
-		/// </summary>
-		class CopyRequest
-		{
-			public string[] Source;
-			public string[] Dest;
-			public Exception Result;
-			public Thread Worker;
-		}
-
-		/// <summary>
-		/// Main thread procedure for copying files
-		/// </summary>
-		private static void CopyThreadProc(object Request)
-		{
-			const bool bQuiet = true;
-			var FileToCopy = (CopyRequest)Request;
-			try
-			{
-				for (int Index = 0; Index < FileToCopy.Source.Length; ++Index)
-				{
-					CopyFile(FileToCopy.Source[Index], FileToCopy.Dest[Index], bQuiet);
-				}
-			}
-			catch (Exception Ex)
-			{
-				FileToCopy.Result = Ex;
-			}
-		}
-
-		/// <summary>
-		/// Copies files using miltiple threads
+        /// <summary>
+		/// Copies files using multiple threads
 		/// </summary>
 		/// <param name="SourceDirectory"></param>
 		/// <param name="DestDirectory"></param>
@@ -1551,83 +1522,37 @@ namespace AutomationTool
 		public static void ThreadedCopyFiles(string SourceDirectory, string DestDirectory, int MaxThreads = 64)
 		{
 			CreateDirectory(DestDirectory);
-
-			var SourceFiles = Directory.GetFiles(SourceDirectory, "*", SearchOption.AllDirectories);
-			var SourceBase = GetDirectoryName(SourceDirectory) + GetPathSeparatorChar(PathSeparator.Default);
-			var DestBase = GetDirectoryName(DestDirectory) + GetPathSeparatorChar(PathSeparator.Default);
-			var DestFiles = new string[SourceFiles.Length];
-			for (int Index = 0; Index < SourceFiles.Length; ++Index)
-			{
-				DestFiles[Index] = SourceFiles[Index].Replace(SourceBase, DestBase);
-			}
+            var SourceFiles = Directory.EnumerateFiles(SourceDirectory, "*", SearchOption.AllDirectories).ToList();
+            var DestFiles = SourceFiles.Select(SourceFile => CommandUtils.MakeRerootedFilePath(SourceFile, SourceDirectory, DestDirectory)).ToList();
 			ThreadedCopyFiles(SourceFiles, DestFiles, MaxThreads);
 		}
 
 		/// <summary>
 		/// Copies files using multiple threads
-		/// </summary>
-		/// <param name="SourceDirectory"></param>
-		/// <param name="DestDirectory"></param>
+        /// </summary>
+		/// <param name="Source"></param>
+		/// <param name="Dest"></param>
 		/// <param name="MaxThreads"></param>
-		public static void ThreadedCopyFiles(string[] Source, string[] Dest, int MaxThreads = 64)
+		public static void ThreadedCopyFiles(List<string> Source, List<string> Dest, int MaxThreads = 64)
 		{
-			LogConsole("Copying {0} file(s) using max {1} thread(s)", Source.Length, MaxThreads);
+			LogConsole("Copying {0} file(s) using max {1} thread(s)", Source.Count, MaxThreads);
 
-			if (Source.Length != Dest.Length)
+            if (Source.Count != Dest.Count)
 			{
-				throw new AutomationException("Source count ({0}) does not match Dest count ({1})", Source.Length, Dest.Length);
-			}
-			const int MinFilesPerThread = 10;
-			if (MaxThreads > 1 && Source.Length > MinFilesPerThread)
-			{
-				// Split evenly across the threads
-				int FilesPerThread = Math.Max(Source.Length / MaxThreads, MinFilesPerThread);
-				int ThreadCount = Math.Min(Source.Length / FilesPerThread + 1, MaxThreads);
-				FilesPerThread = Source.Length / ThreadCount + 1;
-
-				// Divide and copy
-				var WorkerThreads = new List<CopyRequest>(ThreadCount);
-				var FilesToCopy = Source.Length;
-				var FirstFileIndex = 0;
-				while (FilesToCopy > 0)
-				{
-					var Request = new CopyRequest();
-					WorkerThreads.Add(Request);
-
-					Request.Source = new string[Math.Min(FilesToCopy, FilesPerThread)];
-					Request.Dest = new string[Request.Source.Length];
-					Array.Copy(Source, FirstFileIndex, Request.Source, 0, Request.Source.Length);
-					Array.Copy(Dest, FirstFileIndex, Request.Dest, 0, Request.Dest.Length);
-
-					Request.Worker = new Thread(new ParameterizedThreadStart(CopyThreadProc));
-					Request.Worker.Start(Request);
-
-					FirstFileIndex += Request.Source.Length;
-					FilesToCopy -= Request.Source.Length;
-				}
-
-				// Wait for completion
-				foreach (var Request in WorkerThreads)
-				{
-					if (Request.Worker.IsAlive)
-					{
-						Request.Worker.Join();
-					}
-					if (Request.Result != null)
-					{
-						throw new AutomationException(Request.Result, "Failed to thread-copy files: {0}", Request.Result.Message);
-					}
-				}
-			}
-			else
-			{
-				const bool bQuiet = true;
-				for (int Index = 0; Index < Source.Length; ++Index)
-				{
-					CopyFile(Source[Index], Dest[Index], bQuiet);
-				}
+                throw new AutomationException("Source count ({0}) does not match Dest count ({1})", Source.Count, Dest.Count);
 			}
-		}
+            try
+            {
+                Parallel.ForEach(Source.Zip(Dest, (Src, Dst) => new { SourceFile = Src, DestFile = Dst }), new ParallelOptions { MaxDegreeOfParallelism = MaxThreads }, (Pair) =>
+                {
+                    CommandUtils.CopyFile(Pair.SourceFile, Pair.DestFile, true);
+                });
+            }
+            catch (AggregateException Ex)
+            {
+                throw new AutomationException(Ex, "Failed to thread-copy files.");
+            }
+        }
 
 		/// <summary>
 		/// Copies a set of files from one folder to another
@@ -1635,18 +1560,12 @@ namespace AutomationTool
 		/// <param name="SourceDir">Source directory</param>
 		/// <param name="TargetDir">Target directory</param>
 		/// <param name="RelativePaths">Paths relative to the source directory to copy</param>
-		/// <param name="bIgnoreSymlinks">Whether to ignore symlinks during the copy</param>
 		/// <param name="MaxThreads">Maximum number of threads to create</param>
-		/// <returns>Array of filenames copied to the target directory</returns>
-		public static string[] ThreadedCopyFiles(string SourceDir, string TargetDir, string[] RelativePaths, int MaxThreads = 64)
+		/// <returns>List of filenames copied to the target directory</returns>
+		public static List<string> ThreadedCopyFiles(string SourceDir, string TargetDir, List<string> RelativePaths, int MaxThreads = 64)
 		{
-			string[] SourceFileNames = new string[RelativePaths.Length];
-			string[] TargetFileNames = new string[RelativePaths.Length];
-			for(int Idx = 0; Idx < RelativePaths.Length; Idx++)
-			{
-				SourceFileNames[Idx] = CommandUtils.CombinePaths(SourceDir, RelativePaths[Idx]);
-				TargetFileNames[Idx] = CommandUtils.CombinePaths(TargetDir, RelativePaths[Idx]);
-			}
+            var SourceFileNames = RelativePaths.Select(RelativePath => CommandUtils.CombinePaths(SourceDir, RelativePath)).ToList();
+            var TargetFileNames = RelativePaths.Select(RelativePath => CommandUtils.CombinePaths(TargetDir, RelativePath)).ToList();
 			CommandUtils.ThreadedCopyFiles(SourceFileNames, TargetFileNames, MaxThreads);
 			return TargetFileNames;
 		}
@@ -1659,12 +1578,12 @@ namespace AutomationTool
 		/// <param name="Filter">Filter which selects files from the source directory to copy</param>
 		/// <param name="bIgnoreSymlinks">Whether to ignore symlinks during the copy</param>
 		/// <param name="MaxThreads">Maximum number of threads to create</param>
-		/// <returns>Array of filenames copied to the target directory</returns>
-		public static string[] ThreadedCopyFiles(string SourceDir, string TargetDir, FileFilter Filter, bool bIgnoreSymlinks, int MaxThreads = 64)
+		/// <returns>List of filenames copied to the target directory</returns>
+		public static List<string> ThreadedCopyFiles(string SourceDir, string TargetDir, FileFilter Filter, bool bIgnoreSymlinks, int MaxThreads = 64)
 		{
 			// Filter all the relative paths
 			LogConsole("Applying filter to {0}...", SourceDir);
-			string[] RelativePaths = Filter.ApplyToDirectory(SourceDir, bIgnoreSymlinks).ToArray();
+			var RelativePaths = Filter.ApplyToDirectory(SourceDir, bIgnoreSymlinks);
 			return ThreadedCopyFiles(SourceDir, TargetDir, RelativePaths);
 		}
 
@@ -1989,19 +1908,31 @@ namespace AutomationTool
 
 		#endregion
 
+        /// <summary>
+        /// Cached location of the build root storage because the logic to compute it is a little non-trivial.
+        /// </summary>
+        private static string CachedRootBuildStorageDirectory;
 
-        public static string RootSharedTempStorageDirectory()
+        /// <summary>
+        /// "P:\Builds" or "/Volumes/Builds". Root Folder for all build storage.
+        /// </summary>
+        /// <returns>"P:\Builds" or "/Volumes/Builds" unless overridden by -UseLocalBuildStorage from the commandline, where is uses Engine\Saved\LocalBuilds\.</returns>
+        public static string RootBuildStorageDirectory()
         {
-            string StorageDirectory = "";
-            if (UnrealBuildTool.Utils.IsRunningOnMono)
+            if (string.IsNullOrEmpty(CachedRootBuildStorageDirectory))
             {
-                StorageDirectory = "/Volumes/Builds";
-            }
-            else
-            {
-                StorageDirectory = CombinePaths("P:", "Builds");
+                if (GlobalCommandLine.UseLocalBuildStorage)
+                {
+                    CachedRootBuildStorageDirectory = CombinePaths(CmdEnv.LocalRoot, "Engine", "Saved", "LocalBuilds");
+                    // Must create the directory because much of the system assumes it is already there.
+                    CreateDirectory(CombinePaths(CachedRootBuildStorageDirectory, "UE4"));
+                }
+                else
+                {
+                    CachedRootBuildStorageDirectory = Utils.IsRunningOnMono ? "/Volumes/Builds" : CombinePaths("P:", "Builds");
+                }
             }
-            return StorageDirectory;
+            return CachedRootBuildStorageDirectory;
         }
 
         public static bool DirectoryExistsAndIsWritable_NoExceptions(string Dir)
@@ -2037,7 +1968,7 @@ namespace AutomationTool
 
         public static void CleanFormalBuilds(string DirectoryForThisBuild, string CLString = "", int MaximumDaysToKeepTempStorage = 4)
         {
-            if (CLString == "" && (!IsBuildMachine || !DirectoryForThisBuild.StartsWith(RootSharedTempStorageDirectory()) || !P4Enabled))
+            if (CLString == "" && (!IsBuildMachine || !DirectoryForThisBuild.StartsWith(RootBuildStorageDirectory()) || !P4Enabled))
             {
                 return;
             }
@@ -2154,9 +2085,10 @@ namespace AutomationTool
 		/// <returns>List of files written</returns>
 		public static IEnumerable<string> UnzipFiles(string ZipFileName, string BaseDirectory)
 		{
-			using(Ionic.Zip.ZipFile Zip = new Ionic.Zip.ZipFile(ZipFileName))
+            // manually extract the files. There was a problem with the Ionic.Zip library that required this on non-PC at one point,
+            // but that problem is now fixed. Leaving this code as is as we need to return the list of created files and fix up their permissions anyway.
+            using (Ionic.Zip.ZipFile Zip = new Ionic.Zip.ZipFile(ZipFileName))
 			{
-				// For some reason, calling ExtractAll() under Mono throws an exception from a call to Directory.CreateDirectory(). Manually extract all the files instead.
 				List<string> OutputFileNames = new List<string>();
 				foreach(Ionic.Zip.ZipEntry Entry in Zip.Entries)
 				{
diff --git a/Engine/Source/Programs/AutomationTool/AutomationUtils/FileFilter.cs b/Engine/Source/Programs/AutomationTool/AutomationUtils/FileFilter.cs
--- a/Engine/Source/Programs/AutomationTool/AutomationUtils/FileFilter.cs
+++ b/Engine/Source/Programs/AutomationTool/AutomationUtils/FileFilter.cs
@@ -392,13 +392,13 @@ public class FileFilter
 	/// <summary>
 	/// Finds a list of files within a given directory which match the filter.
 	/// </summary>
-	/// <param name="FolderName">File to match</param>
-	/// <returns>True if the file passes the filter</returns>
-	public IEnumerable<string> ApplyToDirectory(string DirectoryName, bool bIgnoreSymlinks)
+    /// <param name="DirectoryName">File to match</param>
+	/// <returns>List of files that pass the filter</returns>
+	public List<string> ApplyToDirectory(string DirectoryName, bool bIgnoreSymlinks)
 	{
 		List<string> MatchingFileNames = new List<string>();
 		FindMatchesFromDirectory(new DirectoryInfo(DirectoryName), "", bIgnoreSymlinks, MatchingFileNames);
-		return MatchingFileNames.ToArray();
+		return MatchingFileNames;
 	}
 
 	/// <summary>
diff --git a/Engine/Source/Programs/AutomationTool/AutomationUtils/Log.cs b/Engine/Source/Programs/AutomationTool/AutomationUtils/Log.cs
--- a/Engine/Source/Programs/AutomationTool/AutomationUtils/Log.cs
+++ b/Engine/Source/Programs/AutomationTool/AutomationUtils/Log.cs
@@ -12,6 +12,7 @@ namespace AutomationTool
 	#region LogUtils
 
 
+    //@todo: This class should be moved somewhere else or the file should be renamed LogUtils.cs
 	public class LogUtils
 	{
         private static string LogFilename;
@@ -146,8 +147,8 @@ namespace AutomationTool
 
         /// <summary>
 		/// Dumps exception info to log.
+		/// @todo: Remove this function as it doesn't do a good job printing the exception information.
 		/// </summary>
-		/// <param name="Verbosity">Verbosity</param>
 		/// <param name="Ex">Exception</param>
 		public static string FormatException(Exception Ex)
 		{
diff --git a/Engine/Source/Programs/AutomationTool/AutomationUtils/MCPPublic.cs b/Engine/Source/Programs/AutomationTool/AutomationUtils/MCPPublic.cs
--- a/Engine/Source/Programs/AutomationTool/AutomationUtils/MCPPublic.cs
+++ b/Engine/Source/Programs/AutomationTool/AutomationUtils/MCPPublic.cs
@@ -216,7 +216,7 @@ namespace EpicGames.MCP.Automation
         static public string GetBuildRootPath()
         {
             return CommandUtils.P4Enabled && CommandUtils.AllowSubmit
-                ? CommandUtils.RootSharedTempStorageDirectory()
+                ? CommandUtils.RootBuildStorageDirectory()
                 : CommandUtils.CombinePaths(CommandUtils.CmdEnv.LocalRoot, "LocalBuilds");
         }
 
diff --git a/Engine/Source/Programs/AutomationTool/AutomationUtils/Utils.cs b/Engine/Source/Programs/AutomationTool/AutomationUtils/Utils.cs
--- a/Engine/Source/Programs/AutomationTool/AutomationUtils/Utils.cs
+++ b/Engine/Source/Programs/AutomationTool/AutomationUtils/Utils.cs
@@ -42,6 +42,7 @@ namespace AutomationTool
 
 		/// <summary>
 		/// Creates a directory.
+		/// @todo: this function should not exchange exception context for error codes that can be ignored.
 		/// </summary>
 		/// <param name="Path">Directory name.</param>
 		/// <returns>True if the directory was created, false otherwise.</returns>
@@ -722,7 +723,8 @@ namespace AutomationTool
 	            {
 	                if (Retry > 0)
 	                {
-	                    CommandUtils.LogConsole("*** Mac temp storage retry {0}", OutputFileName);
+                        //@todo: These retries should be reported so we can track how often they are occurring.
+                        CommandUtils.LogConsole("*** Mac temp storage retry {0}", OutputFileName);
 	                    System.Threading.Thread.Sleep(1000);
 	                }
 	                bCopied = CommandUtils.CopyFile_NoExceptions(InputFileName, OutputFileName, true);
@@ -735,12 +737,12 @@ namespace AutomationTool
 	        }
 	    }
 
-	    public static void Robust_FileExists_NoExceptions(string Filename, string Message)
+	    public static void Robust_FileExists(string Filename, string Message)
 	    {
-	        Robust_FileExists_NoExceptions(false, Filename, Message);
+	        Robust_FileExists(false, Filename, Message);
 	    }
 
-	    public static void Robust_FileExists_NoExceptions(bool bQuiet, string Filename, string Message)
+	    public static void Robust_FileExists(bool bQuiet, string Filename, string Message)
 	    {
 	        if (!CommandUtils.FileExists_NoExceptions(bQuiet, Filename))
 	        {
@@ -751,7 +753,8 @@ namespace AutomationTool
 	                int Retry = 0;
 	                while (!bFound && Retry < 60)
 	                {
-	                    CommandUtils.LogConsole("*** Mac temp storage retry {0}", Filename);
+                        //@todo: These retries should be reported so we can track how often they are occurring.
+                        CommandUtils.LogConsole("*** Mac temp storage retry {0}", Filename);
 	                    System.Threading.Thread.Sleep(10000);
 	                    bFound = CommandUtils.FileExists_NoExceptions(bQuiet, Filename);
 	                    Retry++;
@@ -775,7 +778,8 @@ namespace AutomationTool
 	                int Retry = 0;
 	                while (!bFound && Retry < 60)
 	                {
-	                    CommandUtils.LogConsole("*** Mac temp storage retry {0}", Directoryname);
+                        //@todo: These retries should be reported so we can track how often they are occurring.
+                        CommandUtils.LogConsole("*** Mac temp storage retry {0}", Directoryname);
 	                    System.Threading.Thread.Sleep(10000);
 	                    bFound = CommandUtils.DirectoryExists_NoExceptions(Directoryname);
 	                    Retry++;
@@ -805,7 +809,8 @@ namespace AutomationTool
 	                }
 	                while (!bFound && Retry < NumRetries)
 	                {
-	                    CommandUtils.LogConsole("*** Mac temp storage retry {0}", Directoryname);
+                        //@todo: These retries should be reported so we can track how often they are occurring.
+                        CommandUtils.LogConsole("*** Mac temp storage retry {0}", Directoryname);
 	                    System.Threading.Thread.Sleep(1000);
 	                    bFound = CommandUtils.DirectoryExistsAndIsWritable_NoExceptions(Directoryname);
 	                    Retry++;
diff --git a/Engine/Source/Programs/AutomationTool/GUBP/ElectricCommander.cs b/Engine/Source/Programs/AutomationTool/GUBP/ElectricCommander.cs
--- a/Engine/Source/Programs/AutomationTool/GUBP/ElectricCommander.cs
+++ b/Engine/Source/Programs/AutomationTool/GUBP/ElectricCommander.cs
@@ -210,9 +210,9 @@ namespace AutomationTool
 			}
 		}
 
-		public void SaveStatus(BuildNode NodeToDo, string Suffix, string NodeStoreName, bool bSaveSharedTempStorage, string GameNameIfAny, string JobStepIDForFailure = null)
+        public void SaveStatus(TempStorageNodeInfo TempStorageNodeInfo, bool bSaveSharedTempStorage, string GameNameIfAny, string JobStepIDForFailure = null)
 		{
-			string Contents = "Just a status record: " + Suffix;
+			string Contents = "Just a status record: " + TempStorageNodeInfo.NodeStorageName;
 			if (!String.IsNullOrEmpty(JobStepIDForFailure) && CommandUtils.IsBuildMachine)
 			{
 				try
@@ -225,10 +225,10 @@ namespace AutomationTool
 					CommandUtils.LogWarning(LogUtils.FormatException(Ex));
 				}
 			}
-			string RecordOfSuccess = CommandUtils.CombinePaths(CommandUtils.CmdEnv.LocalRoot, "Engine", "Saved", "Logs", NodeToDo.Name + Suffix +".log");
+			string RecordOfSuccess = CommandUtils.CombinePaths(CommandUtils.CmdEnv.LocalRoot, "Engine", "Saved", "Logs", TempStorageNodeInfo.NodeStorageName +".log");
 			CommandUtils.CreateDirectory(Path.GetDirectoryName(RecordOfSuccess));
-			CommandUtils.WriteAllText(RecordOfSuccess, Contents);		
-			TempStorage.StoreToTempStorage(CommandUtils.CmdEnv, NodeStoreName + Suffix, new List<string> { RecordOfSuccess }, !bSaveSharedTempStorage, GameNameIfAny);		
+			CommandUtils.WriteAllText(RecordOfSuccess, Contents);
+            TempStorage.StoreToTempStorage(TempStorageNodeInfo, new List<string> { RecordOfSuccess }, !bSaveSharedTempStorage, GameNameIfAny, CommandUtils.CmdEnv.LocalRoot);		
 		}
 
 		public void UpdateEmailProperties(BuildNode NodeToDo, int LastSucceededCL, string FailedString, string FailCauserEMails, string EMailNote, bool SendSuccessForGreenAfterRed)
@@ -255,7 +255,7 @@ namespace AutomationTool
 			}
 		}
 
-		public void DoCommanderSetup(IEnumerable<BuildNode> AllNodes, IEnumerable<AggregateNode> AllAggregates, List<BuildNode> OrderedToDo, List<BuildNode> SortedNodes, int TimeIndex, int TimeQuantum, bool bSkipTriggers, bool bFake, bool bFakeEC, string CLString, TriggerNode ExplicitTrigger, List<TriggerNode> UnfinishedTriggers, string FakeFail)
+        public void DoCommanderSetup(IEnumerable<BuildNode> AllNodes, IEnumerable<AggregateNode> AllAggregates, List<BuildNode> OrderedToDo, List<BuildNode> SortedNodes, int TimeIndex, int TimeQuantum, bool bSkipTriggers, bool bFake, bool bFakeEC, TriggerNode ExplicitTrigger, List<TriggerNode> UnfinishedTriggers, string FakeFail)
 		{
 			// Check there's some nodes in the graph
 			if (OrderedToDo.Count == 0)
diff --git a/Engine/Source/Programs/AutomationTool/GUBP/GUBP.cs b/Engine/Source/Programs/AutomationTool/GUBP/GUBP.cs
--- a/Engine/Source/Programs/AutomationTool/GUBP/GUBP.cs
+++ b/Engine/Source/Programs/AutomationTool/GUBP/GUBP.cs
@@ -17,8 +17,6 @@ using System.Diagnostics;
 [Help("NoLinux", "Toggle to exclude the Linux (PC, 64-bit) host platform, default is Win64+Mac+Linux")]
 [Help("NoPC", "Toggle to exclude the PC host platform, default is Win64+Mac+Linux")]
 [Help("CleanLocal", "delete the local temp storage before we start")]
-[Help("Store=", "Sets the name of the temp storage block, normally, this is built for you.")]
-[Help("StoreSuffix=", "Tacked onto a store name constructed from CL, branch, etc")]
 [Help("TimeIndex=", "An integer used to determine subsets to run based on DependentCISFrequencyQuantumShift")]
 [Help("UserTimeIndex=", "An integer used to determine subsets to run based on DependentCISFrequencyQuantumShift, this one overrides TimeIndex")]
 [Help("PreflightUID=", "A unique integer tag from EC used as part of the tempstorage, builds and label names to distinguish multiple attempts.")]
@@ -45,13 +43,97 @@ using System.Diagnostics;
 [Help("ForceIncrementalCompile", "make sure all compiles are incremental")]
 [Help("AutomatedTesting", "Allow automated testing, currently disabled.")]
 [Help("StompCheck", "Look for stomped build products.")]
+[Help("NoZipTempStorage", "By default temp storage is zipped up and copied as a single archive. When this flag is set, the directory structure is replicated instead.")]
 public partial class GUBP : BuildCommand
 {
 	const string StartedTempStorageSuffix = "_Started";
 	const string FailedTempStorageSuffix = "_Failed";
 	const string SucceededTempStorageSuffix = "_Succeeded";
 
-	class NodeHistory
+    /// <summary>
+    /// Contains some global attributes of the GUBP job. A GUBP job consists of various nodes that may be executed on many different machines at the same time.
+    /// These attributes are constant throughout the job, for any node in the job.
+    /// </summary>
+    /// <remarks>
+    /// This does not contain ALL the job attributes, just the ones that were added during incremental refactoring.
+    /// 
+    /// This class is immutable to ensure no side-effect code an tamper with it after the job starts.
+    /// </remarks>
+    public class JobInfo
+    {
+        /// <summary>
+        /// Used as the prefix for temp storage for the job.
+        /// 
+        /// If P4 is enabled, this is P4Env.BuildRootEscaped, ie the branch name safe to create a file or folder from. If P4 is not enabled, it is NoP4, 
+        /// but we never write to the shared temp storage location in those cases.
+        /// For testing, arbitrary values can be used.
+        /// </summary>
+        public string BranchNameForTempStorage { get; private set; }
+
+        /// <summary>
+        /// The changelist associated with this job (ie, the code CL that was sync'ed to run this job.
+        /// The -CL parameter can be used to override the P4 environment for testing.
+        /// If P4 is not enabled, this value will be 0.
+        /// </summary>
+        public int Changelist { get; private set; }
+
+        /// <summary>
+        /// If this is a preflight vuild, this is the changelist that contains the shelve being preflighted. 
+        /// Otherwise, will be 0.
+        /// </summary>
+        public int PreflightShelveChangelist { get; private set; }
+
+        /// <summary>
+        /// If this is a preflight build, contains a unique ID to identify the preflight job.
+        /// Used to disambiguate multiple preflight runs from the same changelist and shelve, since the
+        /// code in the shelve can change from run to run.
+        /// </summary>
+        public int PreflightUID { get; private set; }
+
+        /// <summary>
+        /// Will be true if this job is a preflight run. Shortcut for checking if the shelve is > 0
+        /// </summary>
+        public bool IsPreflight { get { return PreflightShelveChangelist > 0; } }
+
+        /// <summary>
+        /// If this is a preflight run, returns a suffix used to describe the preflight shelve and UID. Used
+        /// to uniquely identify a temp storage block for a particular preflight run.
+        /// </summary>
+        /// <returns></returns>
+        public string GetPreflightSuffix()
+        {
+            return IsPreflight ? string.Format("-PF-{0}-{1}", PreflightShelveChangelist, PreflightUID) : "";
+        }
+
+        /// <summary>
+        /// Quick method to create a new JobInfo using all the same attributes, but a different CL.
+        /// 
+        /// Used by some GUBP code to search for and use previous temp storage blocks when checking for node history.
+        /// </summary>
+        /// <param name="NewCL">The CL to use instead of the current Job's CL.</param>
+        /// <returns>A new JobInfo instance with all the same attributes as the current instance, but a different CL.</returns>
+        public JobInfo CreateWithNewChangelist(int NewCL)
+        {
+            return new JobInfo(BranchNameForTempStorage, NewCL, PreflightShelveChangelist, PreflightUID);
+        }
+
+        /// <summary>
+        /// Public ctor. This is an immutable class, so all attributes must be set a construction time.
+        /// </summary>
+        /// <param name="BranchNameForTempStorage">See similarly named attribute.</param>
+        /// <param name="Changelist">See similarly named attribute.</param>
+        /// <param name="PreflightShelveChangelist">See similarly named attribute.</param>
+        /// <param name="PreflightUID">See similarly named attribute.</param>
+        public JobInfo(string BranchNameForTempStorage, int Changelist, int PreflightShelveChangelist, int PreflightUID)
+        {
+            this.BranchNameForTempStorage = BranchNameForTempStorage;
+            this.Changelist = Changelist;
+            this.PreflightShelveChangelist = PreflightShelveChangelist;
+            this.PreflightUID = PreflightUID;
+        }
+    }
+    
+    class NodeHistory
 	{
 		public int LastSucceeded = 0;
 		public int LastFailed = 0;
@@ -71,23 +153,6 @@ public partial class GUBP : BuildCommand
     {
 		LogConsole("************************* GUBP");
 
-		bool bPreflightBuild = false;
-		int PreflightShelveCL = 0;
-        string PreflightShelveCLString = GetEnvVar("uebp_PreflightShelveCL");
-        if ((!String.IsNullOrEmpty(PreflightShelveCLString) && IsBuildMachine) || ParseParam("PreflightTest"))
-        {
-            LogConsole("**** Preflight shelve {0}", PreflightShelveCLString);
-            if (!String.IsNullOrEmpty(PreflightShelveCLString))
-            {
-                PreflightShelveCL = int.Parse(PreflightShelveCLString);
-                if (PreflightShelveCL < 2000000)
-                {
-                    throw new AutomationException(String.Format( "{0} does not look like a CL", PreflightShelveCL));
-                }
-            }
-            bPreflightBuild = true;
-        }
-        
 		List<UnrealTargetPlatform> HostPlatforms = new List<UnrealTargetPlatform>();
         if (!ParseParam("NoPC"))
         {
@@ -101,59 +166,52 @@ public partial class GUBP : BuildCommand
 		{
 			HostPlatforms.Add(UnrealTargetPlatform.Linux);
 		}
+        TempStorage.SetZipTempStorage(!ParseParam("NoZipTempStorage"));
 
-        string StoreName = ParseParamValue("Store");
-        string StoreSuffix = ParseParamValue("StoreSuffix", "");
-
-		string PreflightMangleSuffix = "";
-        if (bPreflightBuild)
-        {
-            int PreflightUID = ParseParamInt("PreflightUID", 0);
-            PreflightMangleSuffix = String.Format("-PF-{0}-{1}", PreflightShelveCL, PreflightUID);
-            StoreSuffix = StoreSuffix + PreflightMangleSuffix;
-        }
-        int CL = ParseParamInt("CL", 0);
         bool bCleanLocalTempStorage = ParseParam("CleanLocal");
         bool bSkipTriggers = ParseParam("SkipTriggers");
         bool bFake = ParseParam("fake");
         bool bFakeEC = ParseParam("FakeEC");
 
-        bool bSaveSharedTempStorage = false;
+        bool bSaveSharedTempStorage = P4Enabled && (IsBuildMachine || GlobalCommandLine.UseLocalBuildStorage);
 
-        bool LocalOnly = true;
-        string CLString = "";
-        if (String.IsNullOrEmpty(StoreName))
+        // encapsulate the logic and temp vars to determine the job info.
+        // Ensures the temp vars don't escape the scope to be used later.
+        Func<JobInfo> GetJobInfo = () =>
         {
-            if (P4Enabled)
+            int PreflightShelveCL = 0;
+            int PreflightUID = 0;
+            string PreflightShelveCLString = GetEnvVar("uebp_PreflightShelveCL");
+            // We must be passed a valid shelve and be a build machine (or be running a preflight test) to be a valid preflight.
+            if (!string.IsNullOrEmpty(PreflightShelveCLString) && (IsBuildMachine || ParseParam("PreflightTest")))
             {
-                if (CL == 0)
+                LogConsole("**** Preflight shelve {0}", PreflightShelveCLString);
+                PreflightShelveCL = int.Parse(PreflightShelveCLString);
+                if (PreflightShelveCL < 2000000)
                 {
-                    CL = P4Env.Changelist;
+                    throw new AutomationException("{0} does not look like a CL", PreflightShelveCL);
                 }
-                CLString = String.Format("{0}", CL);
-                StoreName = P4Env.BuildRootEscaped + "-" + CLString;
-                bSaveSharedTempStorage = CommandUtils.IsBuildMachine;
-                LocalOnly = false;
-            }
-            else
-            {
-                StoreName = "TempLocal";
-                bSaveSharedTempStorage = false;
+                PreflightUID = ParseParamInt("PreflightUID", 0);
             }
-        }
-        StoreName = StoreName + StoreSuffix;
-        if (bFakeEC)
-        {
-            LocalOnly = true;
-        } 
+            return new JobInfo(
+                P4Enabled ? P4Env.BuildRootEscaped : "NoP4",
+                P4Enabled ? ParseParamInt("CL", P4Env.Changelist) : 0,
+                PreflightShelveCL, 
+                PreflightUID);
+        };
+
+        var JobInfo = GetJobInfo();
+
+        bool LocalOnly = !P4Enabled || bFakeEC;
+
         if (bSaveSharedTempStorage)
         {
-            if (!TempStorage.HaveSharedTempStorage(true))
+            if (!TempStorage.IsSharedTempStorageAvailable(true))
             {
-                throw new AutomationException("Request to save to temp storage, but {0} is unavailable.", TempStorage.UE4TempStorageDirectory());
+                throw new AutomationException("Request to save to shared temp storage, but shared temp storage is unavailable or does not exist.");
             }
         }
-        else if (!LocalOnly && !TempStorage.HaveSharedTempStorage(false))
+        else if (!LocalOnly && !TempStorage.IsSharedTempStorageAvailable(false))
         {
             LogWarning("Looks like we want to use shared temp storage, but since we don't have it, we won't use it.");
             LocalOnly = true;
@@ -170,16 +228,16 @@ public partial class GUBP : BuildCommand
 		List<BuildNode> AllNodes;
 		List<AggregateNode> AllAggregates;
 		int TimeQuantum = 20;
-		AddNodesForBranch(CL, HostPlatforms, bPreflightBuild, PreflightMangleSuffix, out AllNodes, out AllAggregates, ref TimeQuantum);
+		AddNodesForBranch(HostPlatforms, JobInfo, out AllNodes, out AllAggregates, ref TimeQuantum);
 
 		LinkGraph(AllAggregates, AllNodes);
 		FindControllingTriggers(AllNodes);
-		FindCompletionState(AllNodes, StoreName, LocalOnly);
+        FindCompletionState(AllNodes, JobInfo, LocalOnly);
 		ComputeDependentFrequencies(AllNodes);
 
         if (bCleanLocalTempStorage)  // shared temp storage can never be wiped
         {
-            TempStorage.DeleteLocalTempStorageManifests(CmdEnv);
+            TempStorage.DeleteLocalTempStorageManifests();
         }
 
 		int TimeIndex = ParseParamInt("TimeIndex", 0);
@@ -195,7 +253,10 @@ public partial class GUBP : BuildCommand
 
 		HashSet<BuildNode> NodesToDo = ParseNodesToDo(AllNodes, AllAggregates);
 		CullNodesForTimeIndex(NodesToDo, TimeIndex);
-		CullNodesForPreflight(NodesToDo, bPreflightBuild);
+        if (JobInfo.IsPreflight)
+        {
+            CullNodesForPreflight(NodesToDo);
+        }
 
 		TriggerNode ExplicitTrigger = null;
 		if (CommanderSetup)
@@ -248,7 +309,7 @@ public partial class GUBP : BuildCommand
 				throw new AutomationException("Couldn't find node {0}", ShowHistoryParam);
 			}
 
-			NodeHistory History = FindNodeHistory(Node, CLString, StoreName);
+            NodeHistory History = FindNodeHistory(Node, JobInfo);
 			if(History == null)
 			{
 				throw new AutomationException("Couldn't get history for {0}", ShowHistoryParam);
@@ -261,11 +322,11 @@ public partial class GUBP : BuildCommand
 			string FakeFail = ParseParamValue("FakeFail");
 			if(CommanderSetup)
 			{
-				DoCommanderSetup(EC, AllNodes, AllAggregates, OrderedToDo, TimeIndex, TimeQuantum, bSkipTriggers, bFake, bFakeEC, CLString, ExplicitTrigger, UnfinishedTriggers, FakeFail, bPreflightBuild);
+				DoCommanderSetup(EC, AllNodes, AllAggregates, OrderedToDo, TimeIndex, TimeQuantum, bSkipTriggers, bFake, bFakeEC, ExplicitTrigger, UnfinishedTriggers, FakeFail, JobInfo.IsPreflight);
 			}
 			else if(!ParseParam("ListOnly"))
 			{
-				ExecuteNodes(EC, OrderedToDo, bFake, bFakeEC, bSaveSharedTempStorage, CLString, StoreName, FakeFail);
+                ExecuteNodes(EC, OrderedToDo, bFake, bFakeEC, bSaveSharedTempStorage, JobInfo, FakeFail);
 			}
 		}
 		PrintRunTime();
@@ -363,34 +424,36 @@ public partial class GUBP : BuildCommand
 		}
 	}
 
-	static void FindCompletionState(IEnumerable<BuildNode> NodesToDo, string StoreName, bool LocalOnly)
+    static void FindCompletionState(IEnumerable<BuildNode> NodesToDo, JobInfo JobInfo, bool LocalOnly)
 	{
 		foreach(BuildNode NodeToDo in NodesToDo)
 		{
-			string NodeStoreName = StoreName + "-" + NodeToDo.Name;
+            // Construct the full temp storage node info
+            var TempStorageNodeInfo = new TempStorageNodeInfo(JobInfo, NodeToDo.Name);
+
 			string GameNameIfAny = NodeToDo.GameNameIfAnyForTempStorage;
 
 			if (LocalOnly)
 			{
-				NodeToDo.IsComplete = TempStorage.LocalTempStorageExists(CmdEnv, NodeStoreName, bQuiet : true);
+                NodeToDo.IsComplete = TempStorage.LocalTempStorageManifestExists(TempStorageNodeInfo, bQuiet: true);
 			}
 			else
 			{
-				NodeToDo.IsComplete = TempStorage.TempStorageExists(CmdEnv, NodeStoreName, GameNameIfAny, bQuiet: true);
+                NodeToDo.IsComplete = TempStorage.TempStorageExists(TempStorageNodeInfo, GameNameIfAny, bLocalOnly: false, bQuiet: true);
 				if(GameNameIfAny != "" && !NodeToDo.IsComplete)
 				{
-					NodeToDo.IsComplete = TempStorage.TempStorageExists(CmdEnv, NodeStoreName, "", bQuiet: true);
+                    NodeToDo.IsComplete = TempStorage.TempStorageExists(TempStorageNodeInfo, "UE4", bLocalOnly: false, bQuiet: true);
 				}
 			}
 
             LogVerbose("** {0}", NodeToDo.Name);
 			if (!NodeToDo.IsComplete)
 			{
-				LogVerbose("***** GUBP Trigger Node was already triggered {0} -> {1} : {2}", NodeToDo.Name, GameNameIfAny, NodeStoreName);
+                LogVerbose("***** GUBP Trigger Node was already triggered {0} -> {1} : {2}", NodeToDo.Name, GameNameIfAny, TempStorageNodeInfo.GetRelativeDirectory());
 			}
 			else
 			{
-				LogVerbose("***** GUBP Trigger Node was NOT yet triggered {0} -> {1} : {2}", NodeToDo.Name, GameNameIfAny, NodeStoreName);
+                LogVerbose("***** GUBP Trigger Node was NOT yet triggered {0} -> {1} : {2}", NodeToDo.Name, GameNameIfAny, TempStorageNodeInfo.GetRelativeDirectory());
 			}
 		}
     }
@@ -583,22 +646,6 @@ public partial class GUBP : BuildCommand
         }
     }
 
-    static List<int> ConvertCLToIntList(List<string> Strings)
-    {
-        List<int> Result = new List<int>();
-        foreach (string ThisString in Strings)
-        {
-            int ThisInt = int.Parse(ThisString);
-            if (ThisInt < 1960000 || ThisInt > 3000000)
-            {
-				LogConsole("CL {0} appears to be out of range", ThisInt);
-            }
-            Result.Add(ThisInt);
-        }
-        Result.Sort();
-        return Result;
-    }
-
     int CountZeros(int Num)
     {
         if (Num < 0)
@@ -853,30 +900,29 @@ public partial class GUBP : BuildCommand
         return OrderedToDo;
     }
 
-    static NodeHistory FindNodeHistory(BuildNode NodeToDo, string CLString, string StoreName)
+    static NodeHistory FindNodeHistory(BuildNode NodeToDo, JobInfo JobInfo)
     {
 		NodeHistory History = null;
 
-        if (!(NodeToDo is TriggerNode) && CLString != "")
+        // Don't get node history on nodes that don't have a valid CL.
+        if (!(NodeToDo is TriggerNode) && JobInfo.Changelist > 0)
         {
             string GameNameIfAny = NodeToDo.GameNameIfAnyForTempStorage;
-            string NodeStoreWildCard = StoreName.Replace(CLString, "*") + "-" + NodeToDo.Name;
-            History = new NodeHistory();
 
-            History.AllStarted = ConvertCLToIntList(TempStorage.FindTempStorageManifests(CmdEnv, NodeStoreWildCard + StartedTempStorageSuffix, false, true, GameNameIfAny));
-            History.AllSucceeded = ConvertCLToIntList(TempStorage.FindTempStorageManifests(CmdEnv, NodeStoreWildCard + SucceededTempStorageSuffix, false, true, GameNameIfAny));
-            History.AllFailed = ConvertCLToIntList(TempStorage.FindTempStorageManifests(CmdEnv, NodeStoreWildCard + FailedTempStorageSuffix, false, true, GameNameIfAny));
-
-			int CL;
-			int.TryParse(CLString, out CL);
+            History = new NodeHistory
+            {
+                AllStarted = TempStorage.FindMatchingSharedTempStorageNodeCLs(new TempStorageNodeInfo(JobInfo, NodeToDo.Name + StartedTempStorageSuffix), GameNameIfAny),
+                AllSucceeded = TempStorage.FindMatchingSharedTempStorageNodeCLs(new TempStorageNodeInfo(JobInfo, NodeToDo.Name + SucceededTempStorageSuffix), GameNameIfAny),
+                AllFailed = TempStorage.FindMatchingSharedTempStorageNodeCLs(new TempStorageNodeInfo(JobInfo, NodeToDo.Name + FailedTempStorageSuffix), GameNameIfAny)
+            };
 
             if (History.AllFailed.Count > 0)
             {
-                History.LastFailed = History.AllFailed.LastOrDefault(x => x < CL);
+                History.LastFailed = History.AllFailed.LastOrDefault(x => x < JobInfo.Changelist);
             }
             if (History.AllSucceeded.Count > 0)
             {
-                History.LastSucceeded = History.AllSucceeded.LastOrDefault(x => x < CL);
+                History.LastSucceeded = History.AllSucceeded.LastOrDefault(x => x < JobInfo.Changelist);
 
                 foreach (int Failed in History.AllFailed)
                 {
@@ -899,7 +945,7 @@ public partial class GUBP : BuildCommand
 
 		return History;
     }
-	void GetFailureEmails(ElectricCommander EC, BuildNode NodeToDo, NodeHistory History, string CLString, string StoreName, bool OnlyLateUpdates = false)
+    void GetFailureEmails(ElectricCommander EC, BuildNode NodeToDo, NodeHistory History, JobInfo JobInfo, bool OnlyLateUpdates = false)
 	{
         string FailCauserEMails = "";
         string EMailNote = "";
@@ -914,7 +960,7 @@ public partial class GUBP : BuildCommand
                 {
                     if (OnlyLateUpdates)
                     {
-						LastNonDuplicateFail = FindLastNonDuplicateFail(NodeToDo, History, CLString, StoreName);
+                        LastNonDuplicateFail = FindLastNonDuplicateFail(NodeToDo, History, JobInfo);
                         if (LastNonDuplicateFail < P4Env.Changelist)
                         {
 							LogConsole("*** Red-after-red spam reduction, changed CL {0} to CL {1} because the errors didn't change.", P4Env.Changelist, LastNonDuplicateFail);
@@ -1004,12 +1050,12 @@ public partial class GUBP : BuildCommand
         return true;
     }
 
-	int FindLastNonDuplicateFail(BuildNode NodeToDo, NodeHistory History, string CLString, string StoreName)
+    int FindLastNonDuplicateFail(BuildNode NodeToDo, NodeHistory History, JobInfo JobInfo)
     {
         int Result = P4Env.Changelist;
 
         string GameNameIfAny = NodeToDo.GameNameIfAnyForTempStorage;
-        string NodeStore = StoreName + "-" + NodeToDo.Name + FailedTempStorageSuffix;
+        var TempStorageNodeInfo = new TempStorageNodeInfo(JobInfo, NodeToDo.Name + FailedTempStorageSuffix);
 
         List<int> BackwardsFails = new List<int>(History.AllFailed);
         BackwardsFails.Add(P4Env.Changelist);
@@ -1026,14 +1072,15 @@ public partial class GUBP : BuildCommand
             {
                 break;
             }
-            string ThisNodeStore = NodeStore.Replace(CLString, String.Format("{0}", CL));
-            TempStorage.DeleteLocalTempStorage(CmdEnv, ThisNodeStore, true); // these all clash locally, which is fine we just retrieve them from shared
+            // Find any local temp storage manifest for this changelist and delete it.
+            var ThisTempStorageNodeInfo = new TempStorageNodeInfo(JobInfo.CreateWithNewChangelist(CL), TempStorageNodeInfo.NodeStorageName);
+            TempStorage.DeleteLocalTempStorageManifest(ThisTempStorageNodeInfo, true); // these all clash locally, which is fine we just retrieve them from shared
 
             List<string> Files = null;
             try
             {
                 bool WasLocal;
-                Files = TempStorage.RetrieveFromTempStorage(CmdEnv, ThisNodeStore, out WasLocal, GameNameIfAny); // this will fail on our CL if we didn't fail or we are just setting up the branch
+                Files = TempStorage.RetrieveFromTempStorage(ThisTempStorageNodeInfo, out WasLocal, GameNameIfAny, CmdEnv.LocalRoot); // this will fail on our CL if we didn't fail or we are just setting up the branch
             }
             catch (Exception)
             {
@@ -1318,22 +1365,18 @@ public partial class GUBP : BuildCommand
 	/// Culls everything downstream of a trigger if we're running a preflight build.
 	/// </summary>
 	/// <param name="NodesToDo">The current list of nodes to do</param>
-	/// <param name="IsPreflightBuild">Whether this is a preflight build or not</param>
-	private void CullNodesForPreflight(HashSet<BuildNode> NodesToDo, bool IsPreflightBuild)
+	private void CullNodesForPreflight(HashSet<BuildNode> NodesToDo)
 	{
-		if (IsPreflightBuild)
+		HashSet<BuildNode> NodesToCull = new HashSet<BuildNode>();
+		foreach(BuildNode NodeToDo in NodesToDo)
 		{
-			HashSet<BuildNode> NodesToCull = new HashSet<BuildNode>();
-			foreach(BuildNode NodeToDo in NodesToDo)
+			if((NodeToDo is TriggerNode) || NodeToDo.ControllingTriggers.Length > 0)
 			{
-				if((NodeToDo is TriggerNode) || NodeToDo.ControllingTriggers.Length > 0)
-				{
-					LogVerbose(" Culling {0} due to being downstream of trigger in preflight", NodeToDo.Name);
-					NodesToCull.Add(NodeToDo);
-				}
+				LogVerbose(" Culling {0} due to being downstream of trigger in preflight", NodeToDo.Name);
+				NodesToCull.Add(NodeToDo);
 			}
-			NodesToDo.ExceptWith(NodesToCull);
 		}
+		NodesToDo.ExceptWith(NodesToCull);
 	}
 
 	private int UpdateCISCounter(int TimeQuantum)
@@ -1491,7 +1534,7 @@ public partial class GUBP : BuildCommand
 		}
 	}
 
-	private void DoCommanderSetup(ElectricCommander EC, IEnumerable<BuildNode> AllNodes, IEnumerable<AggregateNode> AllAggregates, List<BuildNode> OrderedToDo, int TimeIndex, int TimeQuantum, bool bSkipTriggers, bool bFake, bool bFakeEC, string CLString, TriggerNode ExplicitTrigger, List<TriggerNode> UnfinishedTriggers, string FakeFail, bool bPreflightBuild)
+    private void DoCommanderSetup(ElectricCommander EC, IEnumerable<BuildNode> AllNodes, IEnumerable<AggregateNode> AllAggregates, List<BuildNode> OrderedToDo, int TimeIndex, int TimeQuantum, bool bSkipTriggers, bool bFake, bool bFakeEC, TriggerNode ExplicitTrigger, List<TriggerNode> UnfinishedTriggers, string FakeFail, bool bPreflightBuild)
 	{
 		List<BuildNode> SortedNodes = TopologicalSort(new HashSet<BuildNode>(AllNodes), null, SubSort: false, DoNotConsiderCompletion: true);
 		LogConsole("******* {0} GUBP Nodes", SortedNodes.Count);
@@ -1526,10 +1569,10 @@ public partial class GUBP : BuildCommand
 			PrintNodes(this, FilteredOrderedToDo, AllAggregates, UnfinishedTriggers, TimeQuantum);
 		}
 
-		EC.DoCommanderSetup(AllNodes, AllAggregates, FilteredOrderedToDo, SortedNodes, TimeIndex, TimeQuantum, bSkipTriggers, bFake, bFakeEC, CLString, ExplicitTrigger, UnfinishedTriggers, FakeFail);
+		EC.DoCommanderSetup(AllNodes, AllAggregates, FilteredOrderedToDo, SortedNodes, TimeIndex, TimeQuantum, bSkipTriggers, bFake, bFakeEC, ExplicitTrigger, UnfinishedTriggers, FakeFail);
 	}
 
-	void ExecuteNodes(ElectricCommander EC, List<BuildNode> OrderedToDo, bool bFake, bool bFakeEC, bool bSaveSharedTempStorage, string CLString, string StoreName, string FakeFail)
+    void ExecuteNodes(ElectricCommander EC, List<BuildNode> OrderedToDo, bool bFake, bool bFakeEC, bool bSaveSharedTempStorage, JobInfo JobInfo, string FakeFail)
 	{
         Dictionary<string, BuildNode> BuildProductToNodeMap = new Dictionary<string, BuildNode>();
 		foreach (BuildNode NodeToDo in OrderedToDo)
@@ -1539,7 +1582,7 @@ public partial class GUBP : BuildCommand
                 throw new AutomationException("topological sort error");
             }
 
-            string NodeStoreName = StoreName + "-" + NodeToDo.Name;
+            var TempStorageNodeInfo = new TempStorageNodeInfo(JobInfo, NodeToDo.Name);
             
             string GameNameIfAny = NodeToDo.GameNameIfAnyForTempStorage;
             string StorageRootIfAny = NodeToDo.RootIfAnyForTempStorage;
@@ -1548,37 +1591,41 @@ public partial class GUBP : BuildCommand
             {
                 StorageRootIfAny = ""; // we don't rebase fake runs since those are entirely "records of success", which are always in the logs folder
             }
+            if (string.IsNullOrEmpty(StorageRootIfAny))
+            {
+                StorageRootIfAny = CmdEnv.LocalRoot;
+            }
 
             // this is kinda complicated
             bool SaveSuccessRecords = (IsBuildMachine || bFakeEC) && // no real reason to make these locally except for fakeEC tests
                 (!(NodeToDo is TriggerNode) || NodeToDo.IsSticky); // trigger nodes are run twice, one to start the new workflow and once when it is actually triggered, we will save reconds for the latter
 
-			LogConsole("***** Running GUBP Node {0} -> {1} : {2}", NodeToDo.Name, GameNameIfAny, NodeStoreName);
+            LogConsole("***** Running GUBP Node {0} -> {1} : {2}", NodeToDo.Name, GameNameIfAny, TempStorageNodeInfo.GetRelativeDirectory());
             if (NodeToDo.IsComplete)
             {
                 if (NodeToDo.Name == VersionFilesNode.StaticGetFullName() && !IsBuildMachine)
                 {
-					LogConsole("***** NOT ****** Retrieving GUBP Node {0} from {1}; it is the version files.", NodeToDo.Name, NodeStoreName);
+                    LogConsole("***** NOT ****** Retrieving GUBP Node {0} from {1}; it is the version files.", NodeToDo.Name, TempStorageNodeInfo.GetRelativeDirectory());
                     NodeToDo.Node.BuildProducts = new List<string>();
 
                 }
                 else
                 {
-					LogConsole("***** Retrieving GUBP Node {0} from {1}", NodeToDo.Name, NodeStoreName);
+                    LogConsole("***** Retrieving GUBP Node {0} from {1}", NodeToDo.Name, TempStorageNodeInfo.GetRelativeDirectory());
                     bool WasLocal;
 					try
 					{
-						NodeToDo.Node.BuildProducts = TempStorage.RetrieveFromTempStorage(CmdEnv, NodeStoreName, out WasLocal, GameNameIfAny, StorageRootIfAny);
+                        NodeToDo.Node.BuildProducts = TempStorage.RetrieveFromTempStorage(TempStorageNodeInfo, out WasLocal, GameNameIfAny, StorageRootIfAny);
 					}
-					catch
+					catch (Exception Ex)
 					{
 						if(GameNameIfAny != "")
 						{
-							NodeToDo.Node.BuildProducts = TempStorage.RetrieveFromTempStorage(CmdEnv, NodeStoreName, out WasLocal, "", StorageRootIfAny);
+                            NodeToDo.Node.BuildProducts = TempStorage.RetrieveFromTempStorage(TempStorageNodeInfo, out WasLocal, "", StorageRootIfAny);
 						}
 						else
 						{
-							throw new AutomationException("Build Products cannot be found for node {0}", NodeToDo.Name);
+							throw new AutomationException(Ex, "Build Products cannot be found for node {0}", NodeToDo.Name);
 						}
 					}
                     if (!WasLocal)
@@ -1591,7 +1638,8 @@ public partial class GUBP : BuildCommand
             {
                 if (SaveSuccessRecords) 
                 {
-                    EC.SaveStatus(NodeToDo, StartedTempStorageSuffix, NodeStoreName, bSaveSharedTempStorage, GameNameIfAny);
+                    // We save our status to a new temp storage location specifically named with a suffix so we can find it later.
+                    EC.SaveStatus(new TempStorageNodeInfo(JobInfo, NodeToDo.Name + StartedTempStorageSuffix), bSaveSharedTempStorage, GameNameIfAny);
                 }
 				double BuildDuration = 0.0;
                 try
@@ -1602,12 +1650,12 @@ public partial class GUBP : BuildCommand
                     }
                     if (bFake)
                     {
-						LogConsole("***** FAKE!! Building GUBP Node {0} for {1}", NodeToDo.Name, NodeStoreName);
+                        LogConsole("***** FAKE!! Building GUBP Node {0} for {1}", NodeToDo.Name, TempStorageNodeInfo.GetRelativeDirectory());
                         NodeToDo.DoFakeBuild();
                     }
                     else
                     {
-						LogConsole("***** Building GUBP Node {0} for {1}", NodeToDo.Name, NodeStoreName);
+                        LogConsole("***** Building GUBP Node {0} for {1}", NodeToDo.Name, TempStorageNodeInfo.GetRelativeDirectory());
 						DateTime StartTime = DateTime.UtcNow;
 						using(TelemetryStopwatch DoBuildStopwatch = new TelemetryStopwatch("DoBuild.{0}", NodeToDo.Name))
 						{
@@ -1616,21 +1664,23 @@ public partial class GUBP : BuildCommand
 						BuildDuration = (DateTime.UtcNow - StartTime).TotalMilliseconds / 1000;
                     }
 
-                    TempStorage.StoreToTempStorage(CmdEnv, NodeStoreName, NodeToDo.Node.BuildProducts, !bSaveSharedTempStorage, GameNameIfAny, StorageRootIfAny);
+                    TempStorage.StoreToTempStorage(TempStorageNodeInfo, NodeToDo.Node.BuildProducts, !bSaveSharedTempStorage, GameNameIfAny, StorageRootIfAny);
                 }
                 catch (Exception Ex)
                 {
+                    //@todo: This is a dup of non-exception code. Consolidate this!!
 					NodeHistory History = null;
 
                     if (SaveSuccessRecords)
                     {
 						using(TelemetryStopwatch UpdateNodeHistoryStopwatch = new TelemetryStopwatch("UpdateNodeHistory"))
 						{
-							History = FindNodeHistory(NodeToDo, CLString, StoreName);
+                            History = FindNodeHistory(NodeToDo, JobInfo);
 						}
 						using(TelemetryStopwatch SaveNodeStatusStopwatch = new TelemetryStopwatch("SaveNodeStatus"))
 						{
-							EC.SaveStatus(NodeToDo, FailedTempStorageSuffix, NodeStoreName, bSaveSharedTempStorage, GameNameIfAny, ParseParamValue("MyJobStepId"));
+                            // We save our status to a new temp storage location specifically named with a suffix so we can find it later.
+                            EC.SaveStatus(new TempStorageNodeInfo(JobInfo, NodeToDo.Name + FailedTempStorageSuffix), bSaveSharedTempStorage, GameNameIfAny, ParseParamValue("MyJobStepId"));
 						}
 						using(TelemetryStopwatch UpdateECPropsStopwatch = new TelemetryStopwatch("UpdateECProps"))
 						{
@@ -1641,13 +1691,13 @@ public partial class GUBP : BuildCommand
 						{
 							using(TelemetryStopwatch GetFailEmailsStopwatch = new TelemetryStopwatch("GetFailEmails"))
 							{
-								GetFailureEmails(EC, NodeToDo, History, CLString, StoreName);
+                                GetFailureEmails(EC, NodeToDo, History, JobInfo);
 							}
 						}
 						EC.UpdateECBuildTime(NodeToDo, BuildDuration);
                     }
 
-					LogConsole("{0}", ExceptionToString(Ex));
+					LogConsole("{0}", Ex);
 
 
                     if (History != null)
@@ -1696,11 +1746,12 @@ public partial class GUBP : BuildCommand
 					NodeHistory History = null;
 					using(TelemetryStopwatch UpdateNodeHistoryStopwatch = new TelemetryStopwatch("UpdateNodeHistory"))
 					{
-						History = FindNodeHistory(NodeToDo, CLString, StoreName);
+                        History = FindNodeHistory(NodeToDo, JobInfo);
 					}
 					using(TelemetryStopwatch SaveNodeStatusStopwatch = new TelemetryStopwatch("SaveNodeStatus"))
 					{
-						EC.SaveStatus(NodeToDo, SucceededTempStorageSuffix, NodeStoreName, bSaveSharedTempStorage, GameNameIfAny);
+                        // We save our status to a new temp storage location specifically named with a suffix so we can find it later.
+                        EC.SaveStatus(new TempStorageNodeInfo(JobInfo, NodeToDo.Name + SucceededTempStorageSuffix), bSaveSharedTempStorage, GameNameIfAny);
 					}
 					using(TelemetryStopwatch UpdateECPropsStopwatch = new TelemetryStopwatch("UpdateECProps"))
 					{
@@ -1711,7 +1762,7 @@ public partial class GUBP : BuildCommand
 					{
 						using(TelemetryStopwatch GetFailEmailsStopwatch = new TelemetryStopwatch("GetFailEmails"))
 						{
-							GetFailureEmails(EC, NodeToDo, History, CLString, StoreName);
+                            GetFailureEmails(EC, NodeToDo, History, JobInfo);
 						}
 					}
 					EC.UpdateECBuildTime(NodeToDo, BuildDuration);
diff --git a/Engine/Source/Programs/AutomationTool/GUBP/LegacyBranchSetup.cs b/Engine/Source/Programs/AutomationTool/GUBP/LegacyBranchSetup.cs
--- a/Engine/Source/Programs/AutomationTool/GUBP/LegacyBranchSetup.cs
+++ b/Engine/Source/Programs/AutomationTool/GUBP/LegacyBranchSetup.cs
@@ -13,28 +13,24 @@ partial class GUBP
 {
 	public class GUBPBranchConfig
 	{
-		public int CL = 0;
 		public List<UnrealTargetPlatform> HostPlatforms;
 		public string BranchName;
 		public BranchInfo Branch;
 		public GUBPBranchHacker.BranchOptions BranchOptions;
 		public bool bForceIncrementalCompile;
-		public bool bPreflightBuild;
-		public string PreflightMangleSuffix;
+        public JobInfo JobInfo;
 
 		public Dictionary<string, GUBPNode> GUBPNodes = new Dictionary<string, GUBPNode>();
 		public Dictionary<string, GUBPAggregateNode> GUBPAggregates = new Dictionary<string, GUBPAggregateNode>();
 
-		public GUBPBranchConfig(int InCL, IEnumerable<UnrealTargetPlatform> InHostPlatforms, string InBranchName, BranchInfo InBranch, GUBPBranchHacker.BranchOptions InBranchOptions, bool bInForceIncrementalCompile, bool bInPreflightBuild, string InPreflightMangleSuffix)
+        public GUBPBranchConfig(IEnumerable<UnrealTargetPlatform> InHostPlatforms, string InBranchName, BranchInfo InBranch, GUBPBranchHacker.BranchOptions InBranchOptions, bool bInForceIncrementalCompile, JobInfo JobInfo)
 		{
-			CL = InCL;
 			HostPlatforms = new List<UnrealTargetPlatform>(InHostPlatforms);
 			BranchName = InBranchName;
 			Branch = InBranch;
 			BranchOptions = InBranchOptions;
 			bForceIncrementalCompile = bInForceIncrementalCompile;
-			bPreflightBuild = bInPreflightBuild;
-			PreflightMangleSuffix = InPreflightMangleSuffix;
+			this.JobInfo = JobInfo;
 		}
 
 		public string AddNode(GUBPNode Node)
@@ -345,7 +341,7 @@ partial class GUBP
         return AltHostPlatform;
     }
 
-	void AddNodesForBranch(int CL, List<UnrealTargetPlatform> InitialHostPlatforms, bool bPreflightBuild, string PreflightMangleSuffix, out List<BuildNode> AllNodes, out List<AggregateNode> AllAggregates, ref int TimeQuantum)
+    void AddNodesForBranch(List<UnrealTargetPlatform> InitialHostPlatforms, JobInfo JobInfo, out List<BuildNode> AllNodes, out List<AggregateNode> AllAggregates, ref int TimeQuantum)
 	{
 		string BranchName;
 		if (P4Enabled)
@@ -453,7 +449,7 @@ partial class GUBP
             }
         }
 
-		GUBPBranchConfig BranchConfig = new GUBPBranchConfig(CL, HostPlatforms, BranchName, Branch, BranchOptions, bForceIncrementalCompile, bPreflightBuild, PreflightMangleSuffix);
+		GUBPBranchConfig BranchConfig = new GUBPBranchConfig(HostPlatforms, BranchName, Branch, BranchOptions, bForceIncrementalCompile, JobInfo);
 
         BranchConfig.AddNode(new VersionFilesNode());
 
@@ -1179,7 +1175,8 @@ partial class GUBP
 			{
 				BranchConfig.AddNode(new RootEditorCrossCompileLinuxNode(BranchConfig, UnrealTargetPlatform.Win64));
 			}
-            if (!bPreflightBuild)
+            // Don't run clean on temp shared storage for preflight builds. They might have broken the clean, and it extends to storage beyond this job.
+            if (!BranchConfig.JobInfo.IsPreflight)
             {
                 BranchConfig.AddNode(new CleanSharedTempStorageNode(this, BranchConfig));
             }
diff --git a/Engine/Source/Programs/AutomationTool/GUBP/LegacyNodes.cs b/Engine/Source/Programs/AutomationTool/GUBP/LegacyNodes.cs
--- a/Engine/Source/Programs/AutomationTool/GUBP/LegacyNodes.cs
+++ b/Engine/Source/Programs/AutomationTool/GUBP/LegacyNodes.cs
@@ -2004,7 +2004,7 @@ partial class GUBP
         {
             BuildProducts = new List<string>();
 
-			if (P4Enabled && !BranchConfig.bPreflightBuild)
+			if (P4Enabled && !BranchConfig.JobInfo.IsPreflight)
             {
                 if (AllDependencyBuildProducts.Count == 0)
                 {
@@ -2018,8 +2018,8 @@ partial class GUBP
                 }
                 else
                 {
-					int WorkingCL = P4.CreateChange(P4Env.Client, String.Format("GUBP Node {0} built from changelist {1}", GetFullName(), BranchConfig.CL));
-                    LogConsole("Build from {0}    Working in {1}", BranchConfig.CL, WorkingCL);
+					int WorkingCL = P4.CreateChange(P4Env.Client, String.Format("GUBP Node {0} built from changelist {1}", GetFullName(), BranchConfig.JobInfo.Changelist));
+                    LogConsole("Build from {0}    Working in {1}", BranchConfig.JobInfo.Changelist, WorkingCL);
 
                     var ProductsToSubmit = new List<String>();
 
@@ -2553,11 +2553,7 @@ partial class GUBP
         {
             string BaseDir = TempStorage.ResolveSharedBuildDirectory(InGameProj.GameName);
             string NodeName = StaticGetFullName(InGameProj, InHostPlatform, InClientTargetPlatforms, InClientConfigs, InServerTargetPlatforms, InServerConfigs, InClientNotGame);
-            string Inner = P4Env.BuildRootEscaped + "-CL-" + P4Env.ChangelistString;
-            if (BranchConfig.bPreflightBuild)
-            {
-                Inner = Inner + BranchConfig.PreflightMangleSuffix;
-            }
+            string Inner = P4Env.BuildRootEscaped + "-CL-" + P4Env.ChangelistString + BranchConfig.JobInfo.GetPreflightSuffix();
             string ArchiveDirectory = CombinePaths(BaseDir, NodeName, Inner);
             return ArchiveDirectory;
         }
@@ -3240,10 +3236,7 @@ partial class GUBP
         {
             {
                 var StartTime = DateTime.UtcNow;
-                foreach (string GameName in GameNames)
-                {
-                    TempStorage.CleanSharedTempStorageDirectory(GameName);
-                }
+                TempStorage.CleanSharedTempStorageDirectory(GameNames);
                 var BuildDuration = (DateTime.UtcNow - StartTime).TotalMilliseconds;
                 LogConsole("Took {0}s to clear temp storage of old files.", BuildDuration / 1000);
             }
diff --git a/Engine/Source/Programs/AutomationTool/GUBP/TempStorage.cs b/Engine/Source/Programs/AutomationTool/GUBP/TempStorage.cs
--- a/Engine/Source/Programs/AutomationTool/GUBP/TempStorage.cs
+++ b/Engine/Source/Programs/AutomationTool/GUBP/TempStorage.cs
@@ -5,30 +5,138 @@ using System.Text;
 using System.Linq;
 using System.IO;
 using System.Xml;
+using System.Xml.Linq;
 
 namespace AutomationTool
 {
-
-    public class TempStorage
+    /// <summary>
+    /// Describes temp storage for a particular GUBP Node being run. A GUBP job consists of many nodes organized in a dependency graph that run in a distributed fashion.
+    /// Each node has the ability to store files into a shared temp storage location specifically for that node.
+    /// When a later node runs, it can retrieve those temp storage files and use them to bootstrap itself so it doesn't have to redo the work of earlier nodes.
+    /// A Node's temp storage location is determined from the global GUBP Job Info along with the node's "storage name", which is typically the name of the GUBP node.
+    /// However, some nodes create additional "pseudo-nodes" to store things like completion status of the node, so future jobs can find their history.
+    /// Immutable.
+    /// </summary>
+    /// <remarks>
+    /// This class is immutable to ensure no side-effect code an tamper with it after the node starts.
+    /// 
+    /// @todo: This storage naming convention relies on the build system preventing jobs from being run more than once for a given CL. We should be putting a Job UID in here or something so a job can run twice!!
+    /// </remarks>
+    public class TempStorageNodeInfo
     {
-        public class TempStorageFileInfo
+        /// <summary>
+        /// The GUBP global job attributes.
+        /// </summary>
+        public GUBP.JobInfo JobInfo { get; private set; }
+
+        /// <summary>
+        /// The storage name associated with this node. It is usually the name of the GUBP node running, but sometimes certain suffixes are added to indicate node completion status, etc.
+        /// </summary>
+        public string NodeStorageName { get; private set; }
+
+        /// <summary>
+        /// Constructor. This class is immutable, so all attributes must be fully specified here.
+        /// </summary>
+        /// <param name="JobInfo"></param>
+        /// <param name="NodeStorageName"></param>
+        public TempStorageNodeInfo(GUBP.JobInfo JobInfo, string NodeStorageName)
         {
-            public string Name;
-            public DateTime Timestamp;
-            public long Size;
+            this.JobInfo = JobInfo;
+            this.NodeStorageName = NodeStorageName;
+        }
 
-            public TempStorageFileInfo()
+        /// <summary>
+        /// Cache of the legacy string. These values are not super cheap to construct, and can get rather long, so we only cache the value once it is asked for.
+        /// </summary>
+        private string CachedString;
+
+        /// <summary>
+        /// Cache of the relative directory. These values are not super cheap to construct, and can get rather long, so we only cache the value once it is asked for.
+        /// </summary>
+        private string CachedDir;
+
+        /// <summary>
+        /// Cache of the filename. These values are not super cheap to construct, and can get rather long, so we only cache the value once it is asked for.
+        /// </summary>
+        private string CachedFile;
+
+        /// <summary>
+        /// Used by anything that needs the block info as a string. This should generally be legacy stuff, as we should only be using the temp storage path for temp storage!
+        /// </summary>
+        /// <returns>{BuildRootForTempStorage}-{Changelist}{PreflightInfo}-{NodeStorageName}</returns>
+        public string GetLegacyString()
+        {
+            if (string.IsNullOrEmpty(CachedString))
+            {
+                CachedString = string.Format("{0}-{1}{2}-{3}", JobInfo.BranchNameForTempStorage, JobInfo.Changelist, JobInfo.GetPreflightSuffix(), NodeStorageName);
+            }
+            return CachedString;
+        }
+
+        /// <summary>
+        /// Turns the block info into a relative path to a folder where temp storage should be placed.
+        /// NOTE: <see cref="FindTempStorageNodeCLsMatchingSuffix"/> is dependent on this directory layout!!!
+        /// If this layout is changed, FindTempStorageNodeCLsMatchingSuffix needs to be changed as well.
+        /// </summary>
+        /// <returns>{BuildRootForTempStorage}/{Changelist}{PreflightInfo}/{NodeStorageName}</returns>
+        public string GetRelativeDirectory()
+        {
+            if (string.IsNullOrEmpty(CachedDir))
             {
+                CachedDir = CommandUtils.CombinePaths(JobInfo.BranchNameForTempStorage, string.Format("{0}{1}", JobInfo.Changelist, JobInfo.GetPreflightSuffix()), NodeStorageName);
             }
+            return CachedDir;
+        }
 
-            public TempStorageFileInfo(string Filename, string RelativeFile)
+        /// <summary>
+        /// Turns the block info into a file name that will be used to store the temp manifest for that block.
+        /// </summary>
+        /// <returns>{NodeStorageName}.TempManifest</returns>
+        public string GetManifestFilename()
+        {
+            if (string.IsNullOrEmpty(CachedFile))
             {
-                FileInfo Info = new FileInfo(Filename);
-                Name = RelativeFile.Replace("\\", "/"); 
-                Timestamp = Info.LastWriteTimeUtc;
-                Size = Info.Length;
+                CachedFile = string.Format("{0}.TempManifest", NodeStorageName);
             }
+            return CachedFile;
+        }
+    }
+
+    /// <summary>
+    /// Handles temp storage duties for GUBP. Temp storage is a centralized location where build products from a node are stored
+    /// so that dependent nodes can reuse them, even if that node is not the same machine that built them. The main entry points 
+    /// are <see cref="StoreToTempStorage"/> and <see cref="RetrieveFromTempStorage"/>. The former is called when a node is complete
+    /// to store its build products, and the latter is called when a node starts up for all its dependent nodes to get the previous
+    /// build products.
+    /// 
+    /// Each node writes it's temp storage to a folder whose path is defined by certain attributes of the job being run and the name of the node.
+    /// The contents are the temp data itself (stored in a single .zip file unless -NoZipTempStorage is used), and a .TempManifest file that describes all the files stored
+    /// for this node, along with their expected timestamps and sizes. This is used to verify the files retrieved and stored match what is expected.
+    /// 
+    /// To support scenarios where the node running the step is the same node as the dependent one (agent sharing groups exist to ensure this),
+    /// the system always leaves a copy of the .TempManifest in "local storage", which is a local folder (see <see cref="LocalTempStorageManifestDirectory"/>).
+    /// When asked to retrieve a node from temp storage, the system first checks that local storage to see if the manifest exists there.
+    /// If it does, it assumes this machine already has the required temp storage. Beyond just verifying the files in the manifest exist properly, nothing
+    /// is copied in this case. The function them returns an out parameter (WasLocal) that tells the system whether it successfully determined
+    /// that the local files can be used.
+    /// </summary>
+    public class TempStorage
+    {
+        /// <summary>
+        /// Structure used to store a file entry that will be written to the temp file cache. essentially stores a relative path to the file, the last write time (UTC), and the file size in bytes.
+        /// </summary>
+        private class TempStorageFileInfo
+        {
+            public string Name;
+            public DateTime Timestamp;
+            public long Size;
 
+            /// <summary>
+            /// Compares with another temp storage manifest file. Allows certain files to compare differently because our build system requires it.
+            /// Also uses a relaxed timestamp matching to allow for filesystems with limited granularity in their timestamps.
+            /// </summary>
+            /// <param name="Other"></param>
+            /// <returns></returns>
             public bool Compare(TempStorageFileInfo Other)
             {
                 bool bOk = true;
@@ -78,10 +186,10 @@ namespace AutomationTool
                     // cf. http://msdn.microsoft.com/en-us/library/windows/desktop/ms724290%28v=vs.85%29.aspx
                     if (!((Timestamp - Other.Timestamp).TotalSeconds < 2 && (Timestamp - Other.Timestamp).TotalSeconds > -2))
                     {
-                        CommandUtils.LogWithVerbosity(LogType, "File date mismatch {0} {1} {2} {3}", Name, Timestamp.ToString(), Other.Name, Other.Timestamp.ToString());
+                        CommandUtils.LogWithVerbosity(LogType, "File date mismatch {0} {1} {2} {3}", Name, Timestamp, Other.Name, Other.Timestamp);
                         bOk = bOkToBeDifferent || bSilentOkToBeDifferent;
                     }
-                    if (!(Size == Other.Size))
+                    if (Size != Other.Size)
                     {
                         CommandUtils.LogWithVerbosity(LogType, "File size mismatch {0} {1} {2} {3}", Name, Size, Other.Name, Other.Size);
                         bOk = bOkToBeDifferent;
@@ -96,64 +204,76 @@ namespace AutomationTool
             }
         }
 
+        /// <summary>
+        /// Represents a temp storage manifest file, or a listing of all the files stored in a temp storage folder.
+        /// Essentially stores a mapping of folder names to a list of file infos that describe each file.
+        /// Can be created from list if directory to file mappings, saved to an XML file, and loaded from an XML file.
+        /// All files and directories are represented as relative to the root folder in which the manifest is saved in.
+        /// </summary>
         public class TempStorageManifest
         {
-            private static readonly string RootElementName = "tempstorage";
-            private static readonly string DirectoryElementName = "directory";
-            private static readonly string FileElementName = "file";
-            private static readonly string NameAttributeName = "name";
-            private static readonly string TimestampAttributeName = "timestamp";
-            private static readonly string SizeAttributeName = "size";
-
-            private Dictionary<string, List<TempStorageFileInfo>> Directories = new Dictionary<string, List<TempStorageFileInfo>>();
+            private const string RootElementName = "tempstorage";
+            private const string DirectoryElementName = "directory";
+            private const string FileElementName = "file";
+            private const string NameAttributeName = "name";
+            private const string TimestampAttributeName = "timestamp";
+            private const string SizeAttributeName = "size";
+
+            /// <summary>
+            /// A mapping of relative directory names to a list of file infos inside that directory, also stored as relative paths to the manifest file.
+            /// </summary>
+            private readonly Dictionary<string, List<TempStorageFileInfo>> Directories;
+
+            /// <summary>
+            /// Creates an manifest from a given directory to file list mapping. Internal helper function.
+            /// </summary>
+            /// <param name="Directories">Used to initialize the Directories member.</param>
+            private TempStorageManifest(Dictionary<string, List<TempStorageFileInfo>> Directories)
+            {
+                this.Directories = Directories;
+            }
 
-            public void Create(List<string> InFiles, string BaseFolder)
+            /// <summary>
+            /// Creates a manifest from a flat list of files (in many folders) and a BaseFolder from which they are rooted.
+            /// </summary>
+            /// <param name="InFiles">List of full file paths</param>
+            /// <param name="RootDir">Root folder for all the files. All files must be relative to this RootDir.</param>
+            /// <returns>The newly create TempStorageManifest, if all files exist. Otherwise throws an AutomationException.</returns>
+            public static TempStorageManifest Create(List<string> InFiles, string RootDir)
             {
-                BaseFolder = CommandUtils.CombinePaths(BaseFolder, "/");
-                if (!BaseFolder.EndsWith("/")&& !BaseFolder.EndsWith("\\"))
-                {
-                    throw new AutomationException("base folder {0} should end with a separator", BaseFolder);
-                }
+                var Directories = new Dictionary<string, List<TempStorageFileInfo>>();
 
-                foreach (string InFilename in InFiles)
+                foreach (string Filename in InFiles)
                 {
-                    var Filename = CommandUtils.CombinePaths(InFilename);
-                    InternalUtils.Robust_FileExists_NoExceptions(true, Filename, "Could not add {0} to manifest because it does not exist");
-
-                    FileInfo Info = new FileInfo(Filename);
-                    Filename = Info.FullName;
-                    InternalUtils.Robust_FileExists_NoExceptions(true, Filename, "Could not add {0} to manifest because it does not exist2");
-
-                    if (!Filename.StartsWith(BaseFolder, StringComparison.InvariantCultureIgnoreCase))
-                    {
-                        throw new AutomationException("Could not add {0} to manifest because it does not start with the base folder {1}", Filename, BaseFolder);
-                    }
-                    var RelativeFile = Filename.Substring(BaseFolder.Length);
+                    // use this to warm up the file shared on Mac.
+                    InternalUtils.Robust_FileExists(true, Filename, "Could not add {0} to manifest because it does not exist");
+                    // should exist now, let's get the size and timestamp
+                    // We also use this to get the OS's fullname in a consistent way to ensure we handle case sensitivity issues when comparing values below.
+                    var FileInfo = new FileInfo(Filename);
+
+                    // Strip the root dir off the file, as we only store relative paths in the manifest.
+                    // Manifest only stores path with slashes for consistency, so ensure we convert them here.
+                    var RelativeFile = CommandUtils.ConvertSeparators(PathSeparator.Slash, CommandUtils.StripBaseDirectory(FileInfo.FullName, RootDir));
+                    var RelativeDir = CommandUtils.ConvertSeparators(PathSeparator.Slash, Path.GetDirectoryName(RelativeFile));
+
+                    // add the file entry, adding a directory entry along the way if necessary.
                     List<TempStorageFileInfo> ManifestDirectory;
-                    string DirectoryName = CommandUtils.CombinePaths(Path.GetDirectoryName(Filename), "/");
-                    if (!DirectoryName.StartsWith(BaseFolder, StringComparison.InvariantCultureIgnoreCase))
-                    {
-                        throw new AutomationException("Could not add directory {0} to manifest because it does not start with the base folder {1}", DirectoryName, BaseFolder);
-                    }
-                    var RelativeDir = DirectoryName.Substring(BaseFolder.Length).Replace("\\", "/");
                     if (Directories.TryGetValue(RelativeDir, out ManifestDirectory) == false)
                     {
                         ManifestDirectory = new List<TempStorageFileInfo>();
                         Directories.Add(RelativeDir, ManifestDirectory);
                     }
-                    ManifestDirectory.Add(new TempStorageFileInfo(Filename, RelativeFile));
+                    ManifestDirectory.Add(new TempStorageFileInfo { Name = RelativeFile, Timestamp = FileInfo.LastWriteTimeUtc, Size = FileInfo.Length });
                 }
 
-                Stats("Created manifest");
-
+                return new TempStorageManifest(Directories);
             }
 
-            public void Stats(string Description)
-            {
-
-                CommandUtils.Log("{0}: Directories={1} Files={2} Size={3}", Description, Directories.Count, GetFileCount(), GetTotalSize());
-            }
-            
+            /// <summary>
+            /// Compares this temp manifest with another one. It allows certain files that we know will differ to differe, while requiring the rest of the important details remain the same.
+            /// </summary>
+            /// <param name="Other"></param>
+            /// <returns></returns>
             public bool Compare(TempStorageManifest Other)
             {
                 if (Directories.Count != Other.Directories.Count)
@@ -217,373 +337,348 @@ namespace AutomationTool
                 return true;
             }
 
-            public TempStorageFileInfo FindFile(string LocalPath, string BaseFolder)
+            /// <summary>
+            /// Loads a manifest from a file
+            /// </summary>
+            /// <param name="Filename">Full path to the manifest.</param>
+            /// <returns>The newly created manifest instance</returns>
+            public static TempStorageManifest Load(string Filename)
             {
-                BaseFolder = CommandUtils.CombinePaths(BaseFolder, "/");
-                if (!BaseFolder.EndsWith("/")&& !BaseFolder.EndsWith("\\"))
-                {
-                    throw new AutomationException("base folder {0} should end with a separator", BaseFolder);
-                }
-                var Filename = CommandUtils.CombinePaths(LocalPath);
-                if (!Filename.StartsWith(BaseFolder, StringComparison.InvariantCultureIgnoreCase))
-                {
-                    throw new AutomationException("Could not add {0} to manifest because it does not start with the base folder {1}", Filename, BaseFolder);
-                }
-                var RelativeFile = Filename.Substring(BaseFolder.Length);
-                string DirectoryName = CommandUtils.CombinePaths(Path.GetDirectoryName(Filename), "/");
-                if (!DirectoryName.StartsWith(BaseFolder, StringComparison.InvariantCultureIgnoreCase))
-                {
-                    throw new AutomationException("Could not add directory {0} to manifest because it does not start with the base folder {1}", DirectoryName, BaseFolder);
-                }
-                var RelativeDir = DirectoryName.Substring(BaseFolder.Length);               
-
-                List<TempStorageFileInfo> Directory;
-                if (Directories.TryGetValue(RelativeDir, out Directory))
+                try
                 {
-                    foreach (var SyncedFile in Directory)
+                    // Load the manifest XML file.
+                    var Directories = XDocument.Load(Filename).Root
+                        // Get the directory array of child elements
+                        .Elements(DirectoryElementName)
+                        // convert it to a dictionary of directory names to a list of TempStorageFileInfo elements for each file in the directory.
+                        .ToDictionary(
+                            DirElement => DirElement.Attribute(NameAttributeName).Value,
+                            DirElement => DirElement.Elements(FileElementName).Select(FileElement => new TempStorageFileInfo
+                                            {
+                                                Timestamp = new DateTime(long.Parse(FileElement.Attribute(TimestampAttributeName).Value)),
+                                                Size = long.Parse(FileElement.Attribute(SizeAttributeName).Value),
+                                                Name = FileElement.Value,
+                                            })
+                                            .ToList()
+                    );
+
+                    // The manifest must have at least one file with length > 0
+                    if (IsEmptyManifest(Directories))
                     {
-                        if (SyncedFile.Name.Equals(RelativeFile.Replace("\\", "/"), StringComparison.InvariantCultureIgnoreCase))
-                        {
-                            return SyncedFile;
-                        }
+                        throw new AutomationException("Attempt to load empty manifest.");
                     }
+                    
+                    return new TempStorageManifest(Directories);
                 }
-                return null;
-            }
-
-            public int GetFileCount()
-            {
-                int FileCount = 0;
-                foreach (var Directory in Directories)
+                catch (Exception Ex)
                 {
-                    FileCount += Directory.Value.Count;
+                    throw new AutomationException(Ex, "Failed to load manifest file {0}", Filename);
                 }
-                return FileCount;
             }
 
-
-            public void Load(string Filename, bool bQuiet = false)
+            /// <summary>
+            /// Returns true if the manifest represented by the dictionary is empty (no files with non-zero length)
+            /// </summary>
+            /// <param name="Directories">Dictionary of directory to file mappings that define the manifest</param>
+            /// <returns></returns>
+            private static bool IsEmptyManifest(Dictionary<string, List<TempStorageFileInfo>> Directories)
             {
-                XmlDocument File = new XmlDocument();
-                try
-                {
-                    File.Load(Filename);
-                }
-                catch (Exception Ex)
-                {
-                    throw new AutomationException("Unable to load manifest file {0}: {1}", Filename, Ex.Message);
-                }
+                return !Directories.SelectMany(Pair => Pair.Value).Any(FileInfo => FileInfo.Size > 0);
+            }
 
-                XmlNode RootNode = File.FirstChild;
-                if (RootNode != null && RootNode.Name == RootElementName)
-                {
-                    for (XmlNode ChildNode = RootNode.FirstChild; ChildNode != null; ChildNode = ChildNode.NextSibling)
-                    {
-                        if (ChildNode.Name == DirectoryElementName)
-                        {
-                            LoadDirectory(ChildNode);
-                        }
-                        else
-                        {
-                            throw new AutomationException("Unexpected node {0} while reading manifest files.", ChildNode.Name);
-                        }
-                    }
-                }
-                else
+            /// <summary>
+            /// Saves the manifest to the given filename.
+            /// </summary>
+            /// <param name="Filename">Full path to the file to save to.</param>
+            public void Save(string Filename)
+            {
+                if (IsEmptyManifest(Directories))
                 {
-                    throw new AutomationException("Bad root node ({0}).", RootNode != null ? RootNode.Name : "null");
+                    throw new AutomationException("Attempt to save empty manifest.");
                 }
 
-                if (!bQuiet)
-                {
-                    Stats(String.Format("Loaded manifest {0}", Filename));
-                }
-                if (GetFileCount() <= 0 || GetTotalSize() <= 0)
-                {
-                    throw new AutomationException("Attempt to load empty manifest.");
-                }
+                new XElement(RootElementName,
+                    from Dir in Directories
+                    select new XElement(DirectoryElementName,
+                        new XAttribute(NameAttributeName, Dir.Key),
+                        from File in Dir.Value
+                        select new XElement(FileElementName,
+                            new XAttribute(TimestampAttributeName, File.Timestamp.Ticks),
+                            new XAttribute(SizeAttributeName, File.Size),
+                            File.Name)
+                    )
+                ).Save(Filename);
+                CommandUtils.Log("Saved temp manifest {0} with {1} files and total size {2}", Filename, Directories.SelectMany(Dir=>Dir.Value).Count(), Directories.SelectMany(Dir=>Dir.Value).Sum(File=>File.Size));
             }
 
-            bool TryGetAttribute(XmlNode Node, string AttributeName, out string Value)
+            /// <summary>
+            /// Returns the sum of filesizes in the manifest in bytes.
+            /// </summary>
+            /// <returns>Returns the sum of filesizes in the manifest in bytes.</returns>
+            public long GetTotalSize()
             {
-                bool Result = false;
-                Value = "";
-                if (Node.Attributes != null && Node.Attributes.Count > 0)
-                {
-                    for (int Index = 0; Result == false && Index < Node.Attributes.Count; ++Index)
-                    {
-                        if (Node.Attributes[Index].Name == AttributeName)
-                        {
-                            Value = Node.Attributes[Index].Value;
-                            Result = true;
-                        }
-                    }
-                }
-                return Result;
+                return Directories.SelectMany(Dir => Dir.Value).Sum(FileInfo => FileInfo.Size);
             }
 
-            void LoadDirectory(XmlNode DirectoryNode)
+            /// <summary>
+            /// Gets a flat list of files in the manifest, converting to full path names rooted at the given base dir.
+            /// </summary>
+            /// <param name="RootDir">Root dir to prepend to all the files in the manifest.</param>
+            /// <param name="bFilesMustExist">if true, one file must exist else the function fails.</param>
+            /// <returns>Flat list of all files in the manifest re-rooted at RootDir.</returns>
+            public List<string> GetFiles(string RootDir, bool bFilesMustExist = true)
             {
-                string DirectoryName;
-                if (TryGetAttribute(DirectoryNode, NameAttributeName, out DirectoryName) == false)
-                {
-                    throw new AutomationException("Unable to get directory name attribute.");
-                }
-
-                List<TempStorageFileInfo> Files = new List<TempStorageFileInfo>();
-                for (XmlNode ChildNode = DirectoryNode.FirstChild; ChildNode != null; ChildNode = ChildNode.NextSibling)
-                {
-                    if (ChildNode.Name == FileElementName)
+                // flatten the list of directories, pull the files out, set their root path, and ensure the path is not too long.
+                var Result = Directories.SelectMany(Dir=>Dir.Value).Select(FileInfo =>
                     {
-                        TempStorageFileInfo Info = LoadFile(ChildNode);
-                        if (Info != null)
+                        var NewFilePath = CommandUtils.CombinePaths(RootDir, FileInfo.Name);
+                        // create a FileInfo using the file, which will help us catch path too long exceptions early.
+                        try
                         {
-                            Files.Add(Info);
+                            return new FileInfo(NewFilePath).FullName;
                         }
-                        else
+                        catch (PathTooLongException Ex)
                         {
-                            throw new AutomationException("Error reading file entry in the manifest file ({0})", ChildNode.InnerXml);
+                            throw new AutomationException(Ex, "Path too long ... failed to create FileInfo for {0}", NewFilePath);
                         }
-                    }
-                    else
-                    {
-                        throw new AutomationException("Unexpected node {0} while reading file nodes.", ChildNode.Name);
-                    }
-                }
-                Directories.Add(DirectoryName, Files);
-            }
-
-            TempStorageFileInfo LoadFile(XmlNode FileNode)
-            {
-                TempStorageFileInfo Info = new TempStorageFileInfo();
-
-                long Timestamp;
-                string TimestampAttribute;
-                if (TryGetAttribute(FileNode, TimestampAttributeName, out TimestampAttribute) == false)
-                {
-                   throw new AutomationException("Unable to get timestamp attribute.");
-                }
-                if (Int64.TryParse(TimestampAttribute, out Timestamp) == false)
-                {
-                    throw new AutomationException("Unable to parse timestamp attribute ({0}).", TimestampAttribute);
-                }
-                Info.Timestamp = new DateTime(Timestamp);
+                    }).ToList();
 
-                string SizeAttribute;
-                if (TryGetAttribute(FileNode, SizeAttributeName, out SizeAttribute) == false)
+                if (bFilesMustExist)
                 {
-                    throw new AutomationException("Unable to get size attribute.");
+                    CommandUtils.Log("Checking that first file in manifest exists...");
+                    InternalUtils.Robust_FileExists(false, Result[0], "GetFiles in temp manifest failed because the first Rebased manifest file does not exist: {0}");
                 }
-                if (Int64.TryParse(SizeAttribute, out Info.Size) == false)
-                {
-                    throw new AutomationException("Unable to parse size attribute ({0}).", SizeAttribute);
-                }
-
-                Info.Name = FileNode.InnerText;
-
-                return Info;
+                return Result;
             }
+        }
 
-            public void Save(string Filename)
-            {
-                if (GetFileCount() <= 0 || GetTotalSize() <= 0)
-                {
-                    throw new AutomationException("Attempt to save empty manifest.");
-                }
-                XmlDocument File = new XmlDocument();
-
-                XmlElement RootElement = File.CreateElement(RootElementName);
-                File.AppendChild(RootElement);
-
-                SaveDirectories(File, RootElement);
+        /// <summary>
+        /// returns "[LocalRoot]\Engine\Saved\TmpStore" - the location where temp storage manifests will be stored before copying them to their final temp storage location.
+        /// </summary>
+        /// <returns>returns "[LocalRoot]\Engine\Saved\TmpStore"</returns>
+        private static string LocalTempStorageManifestDirectory()
+        {
+            return CommandUtils.CombinePaths(CommandUtils.CmdEnv.LocalRoot, "Engine", "Saved", TempStorageSubdirectoryName);
+        }
 
-                File.Save(Filename);
-            }
+        /// <summary>
+        /// Legayc code to clean all temp storage data from the given directory that is older than a certain threshold (defined directly in the code).
+        /// </summary>
+        /// <param name="TopDirectory">Fully qualified path of the folder that is the root of a bunch of temp storage folders. Should be a string returned by <see cref="ResolveSharedTempStorageDirectory"/> </param>
+        private static void CleanSharedTempStorageLegacy(string TopDirectory)
+        {
+            // This is a hack to clean out the old temp storage in the old folder name.
+            TopDirectory = TopDirectory.Replace(Path.DirectorySeparatorChar + TempStorageSubdirectoryName + Path.DirectorySeparatorChar, Path.DirectorySeparatorChar + "GUBP" + Path.DirectorySeparatorChar);
 
-            void SaveDirectories(XmlDocument File, XmlElement Root)
+            using (var TelemetryStopwatch = new CommandUtils.TelemetryStopwatch("CleanSharedTempStorageLegacy"))
             {
-                foreach (KeyValuePair<string, List<TempStorageFileInfo>> Directory in Directories)
+                const int MaximumDaysToKeepLegacyTempStorage = 1;
+                var StartTimeDir = DateTime.UtcNow;
+                DirectoryInfo DirInfo = new DirectoryInfo(TopDirectory);
+                var TopLevelDirs = DirInfo.GetDirectories();
                 {
-                    XmlElement DirectoryElement = File.CreateElement(DirectoryElementName);
-                    Root.AppendChild(DirectoryElement);
-                    DirectoryElement.SetAttribute(NameAttributeName, Directory.Key);
-                    SaveFiles(File, DirectoryElement, Directory.Value);
+                    var BuildDuration = (DateTime.UtcNow - StartTimeDir).TotalMilliseconds;
+                    CommandUtils.Log("Took {0}s to enumerate {1} directories.", BuildDuration / 1000, TopLevelDirs.Length);
                 }
-            }
-
-            void SaveFiles(XmlDocument File, XmlElement Directory, List<TempStorageFileInfo> Files)
-            {
-                foreach (TempStorageFileInfo Info in Files)
+                foreach (var TopLevelDir in TopLevelDirs)
                 {
-                    XmlElement FileElement = File.CreateElement(FileElementName);
-                    FileElement.SetAttribute(TimestampAttributeName, Info.Timestamp.Ticks.ToString());
-                    FileElement.SetAttribute(SizeAttributeName, Info.Size.ToString());
-                    FileElement.InnerText = Info.Name;
-                    Directory.AppendChild(FileElement);
-                }
-            }
+                    if (CommandUtils.DirectoryExists_NoExceptions(TopLevelDir.FullName))
+                    {
+                        bool bOld = false;
+                        foreach (var ThisFile in CommandUtils.FindFiles_NoExceptions(true, "*.TempManifest", false, TopLevelDir.FullName))
+                        {
+                            FileInfo Info = new FileInfo(ThisFile);
 
-            public Dictionary<string, List<TempStorageFileInfo>> DirectoryList
-            {
-                get { return Directories; }
+                            if ((DateTime.UtcNow - Info.LastWriteTimeUtc).TotalDays > MaximumDaysToKeepLegacyTempStorage)
+                            {
+                                bOld = true;
+                            }
+                        }
+                        if (bOld)
+                        {
+                            CommandUtils.Log("Deleting temp storage directory {0}, because it is more than {1} days old.", TopLevelDir.FullName, MaximumDaysToKeepLegacyTempStorage);
+                            var StartTime = DateTime.UtcNow;
+                            try
+                            {
+                                if (Directory.Exists(TopLevelDir.FullName))
+                                {
+                                    // try the direct approach first
+                                    Directory.Delete(TopLevelDir.FullName, true);
+                                }
+                            }
+                            catch
+                            {
+                            }
+                            CommandUtils.DeleteDirectory_NoExceptions(true, TopLevelDir.FullName);
+                            var BuildDuration = (DateTime.UtcNow - StartTime).TotalMilliseconds;
+                            CommandUtils.Log("Took {0}s to delete {1}.", BuildDuration / 1000, TopLevelDir.FullName);
+                        }
+                    }
+                }
             }
+        }
 
-            public long GetTotalSize()
-            {
-
-                long TotalSize = 0;
+        /// <summary>
+        /// Cleans all temp storage data from the given directory that is older than a certain threshold (defined directly in the code).
+        /// </summary>
+        /// <param name="TopDirectory">Fully qualified path of the folder that is the root of a bunch of temp storage folders. Should be a string returned by <see cref="ResolveSharedTempStorageDirectory"/> </param>
+        private static void CleanSharedTempStorage(string TopDirectory)
+        {
+            CleanSharedTempStorageLegacy(TopDirectory);
 
-                foreach (var Dir in Directories)
+            Action<string> TryToDeletePossiblyEmptyFolder = (string FolderName) =>
                 {
-                    foreach (var FileInfo in Dir.Value)
+                    try
                     {
-                        TotalSize += FileInfo.Size;
+                        Directory.Delete(FolderName);
                     }
-                }
-
-                return TotalSize;
-            }
-            public List<string> GetFiles(string NewBaseDir)
-            {
-                var Result = new List<string>();
-                foreach (var Dir in Directories)
-                {
-                    foreach (var ThisFileInfo in Dir.Value)
+                    catch (IOException)
                     {
-                        var NewFile = CommandUtils.CombinePaths(NewBaseDir, ThisFileInfo.Name);
-
-						try
-						{
-							FileInfo Info = new FileInfo(NewFile);
-							Result.Add(Info.FullName);
-						}
-						catch (PathTooLongException)
-						{
-							CommandUtils.LogError("Path too long ... failed to create FileInfo for {0}", NewFile);
-							throw;
-						}
+                        // only catch "directory is not empty type exceptions, if possible. Best we can do is check for IOException.
                     }
-                }
-                if (Result.Count < 1)
-                {
-                    throw new AutomationException("No files in attempt to GetFiles().");
-                }
-
-				CommandUtils.Log("Checking that first file in manifest exists...");
-                InternalUtils.Robust_FileExists_NoExceptions(false, Result[0], "Rebased manifest file does not exist {0}");
-
-                return Result;
-            }
-        }
-
-        public static string LocalTempStorageManifestDirectory(CommandEnvironment Env)
-        {
-            return CommandUtils.CombinePaths(Env.LocalRoot, "Engine", "Saved", "GUBP");
-        }
-        static HashSet<string> TopDirectoryTestedForClean = new HashSet<string>();
-        static void CleanSharedTempStorage(string TopDirectory)
-        {
-            if (TopDirectoryTestedForClean.Contains(TopDirectory))
-            {
-                return;
-            }
-            TopDirectoryTestedForClean.Add(TopDirectory);
+                    catch (Exception Ex)
+                    {
+                        CommandUtils.LogWarning("Unexpected failure trying to delete possibly empty temp folder {0}: {1}", FolderName, Ex);
+                    }
+                };
 
-            const int MaximumDaysToKeepTempStorage = 1;
-            var StartTimeDir = DateTime.UtcNow;
-            DirectoryInfo DirInfo = new DirectoryInfo(TopDirectory);
-            var TopLevelDirs = DirInfo.GetDirectories();
+            CommandUtils.LogConsole("Cleaning temp storage for {0}...", TopDirectory);
+            int FoldersCleaned = 0;
+            int FoldersManuallyCleaned = 0;
+            int FoldersFailedCleaned = 0;
+            using (var TelemetryStopwatch = new CommandUtils.TelemetryStopwatch("CleanSharedTempStorage"))
             {
-                var BuildDuration = (DateTime.UtcNow - StartTimeDir).TotalMilliseconds;
-                CommandUtils.Log("Took {0}s to enumerate {1} directories.", BuildDuration / 1000, TopLevelDirs.Length);
-            }
-            foreach (var TopLevelDir in TopLevelDirs)
-            {
-                if (CommandUtils.DirectoryExists_NoExceptions(TopLevelDir.FullName))
+                const double MaximumDaysToKeepTempStorage = 3;
+                var Now = DateTime.UtcNow;
+                // This will search legacy folders as well, but those should go away within a few days.
+                foreach (var OldManifestInfo in
+                    // First subdirectory is the branch name
+                    from BranchDir in Directory.EnumerateDirectories(TopDirectory)
+                    // Second subdirectory is the changelist
+                    from CLDir in Directory.EnumerateDirectories(BranchDir)
+                    // third subdirectory is the node storage name
+                    from NodeNameDir in Directory.EnumerateDirectories(CLDir)
+                    // now look for manifest files.
+                    let ManifestFile = Directory.EnumerateFiles(NodeNameDir, "*.TempManifest").SingleOrDefault()
+                    where ManifestFile != null
+                    // only choose ones that are old enough.
+                    where (Now - new FileInfo(ManifestFile).LastWriteTimeUtc).TotalDays > MaximumDaysToKeepTempStorage
+                    select new { BranchDir, CLDir, NodeNameDir })
                 {
-                    bool bOld = false;
-                    foreach (var ThisFile in CommandUtils.FindFiles_NoExceptions(true, "*.TempManifest", false, TopLevelDir.FullName))
+                    try
                     {
-                        FileInfo Info = new FileInfo(ThisFile);
-
-                        if ((DateTime.UtcNow - Info.LastWriteTimeUtc).TotalDays > MaximumDaysToKeepTempStorage)
-                        {
-                            bOld = true;
-                        }
+                        CommandUtils.LogConsole("Deleting folder with old temp storage {0}...", OldManifestInfo.NodeNameDir);
+                        Directory.Delete(OldManifestInfo.NodeNameDir, true);
+                        FoldersCleaned++;
                     }
-                    if (bOld)
+                    catch (Exception Ex)
                     {
-                        CommandUtils.Log("Deleteing temp storage directory {0}, because it is more than {1} days old.", TopLevelDir.FullName, MaximumDaysToKeepTempStorage);
-                        var StartTime = DateTime.UtcNow;
-                        try
+                        CommandUtils.LogWarning("Failed to delete old manifest folder '{0}', will try one file at a time: {1}", OldManifestInfo.NodeNameDir, Ex);
+                        if (CommandUtils.DeleteDirectory_NoExceptions(true, OldManifestInfo.NodeNameDir))
                         {
-                            if (Directory.Exists(TopLevelDir.FullName))
-                            {
-                                // try the direct approach first
-                                Directory.Delete(TopLevelDir.FullName, true);
-                            }
+                            FoldersManuallyCleaned++;
                         }
-                        catch (Exception)
+                        else
                         {
+                            FoldersFailedCleaned++;
                         }
-                        CommandUtils.DeleteDirectory_NoExceptions(true, TopLevelDir.FullName);
-                        var BuildDuration = (DateTime.UtcNow - StartTime).TotalMilliseconds;
-                        CommandUtils.Log("Took {0}s to delete {1}.", BuildDuration / 1000, TopLevelDir.FullName);
                     }
+                    // Once we are done, try to delete the CLDir and BranchDir (will fail if not empty, so no worries).
+                    TryToDeletePossiblyEmptyFolder(OldManifestInfo.CLDir);
+                    TryToDeletePossiblyEmptyFolder(OldManifestInfo.BranchDir);
                 }
+                TelemetryStopwatch.Finish(string.Format("CleanSharedTempStorage.{0}.{1}.{2}", FoldersCleaned, FoldersManuallyCleaned, FoldersFailedCleaned));
             }
         }
-        
-        public static string SharedTempStorageDirectory()
-        {
-            return "GUBP";
-        }
-        public static string UE4TempStorageDirectory()
+
+        /// <summary>
+        /// Cleans the shared temp storage folders for all the games listed. This code will ensure that any duplicate games
+        /// or games that resolve to use the same temp folder are not cleaned twice.
+        /// </summary>
+        /// <param name="GameNames">List of game names to clean temp folders for.</param>
+        public static void CleanSharedTempStorageDirectory(IEnumerable<string> GameNames)
         {
-            string SharedSubdir = SharedTempStorageDirectory();
-            string Root = CommandUtils.RootSharedTempStorageDirectory();
-            return CommandUtils.CombinePaths(Root, "UE4", SharedSubdir);
+            if (!CommandUtils.IsBuildMachine || UnrealBuildTool.Utils.IsRunningOnMono)  // saw a hang on this, anyway it isn't necessary to clean with macs, they are slow anyway
+            {
+                return;
+            }
+            // Generate a unique set of folder names to clean so we don't clean folders twice.
+            foreach (var TempStorageFolder in new HashSet<string>(GameNames.Select(ResolveSharedTempStorageDirectory)))
+            {
+                try
+                {
+                    CleanSharedTempStorage(TempStorageFolder);
+                }
+                catch (Exception Ex)
+                {
+                    CommandUtils.LogWarning("Unable to Clean Temp Directory {0}. Exception: {1}", TempStorageFolder, Ex);
+                }
+            }
         }
-        public static bool HaveSharedTempStorage(bool ForSaving)
+
+        /// "TmpStore", the folder name where all GUBP temp storage files are added inside a game's network storage location.
+        private const string TempStorageSubdirectoryName = "TmpStore";
+
+        /// <summary>
+        /// Determines if the root UE4 temp storage directory exists.
+        /// </summary>
+        /// <param name="ForSaving">If true, confirm that the directory exists and is writeable.</param>
+        /// <returns>true if the UE4 temp storage directory exists (and is writeable if requested).</returns>
+        public static bool IsSharedTempStorageAvailable(bool ForSaving)
         {
-            var Dir = UE4TempStorageDirectory();
+            var UE4TempStorageDir = CommandUtils.CombinePaths(CommandUtils.RootBuildStorageDirectory(), "UE4", TempStorageSubdirectoryName);
             if (ForSaving)
             {
                 int Retries = 0;
                 while (Retries < 24)
                 {
-                    if (InternalUtils.Robust_DirectoryExistsAndIsWritable_NoExceptions(Dir))
+                    if (InternalUtils.Robust_DirectoryExistsAndIsWritable_NoExceptions(UE4TempStorageDir))
                     {
                         return true;
                     }
-                    CommandUtils.FindDirectories_NoExceptions(false, "*", false, Dir); // there is some internet evidence to suggest this might perk up the mac share
+                    CommandUtils.FindDirectories_NoExceptions(false, "*", false, UE4TempStorageDir); // there is some internet evidence to suggest this might perk up the mac share
                     System.Threading.Thread.Sleep(5000);
-					Retries++;
+                    Retries++;
                 }
             }
-            else if (InternalUtils.Robust_DirectoryExists_NoExceptions(Dir, "Could not find {0}"))
+            else if (InternalUtils.Robust_DirectoryExists_NoExceptions(UE4TempStorageDir, "Could not find {0}"))
             {
                 return true;
             }
             return false;
         }
 
-        static Dictionary<string, string> ResolveCache = new Dictionary<string, string>();
-        public static string ResolveSharedBuildDirectory(string GameFolder)
+
+        /// <summary>
+        /// Cache of previously resolved build directories for games. Since this requires a few filesystem checks,
+        /// we cache the results so future lookups are fast.
+        /// </summary>
+        static readonly Dictionary<string, string> SharedBuildDirectoryResolveCache = new Dictionary<string, string>();
+
+        /// <summary>
+        /// Determines full path of the shared build folder for the given game.
+        /// If the shared build folder does not already exist, will try to use UE4's folder instead.
+        /// This is because the root build folders are independently owned by each game team and managed by IT, so cannot be created on the fly.
+        /// </summary>
+        /// <param name="GameName">GameName to look for the build folder for. If empty or the GameName folder cannot be found, uses the UE4 folder (GameName = UE4)</param>
+        /// <returns>The full path of the shared build directory name, or throws an exception if none is found and the UE4 folder is not found either.</returns>
+        public static string ResolveSharedBuildDirectory(string GameName)
         {
-            if (ResolveCache.ContainsKey(GameFolder))
+            if (GameName == null) throw new ArgumentNullException("GameName");
+
+            if (SharedBuildDirectoryResolveCache.ContainsKey(GameName))
             {
-                return ResolveCache[GameFolder];
+                return SharedBuildDirectoryResolveCache[GameName];
             }
-            string Root = CommandUtils.RootSharedTempStorageDirectory();
-            string Result = CommandUtils.CombinePaths(Root, GameFolder);
-            if (String.IsNullOrEmpty(GameFolder) || String.Compare(GameFolder, "ShooterGame", true) == 0 || !InternalUtils.Robust_DirectoryExistsAndIsWritable_NoExceptions(Result))
+            string Root = CommandUtils.RootBuildStorageDirectory();
+            string Result = CommandUtils.CombinePaths(Root, GameName);
+            if (GameName == "" || string.Equals(GameName, "ShooterGame", StringComparison.InvariantCultureIgnoreCase) || !InternalUtils.Robust_DirectoryExistsAndIsWritable_NoExceptions(Result))
             {
                 string GameStr = "Game";
                 bool HadGame = false;
-                if (GameFolder.EndsWith(GameStr, StringComparison.InvariantCultureIgnoreCase))
+                if (GameName.EndsWith(GameStr, StringComparison.InvariantCultureIgnoreCase))
                 {
-                    string ShortFolder = GameFolder.Substring(0, GameFolder.Length - GameStr.Length);
+                    string ShortFolder = GameName.Substring(0, GameName.Length - GameStr.Length);
                     Result = CommandUtils.CombinePaths(Root, ShortFolder);
                     HadGame = true;
                 }
@@ -596,14 +691,20 @@ namespace AutomationTool
                     }
                 }
             }
-            ResolveCache.Add(GameFolder, Result);
+            SharedBuildDirectoryResolveCache.Add(GameName, Result);
             return Result;
         }
 
-        public static string ResolveSharedTempStorageDirectory(string GameFolder)
+        /// <summary>
+        /// Returns the full path of the temp storage directory for the given game name (ie, P:\Builds\GameName\TmpStore).
+        /// If the game's build folder does not exist, will use the temp storage folder in the UE4 build folder (ie, P:\Builds\UE4\TmpStore).
+        /// If the directory does not exist, we will try to create it.
+        /// </summary>
+        /// <param name="GameName">game name to determine the temp storage folder for. Empty is equivalent to "UE4".</param>
+        /// <returns>Essentially adds the TmpStore/ subdirectory to the SharedBuildDirectory for the game.</returns>
+        public static string ResolveSharedTempStorageDirectory(string GameName)
         {
-            string SharedSubdir = SharedTempStorageDirectory();
-            string Result = CommandUtils.CombinePaths(ResolveSharedBuildDirectory(GameFolder), SharedSubdir);
+            string Result = CommandUtils.CombinePaths(ResolveSharedBuildDirectory(GameName), TempStorageSubdirectoryName);
 
             if (!InternalUtils.Robust_DirectoryExists_NoExceptions(Result, "Could not find {0}"))
             {
@@ -611,153 +712,205 @@ namespace AutomationTool
             }
             if (!InternalUtils.Robust_DirectoryExists_NoExceptions(Result, "Could not find {0}"))
             {
-                throw new AutomationException("Could not create an appropriate shared temp folder {0}", Result);
+                throw new AutomationException("Could not create an appropriate shared temp folder {0} for game {1}", Result, GameName);
             }
             return Result;
         }
 
-        static HashSet<string> TestedForClean = new HashSet<string>();
-        public static void CleanSharedTempStorageDirectory(string GameFolder)
-        {
-
-            if (!CommandUtils.IsBuildMachine || TestedForClean.Contains(GameFolder) ||
-                UnrealBuildTool.Utils.IsRunningOnMono)  // saw a hang on this, anyway it isn't necessary to clean with macs, they are slow anyway
-            {
-                return;
-            }
-            TestedForClean.Add(GameFolder);
-            try
-            {
-                CleanSharedTempStorage(ResolveSharedTempStorageDirectory(GameFolder));
-            }
-            catch (Exception Ex)
-            {
-                CommandUtils.LogWarning("Unable to Clean Directory with GameName {0}", GameFolder);
-                CommandUtils.LogWarning(" Exception was {0}", LogUtils.FormatException(Ex));
-            }
-        }
-        public static string SharedTempStorageDirectory(string StorageBlock, string GameFolder = "")
+        /// <summary>
+        /// Returns the full path to the temp storage directory for the given temp storage node and game (ie, P:\Builds\GameName\TmpStore\NodeInfoDirectory)
+        /// </summary>
+        /// <param name="TempStorageNodeInfo">Node info descibing the block of temp storage (essentially used to identify a subdirectory insides the game's temp storage folder).</param>
+        /// <param name="GameName">game name to determine the temp storage folder for. Empty is equivalent to "UE4".</param>
+        /// <returns>The full path to the temp storage directory for the given storage block name for the given game.</returns>
+        private static string SharedTempStorageDirectory(TempStorageNodeInfo TempStorageNodeInfo, string GameName)
         {
-            return CommandUtils.CombinePaths(ResolveSharedTempStorageDirectory(GameFolder), StorageBlock);
+            return CommandUtils.CombinePaths(ResolveSharedTempStorageDirectory(GameName), TempStorageNodeInfo.GetRelativeDirectory());
         }
 
-        public static TempStorageManifest SaveTempStorageManifest(string RootDir, string FinalFilename, List<string> Files)
+        /// <summary>
+        /// Private helper that saves a temp file manifest described by a list of fully qualified files rooted at RootDir to the given fully qualified path, returning the created TempStorageManifest instance.
+        /// </summary>
+        /// <param name="RootDir">Folder that all the given files are rooted from.</param>
+        /// <param name="FinalFilename">Fully qualified path to of the file that should be written.</param>
+        /// <param name="Files">Fully qualified names of files to reference in the manifest file.</param>
+        /// <returns>The created manifest instance (which has already been saved to disk).</returns>
+        private static TempStorageManifest SaveTempStorageManifest(string RootDir, string FinalFilename, List<string> Files)
         {
-            var Saver = new TempStorageManifest();
-            Saver.Create(Files, RootDir);
-            if (Saver.GetFileCount() != Files.Count)
-            {
-                throw new AutomationException("Saver manifest differs has wrong number of files {0} != {1}", Saver.GetFileCount(), Files.Count);
-            }
-            var TempFilename = FinalFilename + ".temp";
-            if (CommandUtils.FileExists_NoExceptions(true, TempFilename))
-            {
-                throw new AutomationException("Temp manifest file already exists {0}", TempFilename);
-            }
+            var Saver = TempStorageManifest.Create(Files, RootDir);
             CommandUtils.CreateDirectory(true, Path.GetDirectoryName(FinalFilename));
-            Saver.Save(TempFilename);
-
-            var Tester = new TempStorageManifest();
-            Tester.Load(TempFilename, true);
-
-            if (!Saver.Compare(Tester))
-            {
-                throw new AutomationException("Temp manifest differs {0}", TempFilename);
-            }
-
-            CommandUtils.RenameFile(TempFilename, FinalFilename, true);
-            if (CommandUtils.FileExists_NoExceptions(true, TempFilename))
-            {
-                throw new AutomationException("Temp manifest didn't go away {0}", TempFilename);
-            }
-            var FinalTester = new TempStorageManifest();
-            FinalTester.Load(FinalFilename, true);
-
-            if (!Saver.Compare(FinalTester))
-            {
-                throw new AutomationException("Final manifest differs {0}", TempFilename);
-            }
-            CommandUtils.Log("Saved {0} with {1} files and total size {2}", FinalFilename, Saver.GetFileCount(), Saver.GetTotalSize());
+            Saver.Save(FinalFilename);
             return Saver;
         }
-        public static string LocalTempStorageManifestFilename(CommandEnvironment Env, string StorageBlockName)
+
+        /// <summary>
+        /// Gets the name of the local temp storage manifest file for the given temp storage node (ie, Engine\Saved\TmpStore\NodeInfoDirectory\NodeInfoFilename.TempManifest)
+        /// </summary>
+        /// <param name="TempStorageNodeInfo">Node info descibing the block of temp storage (essentially used to identify a subdirectory insides the game's temp storage folder).</param>
+        /// <returns>The name of the temp storage manifest file for the given storage block name for the given game.</returns>
+        private static string LocalTempStorageManifestFilename(TempStorageNodeInfo TempStorageNodeInfo)
         {
-            return CommandUtils.CombinePaths(LocalTempStorageManifestDirectory(Env), StorageBlockName + ".TempManifest");
+            return CommandUtils.CombinePaths(LocalTempStorageManifestDirectory(), TempStorageNodeInfo.GetManifestFilename());
         }
-        public static TempStorageManifest SaveLocalTempStorageManifest(CommandEnvironment Env, string BaseFolder, string StorageBlockName, List<string> Files)
+
+        /// <summary>
+        /// Saves the list of fully qualified files rooted at RootDir to a local temp storage manifest with the given temp storage node.
+        /// </summary>
+        /// <param name="RootDir">Folder that all the given files are rooted from.</param>
+        /// <param name="TempStorageNodeInfo">Node info descibing the block of temp storage (essentially used to identify a subdirectory insides the game's temp storage folder).</param>
+        /// <param name="Files">Fully qualified names of files to reference in the manifest file.</param>
+        /// <returns>The created manifest instance (which has already been saved to disk).</returns>
+        private static TempStorageManifest SaveLocalTempStorageManifest(string RootDir, TempStorageNodeInfo TempStorageNodeInfo, List<string> Files)
         {
-            return SaveTempStorageManifest(BaseFolder, LocalTempStorageManifestFilename(Env, StorageBlockName), Files);
+            return SaveTempStorageManifest(RootDir, LocalTempStorageManifestFilename(TempStorageNodeInfo), Files);
         }
-        public static string SharedTempStorageManifestFilename(CommandEnvironment Env, string StorageBlockName, string GameFolder)
+
+        /// <summary>
+        /// Gets the name of the temp storage manifest file for the given temp storage node and game (ie, P:\Builds\GameName\TmpStore\NodeInfoDirectory\NodeInfoFilename.TempManifest)
+        /// </summary>
+        /// <param name="TempStorageNodeInfo">Node info descibing the block of temp storage (essentially used to identify a subdirectory insides the game's temp storage folder).</param>
+        /// <param name="GameName">game name to determine the temp storage folder for. Empty is equivalent to "UE4".</param>
+        /// <returns>The name of the temp storage manifest file for the given storage block name for the given game.</returns>
+        private static string SharedTempStorageManifestFilename(TempStorageNodeInfo TempStorageNodeInfo, string GameName)
         {
-            return CommandUtils.CombinePaths(SharedTempStorageDirectory(StorageBlockName, GameFolder), StorageBlockName + ".TempManifest");
+            return CommandUtils.CombinePaths(SharedTempStorageDirectory(TempStorageNodeInfo, GameName), TempStorageNodeInfo.GetManifestFilename());
         }
-        public static TempStorageManifest SaveSharedTempStorageManifest(CommandEnvironment Env, string StorageBlockName, string GameFolder, List<string> Files)
+
+        /// <summary>
+        /// Saves the list of fully qualified files (that should be rooted at the shared temp storage location for the game) to a shared temp storage manifest with the given temp storage node and game.
+        /// </summary>
+        /// <param name="TempStorageNodeInfo">Node info descibing the block of temp storage (essentially used to identify a subdirectory insides the game's temp storage folder).</param>
+        /// <param name="GameName">game name to determine the temp storage folder for. Empty is equivalent to "UE4".</param>
+        /// <param name="Files">Fully qualified names of files to reference in the manifest file.</param>
+        /// <returns>The created manifest instance (which has already been saved to disk).</returns>
+        private static TempStorageManifest SaveSharedTempStorageManifest(TempStorageNodeInfo TempStorageNodeInfo, string GameName, List<string> Files)
         {
-            return SaveTempStorageManifest(SharedTempStorageDirectory(StorageBlockName, GameFolder), SharedTempStorageManifestFilename(Env, StorageBlockName, GameFolder), Files);
+            return SaveTempStorageManifest(SharedTempStorageDirectory(TempStorageNodeInfo, GameName), SharedTempStorageManifestFilename(TempStorageNodeInfo, GameName), Files);
         }
-        public static void DeleteLocalTempStorageManifests(CommandEnvironment Env)
+
+        /// <summary>
+        /// Deletes all temp storage manifests from the local storage location (Engine\Saved\TmpStore).
+        /// Local temp storage logic only stores manifests, so this effectively deletes any local temp storage work, while not actually deleting the local files, which the temp 
+        /// storage system doesn't really own.
+        /// </summary>
+        public static void DeleteLocalTempStorageManifests()
         {
-            CommandUtils.DeleteDirectory(true, LocalTempStorageManifestDirectory(Env));
+            CommandUtils.DeleteDirectory(true, LocalTempStorageManifestDirectory());
         }
-        public static void DeleteSharedTempStorageManifests(CommandEnvironment Env, string StorageBlockName, string GameFolder = "")
+
+        /// <summary>
+        /// Deletes all shared temp storage (file and manifest) for the given temp storage node and game.
+        /// </summary>
+        /// <param name="TempStorageNodeInfo">Node info descibing the block of temp storage (essentially used to identify a subdirectory insides the game's temp storage folder).</param>
+        /// <param name="GameName">game name to determine the temp storage folder for. Empty is equivalent to "UE4".</param>
+        private static void DeleteSharedTempStorage(TempStorageNodeInfo TempStorageNodeInfo, string GameName)
         {
-            CommandUtils.DeleteDirectory(true, SharedTempStorageDirectory(StorageBlockName, GameFolder));
+            CommandUtils.DeleteDirectory(true, SharedTempStorageDirectory(TempStorageNodeInfo, GameName));
         }
-        public static bool LocalTempStorageExists(CommandEnvironment Env, string StorageBlockName, bool bQuiet = false)
+
+        /// <summary>
+        /// Checks if a local temp storage manifest exists for the given temp storage node.
+        /// </summary>
+        /// <param name="TempStorageNodeInfo">Node info descibing the block of temp storage (essentially used to identify a subdirectory insides the game's temp storage folder).</param>
+        /// <param name="bQuiet">True to suppress logging during the operation.</param>
+        /// <returns>true if the local temp storage manifest file exists</returns>
+        public static bool LocalTempStorageManifestExists(TempStorageNodeInfo TempStorageNodeInfo, bool bQuiet = false)
         {
-            var LocalManifest = LocalTempStorageManifestFilename(Env, StorageBlockName);
-            if (CommandUtils.FileExists_NoExceptions(bQuiet, LocalManifest))
-            {
-                return true;
-            }
-            return false;
+            var LocalManifest = LocalTempStorageManifestFilename(TempStorageNodeInfo);
+            return CommandUtils.FileExists_NoExceptions(bQuiet, LocalManifest);
         }
-        public static void DeleteLocalTempStorage(CommandEnvironment Env, string StorageBlockName, bool bQuiet = false)
+
+        /// <summary>
+        /// Deletes a local temp storage manifest file for the given temp storage node if it exists.
+        /// </summary>
+        /// <param name="TempStorageNodeInfo">Node info descibing the block of temp storage (essentially used to identify a subdirectory insides the game's temp storage folder).</param>
+        /// <param name="bQuiet">True to suppress logging during the operation.</param>
+        public static void DeleteLocalTempStorageManifest(TempStorageNodeInfo TempStorageNodeInfo, bool bQuiet = false)
         {
-            var LocalManifest = LocalTempStorageManifestFilename(Env, StorageBlockName);
+            var LocalManifest = LocalTempStorageManifestFilename(TempStorageNodeInfo);
             CommandUtils.DeleteFile(bQuiet, LocalManifest);
         }
-        public static bool SharedTempStorageExists(CommandEnvironment Env, string StorageBlockName, string GameFolder = "", bool bQuiet = false)
+
+        /// <summary>
+        /// Checks if a shared temp storage manifest exists for the given temp storage node and game.
+        /// </summary>
+        /// <param name="TempStorageNodeInfo">Node info descibing the block of temp storage (essentially used to identify a subdirectory insides the game's temp storage folder).</param>
+        /// <param name="GameName">game name to determine the temp storage folder for. Empty is equivalent to "UE4".</param>
+        /// <param name="bQuiet">True to suppress logging during the operation.</param>
+        /// <returns>true if the shared temp storage manifest file exists</returns>
+        private static bool SharedTempStorageManifestExists(TempStorageNodeInfo TempStorageNodeInfo, string GameName, bool bQuiet)
         {
-            var SharedManifest = SharedTempStorageManifestFilename(Env, StorageBlockName, GameFolder);
+            var SharedManifest = SharedTempStorageManifestFilename(TempStorageNodeInfo, GameName);
             if (CommandUtils.FileExists_NoExceptions(bQuiet, SharedManifest))
             {
                 return true;
             }
             return false;
         }
-        public static bool TempStorageExists(CommandEnvironment Env, string StorageBlockName, string GameFolder = "", bool bLocalOnly = false, bool bQuiet = false)
+
+        /// <summary>
+        /// Checks if a temp storage manifest exists, either locally or in the shared folder (depending on the value of bLocalOnly) for the given temp storage node and game.
+        /// </summary>
+        /// <param name="TempStorageNodeInfo">Node info descibing the block of temp storage (essentially used to identify a subdirectory insides the game's temp storage folder).</param>
+        /// <param name="GameName">game name to determine the temp storage folder for. Empty is equivalent to "UE4".</param>
+        /// <param name="bLocalOnly">If true, only ensures that the local temp storage manifest exists. Otherwise, checks if the manifest exists in either local or shared storage.</param>
+        /// <param name="bQuiet">True to suppress logging during the operation.</param>
+        /// <returns>true if the shared temp storage manifest file exists</returns>
+        public static bool TempStorageExists(TempStorageNodeInfo TempStorageNodeInfo, string GameName, bool bLocalOnly, bool bQuiet)
         {
-            return LocalTempStorageExists(Env, StorageBlockName, bQuiet) || (!bLocalOnly && SharedTempStorageExists(Env, StorageBlockName, GameFolder, bQuiet));
+            return LocalTempStorageManifestExists(TempStorageNodeInfo, bQuiet) || (!bLocalOnly && SharedTempStorageManifestExists(TempStorageNodeInfo, GameName, bQuiet));
         }
 
-        static int ThreadsToCopyWith()
+        /// <summary>
+        /// Controls how many threads to use when performing parallel copies.
+        /// </summary>
+        /// <returns>The number of threads to use.</returns>
+        private static int ThreadsToCopyWith()
         {
             return 8;
         }
 
-        public static void StoreToTempStorage(CommandEnvironment Env, string StorageBlockName, List<string> Files, bool bLocalOnly = false, string GameFolder = "", string BaseFolder = "")
+        /// <summary>
+        /// Controls whether temp storage for this job is zipped up or not.
+        /// </summary>
+        private static bool bZipTempStorage = true;
+
+        /// <summary>
+        /// Used to set the usage of temp storage as either zipped or having the entire directory structure copied. Defaults to true.
+        /// </summary>
+        /// <param name="Val">true to zip temp storage, false to mirror the directory tree when storing.</param>
+        public static void SetZipTempStorage(bool Val)
         {
-            if (String.IsNullOrEmpty(BaseFolder))
-            {
-                BaseFolder = Env.LocalRoot;
-            }
+            bZipTempStorage = Val;
+        }
 
-            BaseFolder = CommandUtils.CombinePaths(BaseFolder, "/"); 
-            if (!BaseFolder.EndsWith("/") && !BaseFolder.EndsWith("\\"))
+        /// <summary>
+        /// Saves the list of fully qualified files (that should be rooted at the shared temp storage location for the game) to a shared temp storage manifest with the given temp storage node and game.
+        /// </summary>
+        /// <param name="TempStorageNodeInfo">Node info descibing the block of temp storage (essentially used to identify a subdirectory insides the game's temp storage folder).</param>
+        /// <param name="Files">Fully qualified names of files to reference in the manifest file.</param>
+        /// <param name="bLocalOnly">If true, only ensures that the local temp storage manifest exists. Otherwise, checks if the manifest exists in either local or shared storage.</param>
+        /// <param name="GameName">game name to determine the temp storage folder for. Empty is equivalent to "UE4".</param>
+        /// <param name="RootDir">Folder that all the given files are rooted from. If null or empty, CmdEnv.LocalRoot is used.</param>
+        /// <returns>The created manifest instance (which has already been saved to disk).</returns>
+        public static void StoreToTempStorage(TempStorageNodeInfo TempStorageNodeInfo, List<string> Files, bool bLocalOnly, string GameName, string RootDir)
+        {
+            using (var TelemetryStopwatch = new CommandUtils.TelemetryStopwatch("StoreToTempStorage"))
             {
-                throw new AutomationException("base folder {0} should end with a separator", BaseFolder);
-            }
+                // use LocalRoot if one is not specified
+                if (string.IsNullOrEmpty(RootDir))
+                {
+                    RootDir = CommandUtils.CmdEnv.LocalRoot;
+                }
 
-            var Local = SaveLocalTempStorageManifest(Env, BaseFolder, StorageBlockName, Files); 
-            if (!bLocalOnly)
-            {
-                using (CommandUtils.TelemetryStopwatch TelemetryStopwatch = new CommandUtils.TelemetryStopwatch("StoreToTempStorage"))
+                // save the manifest to local temp storage.
+                var Local = SaveLocalTempStorageManifest(RootDir, TempStorageNodeInfo, Files);
+                var LocalTotalSize = Local.GetTotalSize();
+                if (!bLocalOnly)
                 {
                     var StartTime = DateTime.UtcNow;
 
-                    var BlockPath = SharedTempStorageDirectory(StorageBlockName, GameFolder);
+                    var BlockPath = SharedTempStorageDirectory(TempStorageNodeInfo, GameName);
                     CommandUtils.LogConsole("Storing to {0}", BlockPath);
                     if (CommandUtils.DirectoryExists_NoExceptions(BlockPath))
                     {
@@ -769,340 +922,447 @@ namespace AutomationTool
                         throw new AutomationException("Storage Block Could Not Be Created! {0}", BlockPath);
                     }
 
-                    var DestFiles = new List<string>();
-                    if (ThreadsToCopyWith() < 2)
+                    // We know the source files exist and are under RootDir because we created the manifest, which verifies it.
+                    // Now create the list of target files
+                    var DestFiles = Files.Select(Filename => CommandUtils.MakeRerootedFilePath(Filename, RootDir, BlockPath)).ToList();
+                    // make sure the target file doesn't exist, as we never expect this.
+                    foreach (var DestFile in DestFiles)
                     {
-                        foreach (string InFilename in Files)
+                        if (CommandUtils.FileExists_NoExceptions(true, DestFile))
                         {
-                            var Filename = CommandUtils.CombinePaths(InFilename);
-                            InternalUtils.Robust_FileExists_NoExceptions(false, Filename, "Could not add {0} to manifest because it does not exist");
-
-                            if (!Filename.StartsWith(BaseFolder, StringComparison.InvariantCultureIgnoreCase))
-                            {
-                                throw new AutomationException("Could not add {0} to manifest because it does not start with the base folder {1}", Filename, BaseFolder);
-                            }
-                            var RelativeFile = Filename.Substring(BaseFolder.Length);
-                            var DestFile = CommandUtils.CombinePaths(BlockPath, RelativeFile);
-                            if (CommandUtils.FileExists_NoExceptions(true, DestFile))
-                            {
-                                throw new AutomationException("Dest file {0} already exists.", DestFile);
-                            }
-                            CommandUtils.CopyFile(Filename, DestFile, true);
-                            InternalUtils.Robust_FileExists_NoExceptions(true, DestFile, "Could not copy to {0}");
-
-                            DestFiles.Add(DestFile);
+                            throw new AutomationException("Dest file {0} already exists.", DestFile);
                         }
                     }
-                    else
+
+                    long ZipFileSize = 0L;
+                    float ZipCopyDuration = 0.0f;
+                    if (bZipTempStorage)
                     {
-                        var SrcFiles = new List<string>();
-                        foreach (string InFilename in Files)
+                        Ionic.Zip.ZipFile Zip = new Ionic.Zip.ZipFile
                         {
-                            var Filename = CommandUtils.CombinePaths(InFilename);
-                            InternalUtils.Robust_FileExists_NoExceptions(false, Filename, "Could not add {0} to manifest because it does not exist");
-
-                            if (!Filename.StartsWith(BaseFolder, StringComparison.InvariantCultureIgnoreCase))
-                            {
-                                throw new AutomationException("Could not add {0} to manifest because it does not start with the base folder {1}", Filename, BaseFolder);
-                            }
-                            var RelativeFile = Filename.Substring(BaseFolder.Length);
-                            var DestFile = CommandUtils.CombinePaths(BlockPath, RelativeFile);
-                            if (CommandUtils.FileExists_NoExceptions(true, DestFile))
-                            {
-                                throw new AutomationException("Dest file {0} already exists.", DestFile);
-                            }
-                            SrcFiles.Add(Filename);
-                            DestFiles.Add(DestFile);
-                        }
-                        CommandUtils.ThreadedCopyFiles(SrcFiles.ToArray(), DestFiles.ToArray(), ThreadsToCopyWith());
-                        foreach (string DestFile in DestFiles)
+                            UseZip64WhenSaving = Ionic.Zip.Zip64Option.Always
+                        };
+                        foreach (string File in Files)
                         {
-                            InternalUtils.Robust_FileExists_NoExceptions(true, DestFile, "Could not copy to {0}");
+                            // strip the root dir off each file so we get a proper relative path.
+                            Zip.AddFile(Path.Combine(RootDir, File), Path.GetDirectoryName(CommandUtils.StripBaseDirectory(File, RootDir)));
                         }
+                        var LocalManifest = LocalTempStorageManifestFilename(TempStorageNodeInfo);
+                        var SharedManifest = SharedTempStorageManifestFilename(TempStorageNodeInfo, GameName);
+                        var ZipFilename = Path.ChangeExtension(LocalManifest, "zip");
+                        Zip.Save(ZipFilename);
+                        ZipFileSize = new FileInfo(ZipFilename).Length;
+                        var ZipCopyStartTime = DateTime.UtcNow;
+                        InternalUtils.Robust_CopyFile(ZipFilename, CommandUtils.CombinePaths(BlockPath, Path.GetFileName(ZipFilename)));
+                        ZipCopyDuration = (float)(DateTime.UtcNow - ZipCopyStartTime).TotalSeconds;
+                        CommandUtils.DeleteFile(ZipFilename);
+                        // copy the local manifest to the shared location. We have to assume the zip is a good copy.
+                        CommandUtils.CopyFile(LocalManifest, SharedManifest);
+                        TelemetryStopwatch.Finish(string.Format("StoreToTempStorage.{0}.{1}.{2}.Remote", Files.Count, LocalTotalSize, ZipFileSize));
                     }
-                    var Shared = SaveSharedTempStorageManifest(Env, StorageBlockName, GameFolder, DestFiles);
-                    if (!Local.Compare(Shared))
+                    else
                     {
-                        // we will rename this so it can't be used, but leave it around for inspection
-                        CommandUtils.RenameFile_NoExceptions(SharedTempStorageManifestFilename(Env, StorageBlockName, GameFolder), SharedTempStorageManifestFilename(Env, StorageBlockName, GameFolder) + ".broken");
-                        throw new AutomationException("Shared and Local manifest mismatch.");
+                        CommandUtils.ThreadedCopyFiles(Files, DestFiles, ThreadsToCopyWith());
+
+                        // save the shared temp storage manifest file.
+                        var Shared = SaveSharedTempStorageManifest(TempStorageNodeInfo, GameName, DestFiles);
+                        // These manifests better have the same files, timestamps, and sizes, else something failed in the copy.
+                        if (!Local.Compare(Shared))
+                        {
+                            // we will rename this so it can't be used, but leave it around for inspection
+                            var SharedTempManifestFilename = SharedTempStorageManifestFilename(TempStorageNodeInfo, GameName);
+                            CommandUtils.RenameFile_NoExceptions(SharedTempManifestFilename, SharedTempManifestFilename + ".broken");
+                            throw new AutomationException("Shared and Local manifest mismatch. Wrote out {0} for inspection.", SharedTempManifestFilename + ".broken");
+                        }
+                        TelemetryStopwatch.Finish(string.Format("StoreToTempStorage.{0}.{1}.{2}.Remote", Files.Count, LocalTotalSize, 0L));
                     }
+
                     float BuildDuration = (float)((DateTime.UtcNow - StartTime).TotalSeconds);
-                    if (BuildDuration > 10.0f && Shared.GetTotalSize() > 0)
+                    if (LocalTotalSize > 0)
                     {
-                        var MBSec = (((float)(Shared.GetTotalSize())) / (1024.0f * 1024.0f)) / BuildDuration;
-                        CommandUtils.LogConsole("Wrote to shared temp storage at {0} MB/s    {1}B {2}s", MBSec, Shared.GetTotalSize(), BuildDuration);
+                        var MBSec = (LocalTotalSize / (1024.0f * 1024.0f)) / BuildDuration;
+                        CommandUtils.Log("Wrote to shared temp storage at {0} MB/s    {1}B {2}s", MBSec, LocalTotalSize, BuildDuration);
+                        if (bZipTempStorage)
+                        {
+                            var ZipMBSec = (ZipFileSize / (1024.0f * 1024.0f)) / ZipCopyDuration;
+                            CommandUtils.Log("Wrote zip to shared temp storage at {0} MB/s    {1}B {2}s", ZipMBSec, ZipFileSize, ZipCopyDuration);
+                        }
                     }
-                    TelemetryStopwatch.Finish(string.Format("StoreToTempStorage.{0}.{1}", Shared.GetFileCount(), Shared.GetTotalSize()));
+                }
+                else
+                {
+                    TelemetryStopwatch.Finish(string.Format("StoreToTempStorage.{0}.{1}.{2}.Local", Files.Count, LocalTotalSize, 0L));
                 }
             }
-
         }
-		public static void RetrieveFromPermanentStorage(CommandEnvironment Env, string NetworkRoot, string Build, string Platform)
-		{
-			var LocalManifest = LocalTempStorageManifestFilename(Env, Build);
-			if (!CommandUtils.FileExists_NoExceptions(LocalManifest))
-			{
-				string BaseFolder = CommandUtils.CombinePaths(Env.LocalRoot, "Rocket", "TempInst", Platform);
-				string Root = CommandUtils.RootSharedTempStorageDirectory();
-				string SrcFolder = CommandUtils.CombinePaths(Root, NetworkRoot, Build, Platform);
-				var SourceFiles = CommandUtils.FindFiles_NoExceptions("*", true, SrcFolder);
-				foreach (var File in SourceFiles)
-				{
-					var DestFile = File.Replace(SrcFolder, "");
-					DestFile = CommandUtils.CombinePaths(BaseFolder, DestFile);
-					CommandUtils.CopyFile_NoExceptions(File, DestFile);
-				}
-				if(!CommandUtils.DirectoryExists(Path.GetDirectoryName(LocalManifest)))
-				{
-					CommandUtils.CreateDirectory(Path.GetDirectoryName(LocalManifest));
-				}
-				XmlDocument Local = new XmlDocument();
-				XmlElement RootElement = Local.CreateElement("permstorage");
-				Local.AppendChild(RootElement);
-				Local.Save(LocalManifest);
-			}
-		}
-        public static List<string> RetrieveFromTempStorage(CommandEnvironment Env, string StorageBlockName, out bool WasLocal, string GameFolder = "", string BaseFolder = "")
-        {
-            if (String.IsNullOrEmpty(BaseFolder))
-            {
-                BaseFolder = Env.LocalRoot;
-            }
 
-            BaseFolder = CommandUtils.CombinePaths(BaseFolder, "/");
-            if (!BaseFolder.EndsWith("/") && !BaseFolder.EndsWith("\\"))
+        /// <summary>
+        /// Inverse of <see cref="StoreToTempStorage"/>.
+        /// Copies a block of files from a temp storage location given by a temp storage node and game to local storage rooted at the given root dir. 
+        /// If the temp manifest for this block is found locally, the copy is skipped, as we assume this is the same machine that created the temp storage and the files are still there.
+        /// </summary>
+        /// <param name="TempStorageNodeInfo">Node info descibing the block of temp storage (essentially used to identify a subdirectory insides the game's temp storage folder).</param>
+        /// <param name="WasLocal">upon return, this parameter is set to true if the temp manifest was found locally and the copy was avoided.</param>
+        /// <param name="GameName">game name to determine the temp storage folder for. Empty is equivalent to "UE4".</param>
+        /// <param name="RootDir">Folder that all the retrieved files should be rooted from. If null or empty, CmdEnv.LocalRoot is used.</param>
+        /// <returns>List of fully qualified paths to all the files that were retrieved. This is returned even if we skip the copy (set WasLocal = true) .</returns>
+        public static List<string> RetrieveFromTempStorage(TempStorageNodeInfo TempStorageNodeInfo, out bool WasLocal, string GameName, string RootDir)
+        {
+            using (var TelemetryStopwatch = new CommandUtils.TelemetryStopwatch("RetrieveFromTempStorage"))
             {
-                throw new AutomationException("base folder {0} should end with a separator", BaseFolder);
-            }
+                // use LocalRoot if one is not specified
+                if (String.IsNullOrEmpty(RootDir))
+                {
+                    RootDir = CommandUtils.CmdEnv.LocalRoot;
+                }
 
-            var Files = new List<string>();
-            var LocalManifest = LocalTempStorageManifestFilename(Env, StorageBlockName);
-            if (CommandUtils.FileExists_NoExceptions(LocalManifest))
-            {
-                CommandUtils.LogVerbose("Found local manifest {0}", LocalManifest);
-                var Local = new TempStorageManifest();
-                Local.Load(LocalManifest);
-                Files = Local.GetFiles(BaseFolder);
-                var LocalTest = new TempStorageManifest();
-                LocalTest.Create(Files, BaseFolder);
-                if (!Local.Compare(LocalTest))
+                // First see if the local manifest is there.
+                // If it is, then we must be on the same node as the one that originally created the temp storage.
+                // In that case, we just verify all the files exist as described in the manifest and use that.
+                // If there was any tampering, abort immediately because we never accept tampering, and it signifies a build problem.
+                var LocalManifest = LocalTempStorageManifestFilename(TempStorageNodeInfo);
+                if (CommandUtils.FileExists_NoExceptions(LocalManifest))
                 {
-                    throw new AutomationException("Local files in manifest {0} were tampered with.", LocalManifest);
+                    CommandUtils.LogVerbose("Found local manifest {0}", LocalManifest);
+                    var Local = TempStorageManifest.Load(LocalManifest);
+                    var Files = Local.GetFiles(RootDir);
+                    var LocalTest = TempStorageManifest.Create(Files, RootDir);
+                    if (!Local.Compare(LocalTest))
+                    {
+                        throw new AutomationException("Local files in manifest {0} were tampered with.", LocalManifest);
+                    }
+                    WasLocal = true;
+                    TelemetryStopwatch.Finish(string.Format("RetrieveFromTempStorage.{0}.{1}.{2}.Local", Files.Count, Local.GetTotalSize(), 0));
+                    return Files;
                 }
-                WasLocal = true;
-                return Files;
-            }
-            WasLocal = false;
-            var StartTime = DateTime.UtcNow;
+                WasLocal = false;
+                var StartTime = DateTime.UtcNow;
 
-            var BlockPath = CommandUtils.CombinePaths(SharedTempStorageDirectory(StorageBlockName, GameFolder), "/");
-            if (!BlockPath.EndsWith("/") && !BlockPath.EndsWith("\\"))
-            {
-                throw new AutomationException("base folder {0} should end with a separator", BlockPath);
-            }
-            CommandUtils.LogVerbose("Attempting to retrieve from {0}", BlockPath);
-            if (!CommandUtils.DirectoryExists_NoExceptions(BlockPath))
-            {
-                throw new AutomationException("Storage Block Does Not Exists! {0}", BlockPath);
-            }
-            var SharedManifest = SharedTempStorageManifestFilename(Env, StorageBlockName, GameFolder);
-            InternalUtils.Robust_FileExists_NoExceptions(SharedManifest, "Storage Block Manifest Does Not Exists! {0}");
+                // We couldn't find the node storage locally, so get it from the shared location.
+                var BlockPath = SharedTempStorageDirectory(TempStorageNodeInfo, GameName);
 
-            var Shared = new TempStorageManifest();
-            Shared.Load(SharedManifest);
+                CommandUtils.LogVerbose("Attempting to retrieve from {0}", BlockPath);
+                if (!CommandUtils.DirectoryExists_NoExceptions(BlockPath))
+                {
+                    throw new AutomationException("Storage Block Does Not Exists! {0}", BlockPath);
+                }
+                var SharedManifest = SharedTempStorageManifestFilename(TempStorageNodeInfo, GameName);
+                InternalUtils.Robust_FileExists(SharedManifest, "Storage Block Manifest Does Not Exists! {0}");
 
-            var SharedFiles = Shared.GetFiles(BlockPath);
+                var Shared = TempStorageManifest.Load(SharedManifest);
+                var SharedFiles = Shared.GetFiles(BlockPath, !bZipTempStorage);
 
-            var DestFiles = new List<string>();
-            if (ThreadsToCopyWith() < 2)
-            {
-                foreach (string InFilename in SharedFiles)
-                {
-                    var Filename = CommandUtils.CombinePaths(InFilename);
-                    InternalUtils.Robust_FileExists_NoExceptions(true, Filename, "Could not add {0} to manifest because it does not exist");
+                // We know the source files exist and are under RootDir because we created the manifest, which verifies it.
+                // Now create the list of target files
+                var DestFiles = SharedFiles.Select(Filename => CommandUtils.MakeRerootedFilePath(Filename, BlockPath, RootDir)).ToList();
 
-                    if (!Filename.StartsWith(BlockPath, StringComparison.InvariantCultureIgnoreCase))
-                    {
-                        throw new AutomationException("Could not add {0} to manifest because it does not start with the base folder {1}", Filename, BlockPath);
-                    }
-                    var RelativeFile = Filename.Substring(BlockPath.Length);
-                    var DestFile = CommandUtils.CombinePaths(BaseFolder, RelativeFile);
-                    if (CommandUtils.FileExists_NoExceptions(true, DestFile))
+                long ZipFileSize = 0L;
+                float ZipCopyDuration = 0.0f;
+                if (bZipTempStorage)
+                {
+                    var LocalZipFilename = Path.ChangeExtension(LocalManifest, "zip");
+                    var SharedZipFilename = Path.ChangeExtension(SharedManifest, "zip");
+                    var ZipCopyStartTime = DateTime.UtcNow;
+                    if (!CommandUtils.FileExists(SharedZipFilename))
                     {
-                        CommandUtils.LogVerbose("Dest file {0} already exists, deleting and overwriting", DestFile);
-                        CommandUtils.DeleteFile(DestFile);
+                        throw new AutomationException("Storage block zip file does not exist! {0}", SharedZipFilename);
                     }
-                    CommandUtils.CopyFile(Filename, DestFile, true);
+                    InternalUtils.Robust_CopyFile(SharedZipFilename, LocalZipFilename);
+                    ZipCopyDuration = (float)(DateTime.UtcNow - ZipCopyStartTime).TotalSeconds;
+                    ZipFileSize = new FileInfo(LocalZipFilename).Length;
 
-                    InternalUtils.Robust_FileExists_NoExceptions(true, DestFile, "Could not copy to {0}");
-
-                    if (UnrealBuildTool.Utils.IsRunningOnMono)
+                    using (Ionic.Zip.ZipFile Zip = new Ionic.Zip.ZipFile(LocalZipFilename))
                     {
-						CommandUtils.FixUnixFilePermissions(DestFile);
+                        Zip.ExtractAll(RootDir, Ionic.Zip.ExtractExistingFileAction.OverwriteSilently);
                     }
 
-                    FileInfo Info = new FileInfo(DestFile);
-                    DestFiles.Add(Info.FullName);
+                    CommandUtils.DeleteFile(LocalZipFilename);
                 }
-            }
-            else
-            {
-                var SrcFiles = new List<string>();
-                foreach (string InFilename in SharedFiles)
+                else
                 {
-                    var Filename = CommandUtils.CombinePaths(InFilename);
-                    //Robust_FileExists_NoExceptions(true, Filename, "Could not add {0} to manifest because it does not exist");
-
-                    if (!Filename.StartsWith(BlockPath, StringComparison.InvariantCultureIgnoreCase))
+                    // If the local file already exists, it will be overwritten.
+                    foreach (var DestFile in DestFiles)
                     {
-                        throw new AutomationException("Could not add {0} to manifest because it does not start with the base folder {1}", Filename, BlockPath);
+                        if (CommandUtils.FileExists_NoExceptions(true, DestFile))
+                        {
+                            CommandUtils.LogVerbose("Dest file {0} already exists, deleting and overwriting", DestFile);
+                            CommandUtils.DeleteFile(DestFile);
+                        }
                     }
-                    var RelativeFile = Filename.Substring(BlockPath.Length);
-                    var DestFile = CommandUtils.CombinePaths(BaseFolder, RelativeFile);
-                    if (CommandUtils.FileExists_NoExceptions(true, DestFile))
+
+                    // Do the threaded copy to the local file system.
+                    CommandUtils.ThreadedCopyFiles(SharedFiles, DestFiles, ThreadsToCopyWith());
+                }
+
+                // Handle unix permissions/chmod issues.
+                if (UnrealBuildTool.Utils.IsRunningOnMono)
+                {
+                    foreach (string DestFile in DestFiles)
                     {
-                        CommandUtils.LogVerbose("Dest file {0} already exists, deleting and overwriting", DestFile);
-                        CommandUtils.DeleteFile(DestFile);
+                        CommandUtils.FixUnixFilePermissions(DestFile);
                     }
-                    SrcFiles.Add(Filename);
-                    DestFiles.Add(DestFile);
                 }
-                CommandUtils.ThreadedCopyFiles(SrcFiles.ToArray(), DestFiles.ToArray(), ThreadsToCopyWith());
-                var NewDestFiles = new List<string>();
-                foreach (string DestFile in DestFiles)
+                var NewLocal = SaveLocalTempStorageManifest(RootDir, TempStorageNodeInfo, DestFiles);
+                // Now compare the created local files to ensure their attributes match the one we copied from the network.
+                if (!NewLocal.Compare(Shared))
+                {
+                    // we will rename this so it can't be used, but leave it around for inspection
+                    CommandUtils.RenameFile_NoExceptions(LocalManifest, LocalManifest + ".broken");
+                    throw new AutomationException("Shared and Local manifest mismatch.");
+                }
+                float BuildDuration = (float)((DateTime.UtcNow - StartTime).TotalSeconds);
+                var SharedTotalSize = Shared.GetTotalSize();
+                if (SharedTotalSize > 0)
                 {
-                    InternalUtils.Robust_FileExists_NoExceptions(true, DestFile, "Could not copy to {0}");
-                    if (UnrealBuildTool.Utils.IsRunningOnMono)
+                    var MBSec = (Shared.GetTotalSize() / (1024.0f * 1024.0f)) / BuildDuration;
+                    CommandUtils.Log("Read from shared temp storage at {0} MB/s    {1}B {2}s", MBSec, SharedTotalSize, BuildDuration);
+                    if (bZipTempStorage)
                     {
-						CommandUtils.FixUnixFilePermissions(DestFile);
+                        var ZipMBSec = (ZipFileSize / (1024.0f * 1024.0f)) / ZipCopyDuration;
+                        CommandUtils.Log("Read zip from shared temp storage at {0} MB/s    {1}B {2}s", ZipMBSec, ZipFileSize, ZipCopyDuration);
                     }
-                    FileInfo Info = new FileInfo(DestFile);
-                    NewDestFiles.Add(Info.FullName);
                 }
-                DestFiles = NewDestFiles;
-            }
-            var NewLocal = SaveLocalTempStorageManifest(Env, BaseFolder, StorageBlockName, DestFiles);
-            if (!NewLocal.Compare(Shared))
-            {
-                // we will rename this so it can't be used, but leave it around for inspection
-                CommandUtils.RenameFile_NoExceptions(LocalManifest, LocalManifest + ".broken");
-                throw new AutomationException("Shared and Local manifest mismatch.");
+                TelemetryStopwatch.Finish(string.Format("RetrieveFromTempStorage.{0}.{1}.{2}.Remote", DestFiles.Count, SharedTotalSize, ZipFileSize));
+                return DestFiles;
             }
-            float BuildDuration = (float)((DateTime.UtcNow - StartTime).TotalSeconds);
-            if (BuildDuration > 60.0f && Shared.GetTotalSize() > 0)
-            {
-                var MBSec = (((float)(Shared.GetTotalSize())) / (1024.0f * 1024.0f)) / BuildDuration;
-                CommandUtils.Log("Read from shared temp storage at {0} MB/s    {1}B {2}s", MBSec, Shared.GetTotalSize(), BuildDuration);
-            }
-            return DestFiles;
         }
-        public static List<string> FindTempStorageManifests(CommandEnvironment Env, string StorageBlockName, bool LocalOnly = false, bool SharedOnly = false, string GameFolder = "")
+
+        /// <summary>
+        /// Finds temp storage manifests that match the block name (where the block name has embedded wildcards to find all block names of a given pattern).
+        /// This is used by <see cref="GUBP.FindNodeHistory"/> to search for any temp storage manifests that match a certain GUBP node name.
+        /// This is used to construct a "history" of that node which is then used to log info and generate failure emails.
+        /// This method is used multiple times to find nodes matching a specific name with a suffix of _Started _Failed and _Succeeded. 
+        /// The CL# is pulled from those names and used to generate a P4 history.
+        /// Main problem with this is that temp storage is ephemeral, so after a few days, this will typically not find any history for a node.
+        /// 
+        /// NOTE: This entire routine is dependent on how TempStorageNodeInfo lays out its directories!!!
+        /// If this layout is changed, this code needs to be changed as well.
+        /// @todo: Someday Make FindNodeHistory query the build database directly to get a true history of the node.
+        /// </summary>
+        /// <param name="TempStorageNodeInfo">Node info descibing the block of temp storage (essentially used to identify a subdirectory insides the game's temp storage folder).</param>
+        /// <param name="GameName">game name to determine the temp storage folder for. Empty is equivalent to "UE4".</param>
+        /// <returns></returns>
+        public static List<int> FindMatchingSharedTempStorageNodeCLs(TempStorageNodeInfo TempStorageNodeInfo, string GameName)
         {
-            var Files = new List<string>();
+            // Find the shared temp storage folder for this game and branch
+            var SharedTempStorageDirectoryForGameAndBranch = CommandUtils.CombinePaths(ResolveSharedTempStorageDirectory(GameName), TempStorageNodeInfo.JobInfo.BranchNameForTempStorage);
+            int dummy;
+            return (
+                // Look for all folders underneath this, it should be a CL, or CL-PreflightInfo
+                from CLDir in Directory.EnumerateDirectories(SharedTempStorageDirectoryForGameAndBranch)
+                // only accept folders that are plain numbers. Any suffixes ('-') imply preflight builds or something else.
+                let PossibleCL = Path.GetFileName(CLDir)
+                where int.TryParse(PossibleCL, out dummy)
+                let CL = int.Parse(PossibleCL)
+                // if we have a real CL, try to find the node name we are looking for.
+                where File.Exists(CommandUtils.CombinePaths(CLDir, TempStorageNodeInfo.NodeStorageName, TempStorageNodeInfo.GetManifestFilename()))
+                select CL)
+                .OrderBy(CL => CL).ToList();
+        }
 
-            var LocalFiles = LocalTempStorageManifestFilename(Env, StorageBlockName);
-            var LocalParent = Path.GetDirectoryName(LocalFiles);
-            var WildCard = Path.GetFileName(LocalFiles);
+        /// <summary>
+        /// Runs a test of the temp storage system. This function is part of the class so it doesn't have to expose internals just to run the test.
+        /// </summary>
+        /// <param name="CmdEnv">The environment to use.</param>
+        static internal void TestTempStorage(CommandEnvironment CmdEnv)
+        {
+            // We are not a real GUBP job, so fake the values.
+            var TestTempStorageInfo = new TempStorageNodeInfo(new GUBP.JobInfo("Test", 0, 0, 0), "Test");
 
-            int IndexOfStar = WildCard.IndexOf("*");
-            if (IndexOfStar < 0 || WildCard.LastIndexOf("*") != IndexOfStar)
+            // Delete any local and shared temp storage that may exist.
+            TempStorage.DeleteLocalTempStorageManifests();
+            TempStorage.DeleteSharedTempStorage(TestTempStorageInfo, "UE4");
+            if (TempStorage.TempStorageExists(TestTempStorageInfo, "UE4", false, false))
             {
-                throw new AutomationException("Wildcard {0} either has no star or it has more than one.", WildCard);
+                throw new AutomationException("storage should not exist");
             }
 
-            string PreStarWildcard = WildCard.Substring(0, IndexOfStar);
-            string PostStarWildcard = Path.GetFileNameWithoutExtension(WildCard.Substring(IndexOfStar + 1));
-
-            if (!SharedOnly && CommandUtils.DirectoryExists_NoExceptions(LocalParent))
+            // Create some test files in various locations in the local root with unique content.
+            var TestFiles = new[]
+                { 
+                    CommandUtils.CombinePaths(CmdEnv.LocalRoot, "Engine", "Build", "Batchfiles", "TestFile.Txt"),
+                    CommandUtils.CombinePaths(CmdEnv.LocalRoot, "TestFile2.Txt"),
+                    CommandUtils.CombinePaths(CmdEnv.LocalRoot, "engine", "plugins", "TestFile3.Txt"),
+                }
+                .Select(TestFile => new 
+                { 
+                    FileName = TestFile, 
+                    FileContents = string.Format("{0} - {1}", TestFile, DateTime.UtcNow) 
+                })
+                .ToList();
+            foreach (var TestFile in TestFiles)
+            {
+                CommandUtils.Log("Test file {0}", TestFile.FileName);
+                File.WriteAllText(TestFile.FileName, TestFile.FileContents);
+            }
+            
+            // wrap the operation so we are sure to clean up afterward.
+            try
             {
-                foreach (var ThisFile in CommandUtils.FindFiles_NoExceptions(WildCard, true, LocalParent))
+                // Store the test file to temp storage.
+                TempStorage.StoreToTempStorage(TestTempStorageInfo, TestFiles.Select(TestFile => TestFile.FileName).ToList(), false, "UE4", CmdEnv.LocalRoot);
+                // The manifest should exist locally.
+                if (!TempStorage.LocalTempStorageManifestExists(TestTempStorageInfo))
+                {
+                    throw new AutomationException("local storage should exist");
+                }
+                // The manifest should exist on the shared drive.
+                if (!TempStorage.SharedTempStorageManifestExists(TestTempStorageInfo, "UE4", false))
+                {
+                    throw new AutomationException("shared storage should exist");
+                }
+                // Now delete the local manifest
+                TempStorage.DeleteLocalTempStorageManifests();
+                // It should no longer be there.
+                if (TempStorage.LocalTempStorageManifestExists(TestTempStorageInfo))
+                {
+                    throw new AutomationException("local storage should not exist");
+                }
+                // But the shared storage should still exist.
+                if (!TempStorage.TempStorageExists(TestTempStorageInfo, "UE4", false, false))
+                {
+                    throw new AutomationException("some storage should exist");
+                }
+                // Ok, delete our test files locally.
+                foreach (var TestFile in TestFiles)
+                {
+                    CommandUtils.DeleteFile(TestFile.FileName);
+                }
+
+                // Now we should be able to retrieve the test files from shared storage.
+                bool WasLocal;
+                TempStorage.RetrieveFromTempStorage(TestTempStorageInfo, out WasLocal, "UE4", CmdEnv.LocalRoot);
+                // the local manifest should be there, since we just retrieved from shared storage.
+                if (!TempStorage.LocalTempStorageManifestExists(TestTempStorageInfo))
+                {
+                    throw new AutomationException("local storage should exist");
+                }
+                // The shared manifest should also still be there.
+                if (!TempStorage.SharedTempStorageManifestExists(TestTempStorageInfo, "UE4", false))
                 {
-                    CommandUtils.LogConsole("  Found local file {0}", ThisFile);
-                    int IndexOfWildcard = ThisFile.IndexOf(PreStarWildcard);
-                    if (IndexOfWildcard < 0)
+                    throw new AutomationException("shared storage should exist");
+                }
+                // Verify the retrieved files have the correct content.
+                foreach (var TestFile in TestFiles)
+                {
+                    if (File.ReadAllText(TestFile.FileName) != TestFile.FileContents)
                     {
-                        throw new AutomationException("File {0} didn't contain {1}.", ThisFile, PreStarWildcard);
+                        throw new AutomationException("Contents of the test file {0} was changed after restoring from shared temp storage.", TestFile.FileName);
                     }
-                    int LastIndexOfWildcardTail = ThisFile.LastIndexOf(PostStarWildcard);
-                    if (LastIndexOfWildcardTail < 0 || LastIndexOfWildcardTail < IndexOfWildcard + PreStarWildcard.Length)
+                }
+                // Now delete the shared temp storage
+                TempStorage.DeleteSharedTempStorage(TestTempStorageInfo, "UE4");
+                // Shared temp storage manifest should no longer exist.
+                if (TempStorage.SharedTempStorageManifestExists(TestTempStorageInfo, "UE4", false))
+                {
+                    throw new AutomationException("shared storage should not exist");
+                }
+                // Retrieving temp storage should now just retrieve from local
+                TempStorage.RetrieveFromTempStorage(TestTempStorageInfo, out WasLocal, "UE4", CmdEnv.LocalRoot);
+                if (!WasLocal || !TempStorage.LocalTempStorageManifestExists(TestTempStorageInfo))
+                {
+                    throw new AutomationException("local storage should exist");
+                }
+
+                // and now lets test tampering
+                TempStorage.DeleteLocalTempStorageManifests();
+                {
+                    bool bFailedProperly = false;
+                    var MissingFile = new List<string>(TestFiles.Select(TestFile => TestFile.FileName));
+                    // add a file to the manifest that shouldn't be there.
+                    MissingFile.Add(CommandUtils.CombinePaths(CmdEnv.LocalRoot, "Engine", "SomeFileThatDoesntExist.txt"));
+                    try
                     {
-                        throw new AutomationException("File {0} didn't contain {1} or it was before the prefix", ThisFile, PostStarWildcard);
+                        TempStorage.StoreToTempStorage(TestTempStorageInfo, MissingFile, false, "UE4", CmdEnv.LocalRoot);
                     }
-                    string StarReplacement = ThisFile.Substring(IndexOfWildcard + PreStarWildcard.Length, LastIndexOfWildcardTail - IndexOfWildcard - PreStarWildcard.Length);
-                    if (StarReplacement.Length < 1)
+                    catch (AutomationException)
                     {
-                        throw new AutomationException("Dir {0} didn't have any string to fit the star in the wildcard {1}", ThisFile, WildCard);
+                        bFailedProperly = true;
                     }
-                    if (!Files.Contains(StarReplacement))
+                    if (!bFailedProperly)
                     {
-                        Files.Add(StarReplacement);
+                        throw new AutomationException("Missing file did not fail.");
                     }
                 }
-            }
-
-            if (!LocalOnly)
-            {
-                var SharedFiles = SharedTempStorageManifestFilename(Env, StorageBlockName, GameFolder);
-                var SharedParent = Path.GetDirectoryName(Path.GetDirectoryName(SharedFiles));
 
-                if (CommandUtils.DirectoryExists_NoExceptions(SharedParent))
+                // clear the shared temp storage again.
+                TempStorage.DeleteSharedTempStorage(TestTempStorageInfo, "UE4");
+                // store the test files to shared temp storage again.
+                TempStorage.StoreToTempStorage(TestTempStorageInfo, TestFiles.Select(TestFile => TestFile.FileName).ToList(), false, "UE4", CmdEnv.LocalRoot);
+                // force a load from shared by deleting the local manifest
+                TempStorage.DeleteLocalTempStorageManifests();
+                if (bZipTempStorage)
+                {
+                    var ZipFile = Path.ChangeExtension(TempStorage.SharedTempStorageManifestFilename(TestTempStorageInfo, "UE4"), "zip");
+                    // the shared storage file should exist.
+                    if (!CommandUtils.FileExists_NoExceptions(ZipFile))
+                    {
+                        throw new AutomationException("Shared file {0} did not exist", ZipFile);
+                    }
+                    // delete the shared file.
+                    CommandUtils.DeleteFile(ZipFile);
+                }
+                else
                 {
-                    string[] Dirs = null;
+                    var SharedFile = CommandUtils.MakeRerootedFilePath(TestFiles[0].FileName, CmdEnv.LocalRoot, TempStorage.SharedTempStorageDirectory(TestTempStorageInfo, "UE4"));
+                    // the shared storage file should exist.
+                    if (!CommandUtils.FileExists_NoExceptions(SharedFile))
+                    {
+                        throw new AutomationException("Shared file {0} did not exist", SharedFile);
+                    }
+                    // delete the shared file.
+                    CommandUtils.DeleteFile(SharedFile);
+                }
+                // now test that retrieving from shared temp storage properly balks that a file is missing.
+                {
+                    bool bFailedProperly = false;
+                    try
+                    {
+                        TempStorage.RetrieveFromTempStorage(TestTempStorageInfo, out WasLocal, "UE4", CmdEnv.LocalRoot);
+                    }
+                    catch (AutomationException)
+                    {
+                        bFailedProperly = true;
+                    }
+                    if (!bFailedProperly)
+                    {
+                        throw new AutomationException("Did not fail to load from missing file.");
+                    }
+                }
 
+                // clear the shared temp storage again.
+                TempStorage.DeleteSharedTempStorage(TestTempStorageInfo, "UE4");
+                // Copy the files to temp storage.
+                TempStorage.StoreToTempStorage(TestTempStorageInfo, TestFiles.Select(TestFile => TestFile.FileName).ToList(), false, "UE4", CmdEnv.LocalRoot);
+                // Delete a local file.
+                CommandUtils.DeleteFile(TestFiles[0].FileName);
+                // retrieving from temp storage should use WasLocal, but should balk because a local file was deleted.
+                {
+                    bool bFailedProperly = false;
                     try
                     {
-                        Dirs = Directory.GetDirectories(SharedParent, Path.GetFileNameWithoutExtension(SharedFiles), SearchOption.TopDirectoryOnly);
+                        TempStorage.RetrieveFromTempStorage(TestTempStorageInfo, out WasLocal, "UE4", CmdEnv.LocalRoot);
                     }
-                    catch (Exception Ex)
+                    catch (AutomationException)
                     {
-                        CommandUtils.LogConsole("Unable to Find Directories in {0} with wildcard {1}", SharedParent, Path.GetFileNameWithoutExtension(SharedFiles));
-                        CommandUtils.LogConsole(" Exception was {0}", LogUtils.FormatException(Ex));
+                        bFailedProperly = true;
                     }
-                    if (Dirs != null)
+                    if (!bFailedProperly)
                     {
-                        foreach (var ThisSubDir in Dirs)
-                        {
-                            int IndexOfWildcard = ThisSubDir.IndexOf(PreStarWildcard);
-                            if (IndexOfWildcard < 0)
-                            {
-                                throw new AutomationException("Dir {0} didn't contain {1}.", ThisSubDir, PreStarWildcard);
-                            }
-                            int LastIndexOfWildcardTail = ThisSubDir.LastIndexOf(PostStarWildcard);
-                            if (LastIndexOfWildcardTail < 0 || LastIndexOfWildcardTail < IndexOfWildcard + PreStarWildcard.Length)
-                            {
-                                throw new AutomationException("Dir {0} didn't contain {1} or it was before the prefix", ThisSubDir, PostStarWildcard);
-                            }
-                            string StarReplacement = ThisSubDir.Substring(IndexOfWildcard + PreStarWildcard.Length, LastIndexOfWildcardTail - IndexOfWildcard - PreStarWildcard.Length);
-                            if (StarReplacement.Length < 1)
-                            {
-                                throw new AutomationException("Dir {0} didn't have any string to fit the star in the wildcard {1}", ThisSubDir, WildCard);
-                            }
-                            // these are a bunch of false positives
-                            if (StarReplacement.Contains("-"))
-                            {
-                                continue;
-                            }
-                            if (!Files.Contains(StarReplacement))
-                            {
-                                Files.Add(StarReplacement);
-                            }
-                        }
+                        throw new AutomationException("Did not fail to load from missing local file.");
                     }
                 }
             }
-
-            var OutFiles = new List<string>();
-            foreach (var StarReplacement in Files)
+            finally
             {
-                var NewBlock = StorageBlockName.Replace("*", StarReplacement);
-
-                if (TempStorageExists(Env, NewBlock, GameFolder, LocalOnly, true))
+                // Do a final cleanup.
+                TempStorage.DeleteSharedTempStorage(TestTempStorageInfo, "UE4");
+                TempStorage.DeleteLocalTempStorageManifests();
+                foreach (var TestFile in TestFiles)
                 {
-                    OutFiles.Add(StarReplacement);
+                    CommandUtils.DeleteFile(TestFile.FileName);
                 }
             }
-            return OutFiles;
         }
     }
 
@@ -1113,256 +1373,8 @@ namespace AutomationTool
 		{
 			Log("TestTempStorage********");
 
-			TempStorage.DeleteLocalTempStorageManifests(CmdEnv);
-			TempStorage.DeleteSharedTempStorageManifests(CmdEnv, "Test");
-			if (TempStorage.TempStorageExists(CmdEnv, "Test"))
-			{
-				throw new AutomationException("storage should not exist");
-			}
-
-			string UnitTestFile = CombinePaths(CmdEnv.LocalRoot, "Engine", "Build", "Batchfiles", "TestFile.Txt");
-			Log("Test file {0}", UnitTestFile);
-
-			if (FileExists(UnitTestFile))
-			{
-				DeleteFile(UnitTestFile);
-			}
-			string[] LinesToWrite = new string[]
-			{
-				"This is a test",
-				"Of writing to file " + UnitTestFile
-			};
-			foreach (string Line in LinesToWrite)
-			{
-				WriteToFile(UnitTestFile, Line);
-			}
-
-			string UnitTestFile2 = CombinePaths(CmdEnv.LocalRoot, "TestFile2.Txt");
-			Log("Test file {0}", UnitTestFile2);
-
-			if (FileExists(UnitTestFile2))
-			{
-				DeleteFile(UnitTestFile2);
-			}
-			string[] LinesToWrite2 = new string[]
-			{
-				"This is a test",
-				"Of writing to file " + UnitTestFile2
-			};
-			foreach (string Line in LinesToWrite2)
-			{
-				WriteToFile(UnitTestFile2, Line);
-			}
-
-			string UnitTestFile3 = CombinePaths(CmdEnv.LocalRoot, "engine", "plugins", "TestFile3.Txt");
-			Log("Test file {0}", UnitTestFile3);
-
-			if (FileExists(UnitTestFile3))
-			{
-				DeleteFile(UnitTestFile3);
-			}
-			string[] LinesToWrite3 = new string[]
-			{
-				"This is a test",
-				"Of writing to file " + UnitTestFile3
-			};
-			foreach (string Line in LinesToWrite3)
-			{
-				WriteToFile(UnitTestFile3, Line);
-			}
-
-			{
-				string[] LinesRead = ReadAllLines(UnitTestFile);
-				if (LinesRead == null || LinesRead.Length != LinesToWrite.Length)
-				{
-					throw new AutomationException("Contents of the file created is different to the file read.");
-				}
-				for (int LineIndex = 0; LineIndex < LinesRead.Length; ++LineIndex)
-				{
-					if (LinesRead[LineIndex] != LinesToWrite[LineIndex])
-					{
-						throw new AutomationException("Contents of the file created is different to the file read.");
-					}
-				}
-
-				string[] LinesRead2 = ReadAllLines(UnitTestFile2);
-				if (LinesRead2 == null || LinesRead2.Length != LinesToWrite2.Length)
-				{
-					throw new AutomationException("Contents of the file created is different to the file read.");
-				}
-				for (int LineIndex = 0; LineIndex < LinesRead2.Length; ++LineIndex)
-				{
-					if (LinesRead2[LineIndex] != LinesToWrite2[LineIndex])
-					{
-						throw new AutomationException("Contents of the file created is different to the file read.");
-					}
-				}
-
-				string[] LinesRead3 = ReadAllLines(UnitTestFile3);
-				if (LinesRead3 == null || LinesRead3.Length != LinesToWrite3.Length)
-				{
-					throw new AutomationException("Contents of the file created is different to the file read.");
-				}
-				for (int LineIndex = 0; LineIndex < LinesRead3.Length; ++LineIndex)
-				{
-					if (LinesRead3[LineIndex] != LinesToWrite3[LineIndex])
-					{
-						throw new AutomationException("Contents of the file created is different to the file read.");
-					}
-				}
-			}
-			var TestFiles = new List<string> { UnitTestFile, UnitTestFile2, UnitTestFile3 };
-
-			TempStorage.StoreToTempStorage(CmdEnv, "Test", TestFiles);
-
-			if (!TempStorage.LocalTempStorageExists(CmdEnv, "Test"))
-			{
-				throw new AutomationException("local storage should exist");
-			}
-			if (!TempStorage.SharedTempStorageExists(CmdEnv, "Test"))
-			{
-				throw new AutomationException("shared storage should exist");
-			}
-			TempStorage.DeleteLocalTempStorageManifests(CmdEnv);
-			if (TempStorage.LocalTempStorageExists(CmdEnv, "Test"))
-			{
-				throw new AutomationException("local storage should not exist");
-			}
-			if (!TempStorage.TempStorageExists(CmdEnv, "Test"))
-			{
-				throw new AutomationException("some storage should exist");
-			}
-			DeleteFile(UnitTestFile);
-			DeleteFile(UnitTestFile2);
-			DeleteFile(UnitTestFile3);
-
-			bool WasLocal;
-			TempStorage.RetrieveFromTempStorage(CmdEnv, "Test", out WasLocal);
-			if (!TempStorage.LocalTempStorageExists(CmdEnv, "Test"))
-			{
-				throw new AutomationException("local storage should exist");
-			}
-			if (!TempStorage.SharedTempStorageExists(CmdEnv, "Test"))
-			{
-				throw new AutomationException("shared storage should exist");
-			}
-			{
-				string[] LinesRead = ReadAllLines(UnitTestFile);
-				if (LinesRead == null || LinesRead.Length != LinesToWrite.Length)
-				{
-					throw new AutomationException("Contents of the file created is different to the file read.");
-				}
-				for (int LineIndex = 0; LineIndex < LinesRead.Length; ++LineIndex)
-				{
-					if (LinesRead[LineIndex] != LinesToWrite[LineIndex])
-					{
-						throw new AutomationException("Contents of the file created is different to the file read.");
-					}
-				}
-
-				string[] LinesRead2 = ReadAllLines(UnitTestFile2);
-				if (LinesRead2 == null || LinesRead2.Length != LinesToWrite2.Length)
-				{
-					throw new AutomationException("Contents of the file created is different to the file read.");
-				}
-				for (int LineIndex = 0; LineIndex < LinesRead2.Length; ++LineIndex)
-				{
-					if (LinesRead2[LineIndex] != LinesToWrite2[LineIndex])
-					{
-						throw new AutomationException("Contents of the file created is different to the file read.");
-					}
-				}
-
-				string[] LinesRead3 = ReadAllLines(UnitTestFile3);
-				if (LinesRead3 == null || LinesRead3.Length != LinesToWrite3.Length)
-				{
-					throw new AutomationException("Contents of the file created is different to the file read.");
-				}
-				for (int LineIndex = 0; LineIndex < LinesRead3.Length; ++LineIndex)
-				{
-					if (LinesRead3[LineIndex] != LinesToWrite3[LineIndex])
-					{
-						throw new AutomationException("Contents of the file created is different to the file read.");
-					}
-				}
-			}
-			TempStorage.DeleteSharedTempStorageManifests(CmdEnv, "Test");
-			if (TempStorage.SharedTempStorageExists(CmdEnv, "Test"))
-			{
-				throw new AutomationException("shared storage should not exist");
-			}
-			TempStorage.RetrieveFromTempStorage(CmdEnv, "Test", out WasLocal); // this should just rely on the local
-			if (!WasLocal || !TempStorage.LocalTempStorageExists(CmdEnv, "Test"))
-			{
-				throw new AutomationException("local storage should exist");
-			}
-
-			// and now lets test tampering
-			TempStorage.DeleteLocalTempStorageManifests(CmdEnv);
-			{
-				bool bFailedProperly = false;
-				var MissingFile = new List<string>(TestFiles);
-				MissingFile.Add(CombinePaths(CmdEnv.LocalRoot, "Engine", "SomeFileThatDoesntExist.txt"));
-				try
-				{
-					TempStorage.StoreToTempStorage(CmdEnv, "Test", MissingFile);
-				}
-				catch (AutomationException)
-				{
-					bFailedProperly = true;
-				}
-				if (!bFailedProperly)
-				{
-					throw new AutomationException("Missing file did not fail.");
-				}
-			}
-			TempStorage.DeleteSharedTempStorageManifests(CmdEnv, "Test"); // this ends up being junk
-			TempStorage.StoreToTempStorage(CmdEnv, "Test", TestFiles);
-			TempStorage.DeleteLocalTempStorageManifests(CmdEnv); // force a load from shared
-			var SharedFile = CombinePaths(TempStorage.SharedTempStorageDirectory("Test"), "Engine", "Build", "Batchfiles", "TestFile.Txt");
-			if (!FileExists_NoExceptions(SharedFile))
-			{
-				throw new AutomationException("Shared file {0} did not exist", SharedFile);
-			}
-			DeleteFile(SharedFile);
-			{
-				bool bFailedProperly = false;
-				try
-				{
-					TempStorage.RetrieveFromTempStorage(CmdEnv, "Test", out WasLocal);
-				}
-				catch (AutomationException)
-				{
-					bFailedProperly = true;
-				}
-				if (!bFailedProperly)
-				{
-					throw new AutomationException("Did not fail to load from missing file.");
-				}
-			}
-			TempStorage.DeleteSharedTempStorageManifests(CmdEnv, "Test");
-			TempStorage.StoreToTempStorage(CmdEnv, "Test", TestFiles);
-			DeleteFile(UnitTestFile);
-			{
-				bool bFailedProperly = false;
-				try
-				{
-					TempStorage.RetrieveFromTempStorage(CmdEnv, "Test", out WasLocal);
-				}
-				catch (AutomationException)
-				{
-					bFailedProperly = true;
-				}
-				if (!bFailedProperly)
-				{
-					throw new AutomationException("Did not fail to load from missing local file.");
-				}
-			}
-			TempStorage.DeleteSharedTempStorageManifests(CmdEnv, "Test");
-			TempStorage.DeleteLocalTempStorageManifests(CmdEnv);
-			DeleteFile(UnitTestFile);
-			DeleteFile(UnitTestFile2);
-			DeleteFile(UnitTestFile3);
+            TempStorage.TestTempStorage(CmdEnv);
+
 		}
 	}
 }
diff --git a/Engine/Source/Programs/AutomationTool/Program.cs b/Engine/Source/Programs/AutomationTool/Program.cs
--- a/Engine/Source/Programs/AutomationTool/Program.cs
+++ b/Engine/Source/Programs/AutomationTool/Program.cs
@@ -54,8 +54,7 @@ namespace AutomationTool
             catch (Exception Ex)
             {
                 // Catch all exceptions and propagate the ErrorCode if we are given one.
-                Log.TraceError("AutomationTool terminated with exception:");
-                Log.TraceError(LogUtils.FormatException(Ex));
+                Log.TraceError("AutomationTool terminated with exception: {0}", Ex);
                 // set the exit code of the process
                 if (Ex is AutomationException)
                 {
diff --git a/Engine/Source/Programs/AutomationTool/Scripts/BuildPluginCommand.Automation.cs b/Engine/Source/Programs/AutomationTool/Scripts/BuildPluginCommand.Automation.cs
--- a/Engine/Source/Programs/AutomationTool/Scripts/BuildPluginCommand.Automation.cs
+++ b/Engine/Source/Programs/AutomationTool/Scripts/BuildPluginCommand.Automation.cs
@@ -169,6 +169,6 @@ class BuildPlugin : BuildCommand
 		Filter.ExcludeConfidentialPlatforms();
 
 		// Apply the filter to the plugin directory
-		return new List<string>(Filter.ApplyToDirectory(PluginDirectory, true));
+		return Filter.ApplyToDirectory(PluginDirectory, true);
 	}
 }
diff --git a/Engine/Source/Programs/AutomationTool/Scripts/RocketBuild.Automation.cs b/Engine/Source/Programs/AutomationTool/Scripts/RocketBuild.Automation.cs
--- a/Engine/Source/Programs/AutomationTool/Scripts/RocketBuild.Automation.cs
+++ b/Engine/Source/Programs/AutomationTool/Scripts/RocketBuild.Automation.cs
@@ -141,7 +141,7 @@ namespace Rocket
 				string PublishedEngineDir;
 				if (ShouldDoSeriousThingsLikeP4CheckinAndPostToMCP(BranchConfig))
 				{
-					PublishedEngineDir = CommandUtils.CombinePaths(CommandUtils.RootSharedTempStorageDirectory(), "Rocket", "Automated", GetBuildLabel(), CommandUtils.GetGenericPlatformName(HostPlatform));
+					PublishedEngineDir = CommandUtils.CombinePaths(CommandUtils.RootBuildStorageDirectory(), "Rocket", "Automated", GetBuildLabel(), CommandUtils.GetGenericPlatformName(HostPlatform));
 				}
 				else
 				{
@@ -257,7 +257,7 @@ namespace Rocket
 
 		public static bool ShouldDoSeriousThingsLikeP4CheckinAndPostToMCP(GUBP.GUBPBranchConfig BranchConfig)
 		{
-			return CommandUtils.P4Enabled && CommandUtils.AllowSubmit && !BranchConfig.bPreflightBuild; // we don't do serious things in a preflight
+			return CommandUtils.P4Enabled && CommandUtils.AllowSubmit && !BranchConfig.JobInfo.IsPreflight; // we don't do serious things in a preflight
 		}
 
 		public static UnrealTargetPlatform GetSourceHostPlatform(List<UnrealTargetPlatform> HostPlatforms, UnrealTargetPlatform HostPlatform, UnrealTargetPlatform TargetPlatform)
@@ -328,8 +328,7 @@ namespace Rocket
 			// Copy everything that matches the filter to the promotion folder
 			string PromotableFolder = CommandUtils.CombinePaths(CommandUtils.CmdEnv.LocalRoot, "Engine", "Saved", "GitPromotable");
 			CommandUtils.DeleteDirectoryContents(PromotableFolder);
-			string[] PromotableFiles = CommandUtils.ThreadedCopyFiles(CommandUtils.CmdEnv.LocalRoot, PromotableFolder, PromotableFilter, bIgnoreSymlinks: true);
-			BuildProducts = new List<string>(PromotableFiles);
+			BuildProducts = CommandUtils.ThreadedCopyFiles(CommandUtils.CmdEnv.LocalRoot, PromotableFolder, PromotableFilter, bIgnoreSymlinks: true);
 		}
 	}
 
@@ -918,7 +917,7 @@ namespace Rocket
 			}
 
 			// Write the filtered list of depot files to disk, removing any symlinks
-			List<string> DepotFiles = Filter.ApplyToDirectory(CommandUtils.CmdEnv.LocalRoot, true).ToList();
+			List<string> DepotFiles = Filter.ApplyToDirectory(CommandUtils.CmdEnv.LocalRoot, true);
 			WriteManifest(DepotManifestPath, DepotFiles);
 			BuildProducts.Add(DepotManifestPath);
 
@@ -1046,8 +1045,8 @@ namespace Rocket
 
 			// Create lists of source and target files
 			CommandUtils.LogConsole("Preparing file lists...");
-			string[] SourceFiles = Files.Select(x => CommandUtils.CombinePaths(InputDir, x)).ToArray();
-			string[] TargetFiles = Files.Select(x => CommandUtils.CombinePaths(OutputDir, x)).ToArray();
+			var SourceFiles = Files.Select(x => CommandUtils.CombinePaths(InputDir, x)).ToList();
+			var TargetFiles = Files.Select(x => CommandUtils.CombinePaths(OutputDir, x)).ToList();
 
 			// Copy everything
 			CommandUtils.ThreadedCopyFiles(SourceFiles, TargetFiles);
