SetUATLocation uses the EntryAssembly instead of the executing assembly.

diff --git a/Engine/Source/Programs/AutomationTool/AutomationUtils/BuildEnvironment.cs b/Engine/Source/Programs/AutomationTool/AutomationUtils/BuildEnvironment.cs
--- a/Engine/Source/Programs/AutomationTool/AutomationUtils/BuildEnvironment.cs
+++ b/Engine/Source/Programs/AutomationTool/AutomationUtils/BuildEnvironment.cs
@@ -77,7 +77,7 @@ namespace AutomationTool
 		{
 			if (String.IsNullOrEmpty(UATExe))
 			{
-				UATExe = Assembly.GetExecutingAssembly().GetOriginalLocation();
+				UATExe = Assembly.GetEntryAssembly().GetOriginalLocation();
 			}
 			if (!CommandUtils.FileExists_NoExceptions(UATExe))
 			{
