Add additional telemetry to TempStorage. Moved TelemetryStopwatch out of CommandUtils. Added DateTimeStopwatch to make timing things easier.

diff --git a/Engine/Source/Programs/AutomationTool/AutomationUtils/Automation.cs b/Engine/Source/Programs/AutomationTool/AutomationUtils/Automation.cs
--- a/Engine/Source/Programs/AutomationTool/AutomationUtils/Automation.cs
+++ b/Engine/Source/Programs/AutomationTool/AutomationUtils/Automation.cs
@@ -351,7 +351,7 @@ AutomationTool.exe [-verbose] [-compileonly] [-p4] Command0 [-Arg0 -Arg1 -Arg2 
 			// Compile scripts.
 			Log.TraceInformation("Compiling scripts.");
 			ScriptCompiler Compiler = new ScriptCompiler();
-			using(CommandUtils.TelemetryStopwatch ScriptCompileStopwatch = new CommandUtils.TelemetryStopwatch("ScriptCompile"))
+			using(TelemetryStopwatch ScriptCompileStopwatch = new TelemetryStopwatch("ScriptCompile"))
 			{
 				Compiler.FindAndCompileAllScripts(OutScriptsForProjectFileName, AdditionalScriptsFolders);
 			}
diff --git a/Engine/Source/Programs/AutomationTool/AutomationUtils/CommandUtils.cs b/Engine/Source/Programs/AutomationTool/AutomationUtils/CommandUtils.cs
--- a/Engine/Source/Programs/AutomationTool/AutomationUtils/CommandUtils.cs
+++ b/Engine/Source/Programs/AutomationTool/AutomationUtils/CommandUtils.cs
@@ -1828,63 +1828,9 @@ namespace AutomationTool
 				return UBTTargetPlatforms;
 			}
 		}
-		private static UnrealBuildTool.UnrealTargetPlatform[] UBTTargetPlatforms; 
+		private static UnrealBuildTool.UnrealTargetPlatform[] UBTTargetPlatforms;
 
-		public class TelemetryStopwatch : IDisposable
-		{
-			string Name;
-			DateTime StartTime;
-			bool bFinished;
-
-			public TelemetryStopwatch(string Format, params object[] Args)
-			{
-				Name = String.Format(Format, Args);
-				StartTime = DateTime.Now;
-			}
-
-			public void Cancel()
-			{
-				bFinished = true;
-			}
-
-            /// <summary>
-            /// Flushes the time to <see cref="CmdEnv.CSVFile"/> if we are the build machine and that environment variable is specified.
-            /// Call Finish manually with an alternate name to use that one instead. Useful for dynamically generated names that you can't specify at construction.
-            /// </summary>
-            /// <param name="AlternateName">Used in place of the Name specified during construction.</param>
-			public void Finish(string AlternateName = null)
-			{
-				if(!bFinished)
-				{
-                    if (!string.IsNullOrEmpty(AlternateName))
-                    {
-                        Name = AlternateName;
-                    }
-
-                    var OutputStr = String.Format("UAT,{0},{1},{2}" + Environment.NewLine, Name, StartTime, DateTime.Now);
-                    LogVerbose(OutputStr);
-                    if (IsBuildMachine && !String.IsNullOrEmpty(CmdEnv.CSVFile) && CmdEnv.CSVFile != "nul")
-                    {
-					    try
-					    {
-                            File.AppendAllText(CmdEnv.CSVFile, OutputStr);
-                        }
-					    catch (Exception Ex)
-					    {
-						    LogWarning("Could not append to csv file ({0}) : {1}", CmdEnv.CSVFile, Ex.ToString());
-					    }
-                    }
-				}
-                bFinished = true;
-			}
-
-			public void Dispose()
-			{
-				Finish();
-			}
-		}
-
-		#endregion
+	    #endregion
 
 		#region Properties
 
@@ -2109,7 +2055,84 @@ namespace AutomationTool
 		}
 	}
 
-	/// <summary>
+    /// <summary>
+    /// Timer class used for telemetry reporting.
+    /// </summary>
+    public class TelemetryStopwatch : IDisposable
+    {
+        string Name;
+        DateTime StartTime;
+        bool bFinished;
+
+        public TelemetryStopwatch(string Format, params object[] Args)
+        {
+            Name = String.Format(Format, Args);
+            StartTime = DateTime.Now;
+        }
+
+        public void Cancel()
+        {
+            bFinished = true;
+        }
+
+        /// <summary>
+        /// Flushes the time to <see cref="CmdEnv.CSVFile"/> if we are the build machine and that environment variable is specified.
+        /// Call Finish manually with an alternate name to use that one instead. Useful for dynamically generated names that you can't specify at construction.
+        /// </summary>
+        /// <param name="AlternateName">Used in place of the Name specified during construction.</param>
+        public void Finish(string AlternateName = null)
+        {
+            if(!bFinished)
+            {
+                if (!String.IsNullOrEmpty(AlternateName))
+                {
+                    Name = AlternateName;
+                }
+
+                var OutputStr = String.Format("UAT,{0},{1},{2}" + Environment.NewLine, Name, StartTime, DateTime.Now);
+                CommandUtils.LogVerbose(OutputStr);
+                if (CommandUtils.IsBuildMachine && !String.IsNullOrEmpty(CommandUtils.CmdEnv.CSVFile) && CommandUtils.CmdEnv.CSVFile != "nul")
+                {
+                    try
+                    {
+                        File.AppendAllText(CommandUtils.CmdEnv.CSVFile, OutputStr);
+                    }
+                    catch (Exception Ex)
+                    {
+                        CommandUtils.LogWarning("Could not append to csv file ({0}) : {1}", CommandUtils.CmdEnv.CSVFile, Ex.ToString());
+                    }
+                }
+            }
+            bFinished = true;
+        }
+
+        public void Dispose()
+        {
+            Finish();
+        }
+    }
+
+    /// <summary>
+    /// Stopwatch that uses DateTime.UtcNow for timing. Not hi-res, but also not subject to short time limitations of System.Diagnostics.Stopwatch.
+    /// </summary>
+    public class DateTimeStopwatch
+    {
+        public static DateTimeStopwatch Start()
+        {
+            return new DateTimeStopwatch();
+        }
+
+        /// <summary>
+        /// Hide public ctor.
+        /// </summary>
+        private DateTimeStopwatch() { }
+
+        readonly DateTime StartTime = DateTime.UtcNow;
+
+        public TimeSpan ElapsedTime { get { return DateTime.UtcNow - StartTime; } }
+    }
+
+    /// <summary>
 	/// Use with "using" syntax to push and pop directories in a convenient, exception-safe way
 	/// </summary>
 	public class PushedDirectory : IDisposable
diff --git a/Engine/Source/Programs/AutomationTool/AutomationUtils/ProjectUtils.cs b/Engine/Source/Programs/AutomationTool/AutomationUtils/ProjectUtils.cs
--- a/Engine/Source/Programs/AutomationTool/AutomationUtils/ProjectUtils.cs
+++ b/Engine/Source/Programs/AutomationTool/AutomationUtils/ProjectUtils.cs
@@ -720,7 +720,7 @@ namespace AutomationTool
             BaseEngineProject = new BranchUProject();
 
             var AllProjects = UnrealBuildTool.UProjectInfo.FilterGameProjects(false, null);
-			using(CommandUtils.TelemetryStopwatch SortProjectsStopwatch = new CommandUtils.TelemetryStopwatch("SortProjects"))
+			using(TelemetryStopwatch SortProjectsStopwatch = new TelemetryStopwatch("SortProjects"))
 			{
 				foreach (var InfoEntry in AllProjects)
 				{
@@ -745,7 +745,7 @@ namespace AutomationTool
                 throw new AutomationException("All branches must have the blank project /Samples/Sandbox/BlankProject");
             }*/
 
-			using(CommandUtils.TelemetryStopwatch ProjectDumpStopwatch = new CommandUtils.TelemetryStopwatch("Project Dump"))
+			using(TelemetryStopwatch ProjectDumpStopwatch = new TelemetryStopwatch("Project Dump"))
 			{
 				CommandUtils.LogVerbose("  Base Engine:");
 				BaseEngineProject.Dump(InHostPlatforms);
diff --git a/Engine/Source/Programs/AutomationTool/GUBP/ElectricCommander.cs b/Engine/Source/Programs/AutomationTool/GUBP/ElectricCommander.cs
--- a/Engine/Source/Programs/AutomationTool/GUBP/ElectricCommander.cs
+++ b/Engine/Source/Programs/AutomationTool/GUBP/ElectricCommander.cs
@@ -358,7 +358,7 @@ namespace AutomationTool
 			}
 
 			List<JobStep> Steps = new List<JobStep>();
-			using (CommandUtils.TelemetryStopwatch PerlOutputStopwatch = new CommandUtils.TelemetryStopwatch("PerlOutput"))
+			using (TelemetryStopwatch PerlOutputStopwatch = new TelemetryStopwatch("PerlOutput"))
 			{
 				string ParentPath = Command.ParseParamValue("ParentPath");
 
diff --git a/Engine/Source/Programs/AutomationTool/GUBP/TempStorage.cs b/Engine/Source/Programs/AutomationTool/GUBP/TempStorage.cs
--- a/Engine/Source/Programs/AutomationTool/GUBP/TempStorage.cs
+++ b/Engine/Source/Programs/AutomationTool/GUBP/TempStorage.cs
@@ -470,7 +470,7 @@ namespace AutomationTool
             // This is a hack to clean out the old temp storage in the old folder name.
             TopDirectory = TopDirectory.Replace(Path.DirectorySeparatorChar + TempStorageSubdirectoryName + Path.DirectorySeparatorChar, Path.DirectorySeparatorChar + "GUBP" + Path.DirectorySeparatorChar);
 
-            using (var TelemetryStopwatch = new CommandUtils.TelemetryStopwatch("CleanSharedTempStorageLegacy"))
+            using (var TelemetryStopwatch = new TelemetryStopwatch("CleanSharedTempStorageLegacy"))
             {
                 const int MaximumDaysToKeepLegacyTempStorage = 1;
                 var StartTimeDir = DateTime.UtcNow;
@@ -546,7 +546,7 @@ namespace AutomationTool
             int FoldersCleaned = 0;
             int FoldersManuallyCleaned = 0;
             int FoldersFailedCleaned = 0;
-            using (var TelemetryStopwatch = new CommandUtils.TelemetryStopwatch("CleanSharedTempStorage"))
+            using (var TelemetryStopwatch = new TelemetryStopwatch("CleanSharedTempStorage"))
             {
                 const double MaximumDaysToKeepTempStorage = 3;
                 var Now = DateTime.UtcNow;
@@ -895,7 +895,7 @@ namespace AutomationTool
         /// <returns>The created manifest instance (which has already been saved to disk).</returns>
         public static void StoreToTempStorage(TempStorageNodeInfo TempStorageNodeInfo, List<string> Files, bool bLocalOnly, string GameName, string RootDir)
         {
-            using (var TelemetryStopwatch = new CommandUtils.TelemetryStopwatch("StoreToTempStorage"))
+            using (var TelemetryStopwatch = new TelemetryStopwatch("StoreToTempStorage"))
             {
                 // use LocalRoot if one is not specified
                 if (string.IsNullOrEmpty(RootDir))
@@ -908,8 +908,6 @@ namespace AutomationTool
                 var LocalTotalSize = Local.GetTotalSize();
                 if (!bLocalOnly)
                 {
-                    var StartTime = DateTime.UtcNow;
-
                     var BlockPath = SharedTempStorageDirectory(TempStorageNodeInfo, GameName);
                     CommandUtils.LogConsole("Storing to {0}", BlockPath);
                     if (CommandUtils.DirectoryExists_NoExceptions(BlockPath))
@@ -934,11 +932,10 @@ namespace AutomationTool
                         }
                     }
 
-                    long ZipFileSize = 0L;
-                    float ZipCopyDuration = 0.0f;
                     if (bZipTempStorage)
                     {
-                        Ionic.Zip.ZipFile Zip = new Ionic.Zip.ZipFile
+                        var ZipTimer = DateTimeStopwatch.Start();
+                        var Zip = new Ionic.Zip.ZipFile
                         {
                             UseZip64WhenSaving = Ionic.Zip.Zip64Option.Always
                         };
@@ -951,18 +948,21 @@ namespace AutomationTool
                         var SharedManifest = SharedTempStorageManifestFilename(TempStorageNodeInfo, GameName);
                         var ZipFilename = Path.ChangeExtension(LocalManifest, "zip");
                         Zip.Save(ZipFilename);
-                        ZipFileSize = new FileInfo(ZipFilename).Length;
-                        var ZipCopyStartTime = DateTime.UtcNow;
+                        var ZipFileSize = new FileInfo(ZipFilename).Length;
+                        var ZipTimeMS = (long)ZipTimer.ElapsedTime.TotalMilliseconds;
+                        var CopyTimer = DateTimeStopwatch.Start();
                         InternalUtils.Robust_CopyFile(ZipFilename, CommandUtils.CombinePaths(BlockPath, Path.GetFileName(ZipFilename)));
-                        ZipCopyDuration = (float)(DateTime.UtcNow - ZipCopyStartTime).TotalSeconds;
+                        var CopyTimeMS = (long)CopyTimer.ElapsedTime.TotalMilliseconds;
                         CommandUtils.DeleteFile(ZipFilename);
                         // copy the local manifest to the shared location. We have to assume the zip is a good copy.
                         CommandUtils.CopyFile(LocalManifest, SharedManifest);
-                        TelemetryStopwatch.Finish(string.Format("StoreToTempStorage.{0}.{1}.{2}.Remote", Files.Count, LocalTotalSize, ZipFileSize));
+                        TelemetryStopwatch.Finish(string.Format("StoreToTempStorage.{0}.{1}.{2}.Remote.{3}.{4}.{5}", Files.Count, LocalTotalSize, ZipFileSize, CopyTimeMS, ZipTimeMS, TempStorageNodeInfo.NodeStorageName));
                     }
                     else
                     {
+                        var CopyTimer = DateTimeStopwatch.Start();
                         CommandUtils.ThreadedCopyFiles(Files, DestFiles, ThreadsToCopyWith());
+                        var CopyTimeMS = (long)CopyTimer.ElapsedTime.TotalMilliseconds;
 
                         // save the shared temp storage manifest file.
                         var Shared = SaveSharedTempStorageManifest(TempStorageNodeInfo, GameName, DestFiles);
@@ -974,24 +974,12 @@ namespace AutomationTool
                             CommandUtils.RenameFile_NoExceptions(SharedTempManifestFilename, SharedTempManifestFilename + ".broken");
                             throw new AutomationException("Shared and Local manifest mismatch. Wrote out {0} for inspection.", SharedTempManifestFilename + ".broken");
                         }
-                        TelemetryStopwatch.Finish(string.Format("StoreToTempStorage.{0}.{1}.{2}.Remote", Files.Count, LocalTotalSize, 0L));
-                    }
-
-                    float BuildDuration = (float)((DateTime.UtcNow - StartTime).TotalSeconds);
-                    if (LocalTotalSize > 0)
-                    {
-                        var MBSec = (LocalTotalSize / (1024.0f * 1024.0f)) / BuildDuration;
-                        CommandUtils.Log("Wrote to shared temp storage at {0} MB/s    {1}B {2}s", MBSec, LocalTotalSize, BuildDuration);
-                        if (bZipTempStorage)
-                        {
-                            var ZipMBSec = (ZipFileSize / (1024.0f * 1024.0f)) / ZipCopyDuration;
-                            CommandUtils.Log("Wrote zip to shared temp storage at {0} MB/s    {1}B {2}s", ZipMBSec, ZipFileSize, ZipCopyDuration);
-                        }
+                        TelemetryStopwatch.Finish(string.Format("StoreToTempStorage.{0}.{1}.{2}.Remote.{3}.{4}.{5}", Files.Count, LocalTotalSize, 0L, CopyTimeMS, 0L, TempStorageNodeInfo.NodeStorageName));
                     }
                 }
                 else
                 {
-                    TelemetryStopwatch.Finish(string.Format("StoreToTempStorage.{0}.{1}.{2}.Local", Files.Count, LocalTotalSize, 0L));
+                    TelemetryStopwatch.Finish(string.Format("StoreToTempStorage.{0}.{1}.{2}.Local.{3}.{4}.{5}", Files.Count, LocalTotalSize, 0L, 0L, 0L, TempStorageNodeInfo.NodeStorageName));
                 }
             }
         }
@@ -1008,7 +996,7 @@ namespace AutomationTool
         /// <returns>List of fully qualified paths to all the files that were retrieved. This is returned even if we skip the copy (set WasLocal = true) .</returns>
         public static List<string> RetrieveFromTempStorage(TempStorageNodeInfo TempStorageNodeInfo, out bool WasLocal, string GameName, string RootDir)
         {
-            using (var TelemetryStopwatch = new CommandUtils.TelemetryStopwatch("RetrieveFromTempStorage"))
+            using (var TelemetryStopwatch = new TelemetryStopwatch("RetrieveFromTempStorage"))
             {
                 // use LocalRoot if one is not specified
                 if (String.IsNullOrEmpty(RootDir))
@@ -1032,11 +1020,10 @@ namespace AutomationTool
                         throw new AutomationException("Local files in manifest {0} were tampered with.", LocalManifest);
                     }
                     WasLocal = true;
-                    TelemetryStopwatch.Finish(string.Format("RetrieveFromTempStorage.{0}.{1}.{2}.Local", Files.Count, Local.GetTotalSize(), 0));
+                    TelemetryStopwatch.Finish(string.Format("RetrieveFromTempStorage.{0}.{1}.{2}.Local.{3}.{4}.{5}", Files.Count, Local.GetTotalSize(), 0L, 0L, 0L, TempStorageNodeInfo.NodeStorageName));
                     return Files;
                 }
                 WasLocal = false;
-                var StartTime = DateTime.UtcNow;
 
                 // We couldn't find the node storage locally, so get it from the shared location.
                 var BlockPath = SharedTempStorageDirectory(TempStorageNodeInfo, GameName);
@@ -1057,29 +1044,32 @@ namespace AutomationTool
                 var DestFiles = SharedFiles.Select(Filename => CommandUtils.MakeRerootedFilePath(Filename, BlockPath, RootDir)).ToList();
 
                 long ZipFileSize = 0L;
-                float ZipCopyDuration = 0.0f;
+                long CopyTimeMS = 0L;
+                long UnzipTimeMS = 0L;
                 if (bZipTempStorage)
                 {
                     var LocalZipFilename = Path.ChangeExtension(LocalManifest, "zip");
                     var SharedZipFilename = Path.ChangeExtension(SharedManifest, "zip");
-                    var ZipCopyStartTime = DateTime.UtcNow;
                     if (!CommandUtils.FileExists(SharedZipFilename))
                     {
                         throw new AutomationException("Storage block zip file does not exist! {0}", SharedZipFilename);
                     }
+                    var CopyTimer = DateTimeStopwatch.Start();
                     InternalUtils.Robust_CopyFile(SharedZipFilename, LocalZipFilename);
-                    ZipCopyDuration = (float)(DateTime.UtcNow - ZipCopyStartTime).TotalSeconds;
+                    CopyTimeMS = (long)CopyTimer.ElapsedTime.TotalMilliseconds;
                     ZipFileSize = new FileInfo(LocalZipFilename).Length;
 
+                    var UnzipTimer = DateTimeStopwatch.Start();
                     using (Ionic.Zip.ZipFile Zip = new Ionic.Zip.ZipFile(LocalZipFilename))
                     {
                         Zip.ExtractAll(RootDir, Ionic.Zip.ExtractExistingFileAction.OverwriteSilently);
                     }
-
                     CommandUtils.DeleteFile(LocalZipFilename);
+                    UnzipTimeMS = (long)UnzipTimer.ElapsedTime.TotalMilliseconds;
                 }
                 else
                 {
+                    var CopyTimer = DateTimeStopwatch.Start();
                     // If the local file already exists, it will be overwritten.
                     foreach (var DestFile in DestFiles)
                     {
@@ -1092,6 +1082,7 @@ namespace AutomationTool
 
                     // Do the threaded copy to the local file system.
                     CommandUtils.ThreadedCopyFiles(SharedFiles, DestFiles, ThreadsToCopyWith());
+                    CopyTimeMS = (long)CopyTimer.ElapsedTime.TotalMilliseconds;
                 }
 
                 // Handle unix permissions/chmod issues.
@@ -1110,19 +1101,7 @@ namespace AutomationTool
                     CommandUtils.RenameFile_NoExceptions(LocalManifest, LocalManifest + ".broken");
                     throw new AutomationException("Shared and Local manifest mismatch.");
                 }
-                float BuildDuration = (float)((DateTime.UtcNow - StartTime).TotalSeconds);
-                var SharedTotalSize = Shared.GetTotalSize();
-                if (SharedTotalSize > 0)
-                {
-                    var MBSec = (Shared.GetTotalSize() / (1024.0f * 1024.0f)) / BuildDuration;
-                    CommandUtils.Log("Read from shared temp storage at {0} MB/s    {1}B {2}s", MBSec, SharedTotalSize, BuildDuration);
-                    if (bZipTempStorage)
-                    {
-                        var ZipMBSec = (ZipFileSize / (1024.0f * 1024.0f)) / ZipCopyDuration;
-                        CommandUtils.Log("Read zip from shared temp storage at {0} MB/s    {1}B {2}s", ZipMBSec, ZipFileSize, ZipCopyDuration);
-                    }
-                }
-                TelemetryStopwatch.Finish(string.Format("RetrieveFromTempStorage.{0}.{1}.{2}.Remote", DestFiles.Count, SharedTotalSize, ZipFileSize));
+                TelemetryStopwatch.Finish(string.Format("RetrieveFromTempStorage.{0}.{1}.{2}.Remote.{3}.{4}.{5}", DestFiles.Count, Shared.GetTotalSize(), ZipFileSize, CopyTimeMS, UnzipTimeMS, TempStorageNodeInfo.NodeStorageName));
                 return DestFiles;
             }
         }
