Set up stagedir even in skipstage so the run command can find the proper executable

diff --git a/Engine/Source/Programs/AutomationTool/Scripts/CopyBuildToStagingDirectory.Automation.cs b/Engine/Source/Programs/AutomationTool/Scripts/CopyBuildToStagingDirectory.Automation.cs
--- a/Engine/Source/Programs/AutomationTool/Scripts/CopyBuildToStagingDirectory.Automation.cs
+++ b/Engine/Source/Programs/AutomationTool/Scripts/CopyBuildToStagingDirectory.Automation.cs
@@ -1504,7 +1504,7 @@ public partial class Project : CommandUtils
 				}
             }
 
-			string StageDirectory = ((ShouldCreatePak(Params) || (Params.Stage && !Params.SkipStage)) || !String.IsNullOrEmpty(Params.StageDirectoryParam)) ? Params.BaseStageDirectory : "";
+			string StageDirectory = ((ShouldCreatePak(Params) || (Params.Stage)) || !String.IsNullOrEmpty(Params.StageDirectoryParam)) ? Params.BaseStageDirectory : "";
             string ArchiveDirectory = (Params.Archive || !String.IsNullOrEmpty(Params.ArchiveDirectoryParam)) ? Params.BaseArchiveDirectory : "";
             if (prefixArchiveDir && (StagePlatform == UnrealTargetPlatform.Win32 || StagePlatform == UnrealTargetPlatform.Win64))
             {
