Packaging HTML5 deploys double-click script for HTML5LaunchHelper to avoid using the command line.
Fixes UE-19637

diff --git a/Engine/Source/Programs/AutomationTool/HTML5/HTML5Platform.Automation.cs b/Engine/Source/Programs/AutomationTool/HTML5/HTML5Platform.Automation.cs
--- a/Engine/Source/Programs/AutomationTool/HTML5/HTML5Platform.Automation.cs
+++ b/Engine/Source/Programs/AutomationTool/HTML5/HTML5Platform.Automation.cs
@@ -142,6 +142,11 @@ public class HTML5Platform : Platform
 
 
 		GenerateFileFromTemplate(TemplateFile, OutputFile, Params.ShortProjectName, Params.ClientConfigsToBuild[0].ToString(), Params.StageCommandline, !Params.IsCodeBasedProject, HeapSize);
+		
+		string MacBashTemplateFile = Path.Combine(CombinePaths(CmdEnv.LocalRoot, "Engine"), "Build", "HTML5", "RunMacHTML5LaunchHelper.command.template");
+		string MacBashOutputFile = Path.Combine(PackagePath, "RunMacHTML5LaunchHelper.command");
+		string MonoPath = Path.Combine(CombinePaths(CmdEnv.LocalRoot, "Engine"), "Build", "BatchFiles", "Mac", "SetupMono.sh");
+		GenerateMacCommandFromTemplate(MacBashTemplateFile, MacBashOutputFile, MonoPath);
 
 		string JSDir = Path.Combine(CombinePaths(CmdEnv.LocalRoot, "Engine"), "Build", "HTML5");
 		string OutDir = PackagePath;
@@ -285,6 +290,44 @@ public class HTML5Platform : Platform
 		}
 	}
 
+	protected void GenerateMacCommandFromTemplate(string InTemplateFile, string InOutputFile, string InMonoPath)
+	{
+		StringBuilder outputContents = new StringBuilder();
+		using (StreamReader reader = new StreamReader(InTemplateFile))
+		{
+			string InMonoPathParent = Path.GetDirectoryName(InMonoPath);
+			string LineStr = null;
+			while (reader.Peek() != -1)
+			{
+				LineStr = reader.ReadLine();
+				if (LineStr.Contains("${unreal_mono_pkg_path}"))
+				{
+					LineStr = LineStr.Replace("${unreal_mono_pkg_path}", InMonoPath);
+				}
+				if (LineStr.Contains("${unreal_mono_pkg_path_base}"))
+				{
+					LineStr = LineStr.Replace("${unreal_mono_pkg_path_base}", InMonoPathParent);
+				}
+
+				outputContents.AppendLine(LineStr);
+			}
+		}
+
+		if (outputContents.Length > 0)
+		{
+			// Save the file
+			try
+			{
+				Directory.CreateDirectory(Path.GetDirectoryName(InOutputFile));
+				File.WriteAllText(InOutputFile, outputContents.ToString(), Encoding.UTF8);
+			}
+			catch (Exception)
+			{
+				// Unable to write to the project file.
+			}
+		}
+	}
+
 	public override void GetFilesToDeployOrStage(ProjectParams Params, DeploymentContext SC)
 	{
 	}
@@ -329,6 +372,7 @@ public class HTML5Platform : Platform
 		var LaunchHelperPath = CombinePaths(CmdEnv.LocalRoot, "Engine/Binaries/DotNET/");
 		SC.ArchiveFiles(LaunchHelperPath, "HTML5LaunchHelper.exe");
 		SC.ArchiveFiles(Path.Combine(CombinePaths(CmdEnv.LocalRoot, "Engine"), "Build", "HTML5"), "Readme.txt");
+		SC.ArchiveFiles(PackagePath, Path.GetFileName(Path.Combine(PackagePath, "RunMacHTML5LaunchHelper.command")));
 
 		if (HTMLPakAutomation.CanCreateMapPaks(Params))
 		{
