Separate Rocket builds from Git promotions. A separate trigger email will be sent for making Rocket builds from now on.

diff --git a/Engine/Source/Programs/AutomationTool/Scripts/RocketBuild.Automation.cs b/Engine/Source/Programs/AutomationTool/Scripts/RocketBuild.Automation.cs
--- a/Engine/Source/Programs/AutomationTool/Scripts/RocketBuild.Automation.cs
+++ b/Engine/Source/Programs/AutomationTool/Scripts/RocketBuild.Automation.cs
@@ -77,6 +77,12 @@ namespace Rocket
 		{
 			if (!BranchConfig.BranchOptions.bNoInstalledEngine)
 			{
+				// Add the aggregate for making a rocket build
+				if(!BranchConfig.HasNode(WaitToMakeRocketBuild.StaticGetFullName()))
+				{
+					BranchConfig.AddNode(new WaitToMakeRocketBuild(BranchConfig.HostPlatforms));
+				}
+
 				// Find all the target platforms for this host platform.
 				List<UnrealTargetPlatform> TargetPlatforms = GetTargetPlatforms(bp, HostPlatform);
 
@@ -121,11 +127,6 @@ namespace Rocket
 				string LocalOutputDir = CommandUtils.CombinePaths(CommandUtils.CmdEnv.LocalRoot, "LocalBuilds", "Rocket", CommandUtils.GetGenericPlatformName(HostPlatform));
 				BranchConfig.AddNode(new GatherRocketNode(BranchConfig, HostPlatform, TargetPlatforms, LocalOutputDir));
 
-				// Add the aggregate node for the entire install
-				GUBP.SharedAggregatePromotableNode PromotableNode = (GUBP.SharedAggregatePromotableNode)BranchConfig.FindAggregateNode(GUBP.SharedAggregatePromotableNode.StaticGetFullName());
-				PromotableNode.AddDependency(FilterRocketNode.StaticGetFullName(HostPlatform));
-				PromotableNode.AddDependency(BuildDerivedDataCacheNode.StaticGetFullName(HostPlatform));
-
 				// Add a node for GitHub promotions
 				if(HostPlatform == UnrealTargetPlatform.Win64)
 				{
@@ -133,7 +134,6 @@ namespace Rocket
 					if(CommandUtils.FileExists(CommandUtils.CombinePaths(CommandUtils.CmdEnv.LocalRoot, GitConfigRelativePath)))
 					{
 						BranchConfig.AddNode(new BuildGitPromotable(HostPlatform, BranchConfig.HostPlatforms, GitConfigRelativePath));
-						PromotableNode.AddDependency(BuildGitPromotable.StaticGetFullName(HostPlatform));
 					}
 				}
 
@@ -151,28 +151,6 @@ namespace Rocket
 				// Publish the install to the network
 				BranchConfig.AddNode(new PublishRocketNode(HostPlatform, LocalOutputDir, PublishedEngineDir));
 				BranchConfig.AddNode(new PublishRocketSymbolsNode(BranchConfig, HostPlatform, TargetPlatforms, PublishedEngineDir + "Symbols"));
-
-				// Add a dependency on this being published as part of the shared promotable being labeled
-				GUBP.SharedLabelPromotableSuccessNode LabelPromotableNode = (GUBP.SharedLabelPromotableSuccessNode)BranchConfig.FindAggregateNode(GUBP.SharedLabelPromotableSuccessNode.StaticGetFullName());
-				LabelPromotableNode.AddDependency(PublishRocketNode.StaticGetFullName(HostPlatform));
-				LabelPromotableNode.AddDependency(PublishRocketSymbolsNode.StaticGetFullName(HostPlatform));
-
-				// Add dependencies on a promotable to do these steps too
-				GUBP.WaitForSharedPromotionUserInput WaitForPromotionNode = (GUBP.WaitForSharedPromotionUserInput)BranchConfig.FindNode(GUBP.WaitForSharedPromotionUserInput.StaticGetFullName(true));
-				WaitForPromotionNode.AddDependency(PublishRocketNode.StaticGetFullName(HostPlatform));
-				WaitForPromotionNode.AddDependency(PublishRocketSymbolsNode.StaticGetFullName(HostPlatform));
-
-				// Push everything behind the promotion triggers if we're doing things on the build machines
-				if (ShouldDoSeriousThingsLikeP4CheckinAndPostToMCP(BranchConfig) || bp.ParseParam("WithRocketPromotable"))
-				{
-					string WaitForTrigger = GUBP.WaitForSharedPromotionUserInput.StaticGetFullName(false);
-
-					GatherRocketNode GatherRocket = (GatherRocketNode)BranchConfig.FindNode(GatherRocketNode.StaticGetFullName(HostPlatform));
-					GatherRocket.AddDependency(WaitForTrigger);
-
-					PublishRocketSymbolsNode PublishRocketSymbols = (PublishRocketSymbolsNode)BranchConfig.FindNode(PublishRocketSymbolsNode.StaticGetFullName(HostPlatform));
-					PublishRocketSymbols.AddDependency(WaitForTrigger);
-				}
 			}
 		}
 
@@ -291,6 +269,38 @@ namespace Rocket
 		}
 	}
 
+	public class WaitToMakeRocketBuild : GUBP.WaitForUserInput
+	{
+        public WaitToMakeRocketBuild(List<UnrealTargetPlatform> HostPlatforms)
+        {
+			foreach(UnrealTargetPlatform HostPlatform in HostPlatforms)
+			{
+				AddDependency(FilterRocketNode.StaticGetFullName(HostPlatform));
+				AddDependency(BuildDerivedDataCacheNode.StaticGetFullName(HostPlatform));
+			}
+		}
+
+        public static string StaticGetFullName()
+        {
+            return "WaitToMakeRocketBuild";
+        }
+
+        public override string GetFullName()
+        {
+            return StaticGetFullName();
+        }
+
+        public override string GetTriggerDescText()
+        {
+			return "Ready to make Rocket build";
+        }
+
+        public override string GetTriggerActionText()
+        {
+			return "Make Rocket build";
+        }
+    }
+
 	public class BuildGitPromotable : GUBP.HostPlatformNode
 	{
 		string ConfigRelativePath;
@@ -977,6 +987,7 @@ namespace Rocket
 
 			AddDependency(FilterRocketNode.StaticGetFullName(HostPlatform));
 			AddDependency(BuildDerivedDataCacheNode.StaticGetFullName(HostPlatform));
+			AddDependency(WaitToMakeRocketBuild.StaticGetFullName());
 
 			AgentSharingGroup = "RocketGroup" + StaticGetHostPlatformSuffix(HostPlatform);
 		}
@@ -1131,6 +1142,7 @@ namespace Rocket
 			AddDependency(GUBP.ToolsForCompileNode.StaticGetFullName(HostPlatform));
 			AddDependency(GUBP.RootEditorNode.StaticGetFullName(HostPlatform));
 			AddDependency(GUBP.ToolsNode.StaticGetFullName(HostPlatform));
+			AddDependency(WaitToMakeRocketBuild.StaticGetFullName());
 
 			foreach(UnrealTargetPlatform TargetPlatform in TargetPlatforms)
 			{
