Added support for diffing generated cooked content against supplied cooked pak file.

diff --git a/Engine/Source/Programs/AutomationTool/AutomationUtils/Platform.cs b/Engine/Source/Programs/AutomationTool/AutomationUtils/Platform.cs
--- a/Engine/Source/Programs/AutomationTool/AutomationUtils/Platform.cs
+++ b/Engine/Source/Programs/AutomationTool/AutomationUtils/Platform.cs
@@ -117,6 +117,16 @@ namespace AutomationTool
 			throw new AutomationException("{0} does not yet implement Packaging.", PlatformType);
 		}
 
+        /// <summary>
+        /// Does the reverse of the output from the package process
+        /// </summary>
+        /// <param name="SourcePath"></param>
+        /// <param name="DestinationPath"></param>
+        public virtual void ExtractPackage(ProjectParams Params, string SourcePath, string DestinationPath)
+        {
+            throw new AutomationException("{0} does not yet implement ExtractPackage.", PlatformType);
+        }
+
 		/// <summary>
 		/// Allow platform to do platform specific work on archived project before it's deployed.
 		/// </summary>
diff --git a/Engine/Source/Programs/AutomationTool/AutomationUtils/ProjectParams.cs b/Engine/Source/Programs/AutomationTool/AutomationUtils/ProjectParams.cs
--- a/Engine/Source/Programs/AutomationTool/AutomationUtils/ProjectParams.cs
+++ b/Engine/Source/Programs/AutomationTool/AutomationUtils/ProjectParams.cs
@@ -202,6 +202,7 @@ namespace AutomationTool
             this.GeneratePatch = InParams.GeneratePatch;
             this.DLCName = InParams.DLCName;
             this.DLCIncludeEngineContent = InParams.DLCIncludeEngineContent;
+            this.DiffCookedContentPath = InParams.DiffCookedContentPath;
             this.NewCook = InParams.NewCook;
             this.OldCook = InParams.OldCook;
             this.AdditionalCookerOptions = InParams.AdditionalCookerOptions;
@@ -355,6 +356,7 @@ namespace AutomationTool
             string CreateReleaseVersion = null,
             bool? GeneratePatch = null,
             string DLCName = null,
+            string DiffCookedContentPath = null,
             bool? DLCIncludeEngineContent = null,
             bool? NewCook = null,
             bool? OldCook = null,
@@ -463,6 +465,7 @@ namespace AutomationTool
             this.GeneratePatch = GetParamValueIfNotSpecified(Command, GeneratePatch, this.GeneratePatch, "GeneratePatch");
             this.AdditionalCookerOptions = ParseParamValueIfNotSpecified(Command, AdditionalCookerOptions, "AdditionalCookerOptions", String.Empty);
             this.DLCName = ParseParamValueIfNotSpecified(Command, DLCName, "DLCName", String.Empty);
+            this.DiffCookedContentPath = ParseParamValueIfNotSpecified(Command, DiffCookedContentPath, "DiffCookedContentPath", String.Empty);
             this.DLCIncludeEngineContent = GetParamValueIfNotSpecified(Command, DLCIncludeEngineContent, this.DLCIncludeEngineContent, "DLCIncludeEngineContent");
 			this.SkipCook = GetParamValueIfNotSpecified(Command, SkipCook, this.SkipCook, "skipcook");
 			if (this.SkipCook)
@@ -1078,6 +1081,12 @@ namespace AutomationTool
         /// </summary>
         public bool DLCIncludeEngineContent;
 
+        /// <summary>
+        /// After cook completes diff the cooked content against another cooked content directory.
+        ///  report all errors to the log
+        /// </summary>
+        public string DiffCookedContentPath;
+
         /// <summary>
         /// Cook: Additional cooker options to include on the cooker commandline
         /// </summary>
@@ -1728,6 +1737,11 @@ namespace AutomationTool
             get { return !String.IsNullOrEmpty(DLCName); }
         }
 
+        public bool HasDiffCookedContentPath
+        {
+            get { return !String.IsNullOrEmpty(DiffCookedContentPath); }
+        }
+
         public bool HasCreateReleaseVersion
         {
             get { return !String.IsNullOrEmpty(CreateReleaseVersion); }
@@ -2067,6 +2081,7 @@ namespace AutomationTool
                 CommandUtils.Log("BasedOnReleaseVersion={0}", BasedOnReleaseVersion);
                 CommandUtils.Log("DLCName={0}", DLCName);
                 CommandUtils.Log("DLCIncludeEngineContent={0}", DLCIncludeEngineContent);
+                CommandUtils.Log("DiffCookedContentPath={0}", DiffCookedContentPath);
                 CommandUtils.Log("AdditionalCookerOptions={0}", AdditionalCookerOptions);
 				CommandUtils.Log("DedicatedServer={0}", DedicatedServer);
 				CommandUtils.Log("DirectoriesToCook={0}", DirectoriesToCook.ToString());
diff --git a/Engine/Source/Programs/AutomationTool/Scripts/CookCommand.Automation.cs b/Engine/Source/Programs/AutomationTool/Scripts/CookCommand.Automation.cs
--- a/Engine/Source/Programs/AutomationTool/Scripts/CookCommand.Automation.cs
+++ b/Engine/Source/Programs/AutomationTool/Scripts/CookCommand.Automation.cs
@@ -246,11 +246,119 @@ public partial class Project : CommandUtils
 					throw new AutomationException(ErrorCodes.Error_UnknownCookFailure, Ex, "Cook failed.");
 				}
 			}
+
+            if (Params.HasDiffCookedContentPath)
+            {
+                try
+                {
+                    DiffCookedContent(Params);
+                }
+                catch ( Exception Ex )
+                {
+                    // Delete cooked data (if any) as it may be incomplete / corrupted.
+                    LogConsole("Cook failed. Deleting cooked data.");
+                    CleanupCookedData(PlatformsToCook.ToList(), Params);
+                    throw new AutomationException(ErrorCodes.Error_UnknownCookFailure, Ex, "Cook failed.");
+                }
+            }
+            
 		}
 
+
 		LogConsole("********** COOK COMMAND COMPLETED **********");
 	}
 
+    private static void DiffCookedContent( ProjectParams Params)
+    {
+        List<UnrealTargetPlatform> PlatformsToCook = Params.ClientTargetPlatforms;
+        string ProjectPath = Path.GetFullPath(Params.RawProjectPath);
+
+        var CookedSandboxesPath = CombinePaths(GetDirectoryName(ProjectPath), "Saved", "Cooked");
+
+        for (int CookPlatformIndex = 0; CookPlatformIndex < PlatformsToCook.Count; ++CookPlatformIndex)
+        {
+            // temporary directory to save the pak file to (pak file is usually not local and on network drive)
+            var TemporaryPakPath = CombinePaths(GetDirectoryName(ProjectPath), "Saved", "Temp", "LocalPKG");
+            // extracted files from pak file
+            var TemporaryFilesPath = CombinePaths(GetDirectoryName(ProjectPath), "Saved", "Temp", "LocalFiles");
+
+            
+
+            try
+            {
+                Directory.Delete(TemporaryPakPath, true);
+                Directory.Delete(TemporaryFilesPath, true);
+            }
+            catch(Exception )
+            {
+                LogConsole("Failed deleting temporary directories "+TemporaryPakPath+" "+TemporaryFilesPath+" continuing.");
+            }
+
+            Directory.CreateDirectory(TemporaryPakPath);
+            Directory.CreateDirectory(TemporaryFilesPath);
+
+            Platform CurrentPlatform = Params.GetTargetPlatformInstance(PlatformsToCook[CookPlatformIndex]);
+
+            if (Params.DiffCookedContentPath.EndsWith(".pak"))
+            {
+                string TemporaryPakFilename = CombinePaths( TemporaryPakPath, Path.GetFileName(Params.DiffCookedContentPath) );
+                File.Copy(Params.DiffCookedContentPath, TemporaryPakFilename);
+            }
+            else
+            {
+                CurrentPlatform.ExtractPackage(Params, Params.DiffCookedContentPath, TemporaryPakPath);
+            }
+
+            
+
+            string FullCookPath = CombinePaths(CookedSandboxesPath, PlatformsToCook[CookPlatformIndex].ToString());
+
+            
+
+            var UnrealPakExe = CombinePaths(CmdEnv.LocalRoot, "Engine/Binaries/Win64/UnrealPak.exe");
+
+            // find the pak file
+            var PakFiles = Directory.EnumerateFiles(TemporaryPakPath, "*.pak");
+            foreach (var Name in PakFiles)
+            {
+                string UnrealPakParams = Name + " -Extract " + " " + TemporaryFilesPath;
+
+                RunAndLog(CmdEnv, UnrealPakExe, UnrealPakParams, Options: ERunOptions.Default | ERunOptions.UTF8Output);
+            }
+
+            // diff the content
+            List<string> AllFiles = Directory.EnumerateFiles(FullCookPath, "*.uasset", System.IO.SearchOption.AllDirectories).ToList();
+            AllFiles.AddRange(Directory.EnumerateFiles(FullCookPath, "*.map", System.IO.SearchOption.AllDirectories).ToList());
+            foreach (string SourceFilename in AllFiles)
+            {
+                // Filename.StartsWith( CookedSandboxesPath );
+                string RelativeFilename = SourceFilename.Remove(0, FullCookPath.Length);
+                
+                string DestFilename = TemporaryFilesPath + RelativeFilename;
+
+                byte[] SourceFile = File.ReadAllBytes(SourceFilename);
+
+                byte[] DestFile = File.ReadAllBytes(DestFilename);
+
+                if ( SourceFile.LongLength == DestFile.LongLength )
+                {
+                    for ( long Index = 0; Index < SourceFile.LongLength; ++Index )
+                    {
+                        if ( SourceFile[Index] != DestFile[Index] )
+                        {
+					        LogConsole("Diff cooked content failed on file " +SourceFilename + " when comparing against "+DestFilename + " at offset " + Index.ToString() );
+                            break;
+                        }
+                    }
+                }
+                else
+                {
+                    LogConsole("Diff cooked content failed on file " +SourceFilename + " when comparing against "+DestFilename + " files are different sizes " + SourceFile.LongLength.ToString() + " " + DestFile.LongLength.ToString() );
+                }
+            } 
+        }
+    }
+
 	private static void CleanupCookedData(List<string> PlatformsToCook, ProjectParams Params)
 	{
 		var ProjectPath = Path.GetFullPath(Params.RawProjectPath);
