Include the Templates/TemplateResources folder when filtering the engine to use for generating DDC.

diff --git a/Engine/Source/Programs/AutomationTool/Scripts/RocketBuild.Automation.cs b/Engine/Source/Programs/AutomationTool/Scripts/RocketBuild.Automation.cs
--- a/Engine/Source/Programs/AutomationTool/Scripts/RocketBuild.Automation.cs
+++ b/Engine/Source/Programs/AutomationTool/Scripts/RocketBuild.Automation.cs
@@ -1254,6 +1254,7 @@ namespace Rocket
 				Filter.Exclude("/Engine/DerivedDataCache/...");
 				Filter.Exclude("/Samples/...");
 				Filter.Exclude("/Templates/...");
+				Filter.Include("/Templates/TemplateResources/...");
 				Filter.Exclude(".../Source/...");
 				Filter.Exclude(".../Intermediate/...");
 				Filter.ExcludeConfidentialPlatforms();
