UEB-327 - Signing process should fail over to alternate timestamp server when we receive "The specified timestamp server either could not be reached or returned an invalid response"

diff --git a/Engine/Source/Programs/AutomationTool/AutomationUtils/CommandUtils.cs b/Engine/Source/Programs/AutomationTool/AutomationUtils/CommandUtils.cs
--- a/Engine/Source/Programs/AutomationTool/AutomationUtils/CommandUtils.cs
+++ b/Engine/Source/Programs/AutomationTool/AutomationUtils/CommandUtils.cs
@@ -2281,20 +2281,25 @@ namespace AutomationTool
 			TargetFileInfo.IsReadOnly = false;
 
 			// Code sign the executable
-			string TimestampServer = "http://timestamp.verisign.com/scripts/timestamp.dll";
+			string[] TimestampServer = { "http://timestamp.verisign.com/scripts/timestamp.dll",
+									     "http://timestamp.globalsign.com/scripts/timstamp.dll",
+										 "http://timestamp.comodoca.com/authenticode",
+										 "http://www.startssl.com/timestamp"
+									   };
+			int TimestampServerIndex = 0;
 
 			string SpecificStoreArg = bUseMachineStoreInsteadOfUserStore ? " /sm" : "";
 
-			//@TODO: Verbosity choosing
-			//  /v will spew lots of info
-			//  /q does nothing on success and minimal output on failure
-			string CodeSignArgs = String.Format("sign{0} /a /n \"{1}\" /t {2} /d \"{3}\" /v \"{4}\"", SpecificStoreArg, SigningIdentity, TimestampServer, TargetFileInfo.Name, TargetFileInfo.FullName);
-
 			DateTime StartTime = DateTime.Now;
 
 			int NumTrials = 0;
 			for (; ; )
 			{
+				//@TODO: Verbosity choosing
+				//  /v will spew lots of info
+				//  /q does nothing on success and minimal output on failure
+				string CodeSignArgs = String.Format("sign{0} /a /n \"{1}\" /t {2} /d \"{3}\" /v \"{4}\"", SpecificStoreArg, SigningIdentity, TimestampServer[TimestampServerIndex], TargetFileInfo.Name, TargetFileInfo.FullName);
+
 				ProcessResult Result = CommandUtils.Run(SignToolName, CodeSignArgs, null, CommandUtils.ERunOptions.AllowSpew);
 				++NumTrials;
 
@@ -2309,6 +2314,14 @@ namespace AutomationTool
 				}
 				else
 				{
+					// try another timestamp server on the next iteration
+					TimestampServerIndex++;
+					if (TimestampServerIndex >= TimestampServer.Count())
+					{
+						// loop back to the first timestamp server
+						TimestampServerIndex = 0;
+					}
+
 					// Keep retrying until we run out of time
 					TimeSpan RunTime = DateTime.Now - StartTime;
 					if (RunTime > CodeSignTimeOut)
@@ -2444,19 +2457,26 @@ namespace AutomationTool
 			}
 
 			// Code sign the executable
-			string TimestampServer = "http://timestamp.verisign.com/scripts/timestamp.dll";
+			// Code sign the executable
+			string[] TimestampServer = { "http://timestamp.verisign.com/scripts/timestamp.dll",
+									     "http://timestamp.globalsign.com/scripts/timstamp.dll",
+										 "http://timestamp.comodoca.com/authenticode",
+										 "http://www.startssl.com/timestamp"
+									   };
+			int TimestampServerIndex = 0;
 
 			string SpecificStoreArg = bUseMachineStoreInsteadOfUserStore ? " /sm" : "";	
-			//@TODO: Verbosity choosing
-			//  /v will spew lots of info
-			//  /q does nothing on success and minimal output on failure
-			string CodeSignArgs = String.Format("sign{0} /a /n \"{1}\" /t {2} /v {3}", SpecificStoreArg, SigningIdentity, TimestampServer, FilesToSign);
-
+			
 			DateTime StartTime = DateTime.Now;
 
 			int NumTrials = 0;
 			for (; ; )
 			{
+				//@TODO: Verbosity choosing
+				//  /v will spew lots of info
+				//  /q does nothing on success and minimal output on failure
+				string CodeSignArgs = String.Format("sign{0} /a /n \"{1}\" /t {2} /v {3}", SpecificStoreArg, SigningIdentity, TimestampServer[TimestampServerIndex], FilesToSign);
+
 				ProcessResult Result = CommandUtils.Run(SignToolName, CodeSignArgs, null, CommandUtils.ERunOptions.AllowSpew);
 				++NumTrials;
 
@@ -2471,6 +2491,14 @@ namespace AutomationTool
 				}
 				else
 				{
+					// try another timestamp server on the next iteration
+					TimestampServerIndex++;
+					if (TimestampServerIndex >= TimestampServer.Count())
+					{
+						// loop back to the first timestamp server
+						TimestampServerIndex = 0;
+					}
+					
 					// Keep retrying until we run out of time
 					TimeSpan RunTime = DateTime.Now - StartTime;
 					if (RunTime > CodeSignTimeOut)
