Remove promotion build nodes, aggregate nodes, and triggers.

diff --git a/Engine/Source/Programs/AutomationTool/GUBP/LegacyBranchSetup.cs b/Engine/Source/Programs/AutomationTool/GUBP/LegacyBranchSetup.cs
--- a/Engine/Source/Programs/AutomationTool/GUBP/LegacyBranchSetup.cs
+++ b/Engine/Source/Programs/AutomationTool/GUBP/LegacyBranchSetup.cs
@@ -1117,57 +1117,8 @@ partial class GUBP
 			}
 		}
 
-        int NumSharedAllHosts = 0;
-        foreach (var CodeProj in Branch.CodeProjects)
-        {
-            if (CodeProj.Properties.Targets.ContainsKey(TargetRules.TargetType.Editor))
-            {
-                bool AnySeparate = false;
-                var PromotedHosts = new List<UnrealTargetPlatform>();
-
-                foreach (var HostPlatform in HostPlatforms)
-                {
-                    if (!BranchOptions.ExcludePlatformsForEditor.Contains(HostPlatform) && !BranchOptions.RemovePlatformFromPromotable.Contains(HostPlatform))
-					{
-						var Options = CodeProj.Options(HostPlatform);
-						AnySeparate = AnySeparate || Options.bSeparateGamePromotion;
-						if (Options.bIsPromotable)
-						{
-							if (!Options.bSeparateGamePromotion)
-							{
-								NumSharedAllHosts++;
-							}
-							PromotedHosts.Add(HostPlatform);
-						}
-					}
-                }
-                if (PromotedHosts.Count > 0)
-                {
-					BranchConfig.AddNode(new GameAggregatePromotableNode(BranchConfig, PromotedHosts, ActivePlatforms, CodeProj, true, bNoIOSOnPC));
-                    if (AnySeparate)
-                    {
-                        BranchConfig.AddNode(new WaitForGamePromotionUserInput(this, CodeProj, false));
-						BranchConfig.AddNode(new GameLabelPromotableNode(BranchConfig, CodeProj, false));
-                        BranchConfig.AddNode(new WaitForGamePromotionUserInput(this, CodeProj, true));
-						BranchConfig.AddNode(new GameLabelPromotableNode(BranchConfig, CodeProj, true));
-                    }
-                }
-            }
-        }
-        if (NumSharedAllHosts > 0)
-        {
-			BranchConfig.AddNode(new GameAggregatePromotableNode(BranchConfig, HostPlatforms, ActivePlatforms, Branch.BaseEngineProject, false, bNoIOSOnPC));
-
-            BranchConfig.AddNode(new SharedAggregatePromotableNode(BranchConfig, ActivePlatforms));
-            BranchConfig.AddNode(new WaitForSharedPromotionUserInput(this, false));
-            BranchConfig.AddNode(new SharedLabelPromotableNode(BranchConfig, false));
-            BranchConfig.AddNode(new SharedLabelPromotableSuccessNode());
-            BranchConfig.AddNode(new WaitForTestShared(this));
-            BranchConfig.AddNode(new WaitForSharedPromotionUserInput(this, true));
-			BranchConfig.AddNode(new SharedLabelPromotableNode(BranchConfig, true));		
-
-			BranchConfig.AddNode(new WaitToPackageSamplesNode());
-        }
+        BranchConfig.AddNode(new WaitForTestShared(this));
+		BranchConfig.AddNode(new WaitToPackageSamplesNode(BranchConfig.HostPlatforms));
 
 		AddCustomNodes(BranchConfig, HostPlatforms, ActivePlatforms);
         
@@ -1218,10 +1169,6 @@ partial class GUBP
 		foreach (BranchInfo.BranchUProject GameProj in Branch.AllProjects)
 		{
 			List<string> NodeNames = new List<string>();
-			if (GameProj.Options(UnrealTargetPlatform.Win64).bIsPromotable)
-			{
-				NodeNames.Add(GameAggregatePromotableNode.StaticGetFullName(GameProj));
-			}
 			foreach(GUBP.GUBPNode Node in BranchConfig.GUBPNodes.Values)
 			{
 				if (Node.GameNameIfAnyForTempStorage() == GameProj.GameName)
diff --git a/Engine/Source/Programs/AutomationTool/GUBP/LegacyNodes.cs b/Engine/Source/Programs/AutomationTool/GUBP/LegacyNodes.cs
--- a/Engine/Source/Programs/AutomationTool/GUBP/LegacyNodes.cs
+++ b/Engine/Source/Programs/AutomationTool/GUBP/LegacyNodes.cs
@@ -1592,151 +1592,6 @@ partial class GUBP
         }
     }
 
-
-    public class AggregatePromotableNode : GUBPAggregateNode
-    {
-        protected List<UnrealTargetPlatform> HostPlatforms;
-        string PromotionLabelPrefix;		
-
-        public AggregatePromotableNode(List<UnrealTargetPlatform> InHostPlatforms, string InPromotionLabelPrefix)
-        {
-            HostPlatforms = InHostPlatforms;            
-            PromotionLabelPrefix = InPromotionLabelPrefix; 
-        }
-        public static string StaticGetFullName(string InPromotionLabelPrefix)
-        {
-            return InPromotionLabelPrefix + "_Promotable_Aggregate";
-        }
-		public override bool IsPromotableAggregate()
-		{
-			return true;
-		}
-        public override string GetFullName()
-        {
-            return StaticGetFullName(PromotionLabelPrefix);
-        }
-    }
-									
-    public class GameAggregatePromotableNode : AggregatePromotableNode
-    {
-        BranchInfo.BranchUProject GameProj;		
-
-        public GameAggregatePromotableNode(GUBP.GUBPBranchConfig BranchConfig, List<UnrealTargetPlatform> InHostPlatforms, List<UnrealTargetPlatform> InActivePlatforms, BranchInfo.BranchUProject InGameProj, bool IsSeparate, bool bNoIOSOnPC)
-            : base(InHostPlatforms, InGameProj.GameName)
-        {
-            GameProj = InGameProj;
-			
-            foreach (var HostPlatform in HostPlatforms)
-            {
-                AddDependency(RootEditorNode.StaticGetFullName(HostPlatform));
-			    if(!BranchConfig.BranchOptions.PromotablesWithoutTools.Contains(GameProj.GameName))
-                {
-                    AddDependency(ToolsNode.StaticGetFullName(HostPlatform));
-                    AddDependency(InternalToolsNode.StaticGetFullName(HostPlatform));
-                }
-                if (InGameProj.GameName != BranchConfig.Branch.BaseEngineProject.GameName && GameProj.Properties.Targets.ContainsKey(TargetRules.TargetType.Editor))
-                {
-                    AddDependency(EditorGameNode.StaticGetFullName(HostPlatform, GameProj));
-                }				
-                // add all of the platforms I use
-                {
-                    if (!GameProj.Options(HostPlatform).bPromoteEditorOnly)
-                    {
-						var Platforms = GamePlatformMonolithicsNode.GetMonolithicPlatformsForUProject(HostPlatform, InActivePlatforms, InGameProj, true, bNoIOSOnPC);
-						foreach (var Plat in Platforms)
-						{
-							AddDependency(GamePlatformMonolithicsNode.StaticGetFullName(HostPlatform, GameProj, Plat));
-								if (Plat == UnrealTargetPlatform.Win32 && GameProj.Properties.Targets.ContainsKey(TargetRules.TargetType.Game))
-							{
-									if (GameProj.Properties.Targets[TargetRules.TargetType.Game].Rules.GUBP_BuildWindowsXPMonolithics())
-								{
-									AddDependency(GamePlatformMonolithicsNode.StaticGetFullName(HostPlatform, GameProj, Plat, true));
-								}
-							}
-						}
-					}
-				}
-			}
-        }
-
-        public static string StaticGetFullName(BranchInfo.BranchUProject InGameProj)
-        {
-            return AggregatePromotableNode.StaticGetFullName(InGameProj.GameName);
-        }
-
-        public override string GameNameIfAnyForTempStorage()
-        {
-            return GameProj.GameName;
-        }
-		public override bool IsSeparatePromotable()
-		{
-			bool IsSeparate = false;
-			foreach(UnrealTargetPlatform HostPlatform in HostPlatforms)
-			{
-				if(GameProj.Options(HostPlatform).bSeparateGamePromotion)
-				{
-					IsSeparate = true;
-				}
-			}
-			return IsSeparate;
-		}
-    }
-
-    public class SharedAggregatePromotableNode : AggregatePromotableNode
-    {
-
-        public SharedAggregatePromotableNode(GUBP.GUBPBranchConfig BranchConfig, List<UnrealTargetPlatform> InActivePlatforms)
-            : base(BranchConfig.HostPlatforms, "Shared")
-        {
-            foreach (var HostPlatform in HostPlatforms)
-            {
-                AddDependency(EditorAndToolsNode.StaticGetFullName(HostPlatform));
-                {
-                    var Options = BranchConfig.Branch.BaseEngineProject.Options(HostPlatform);
-                    if (Options.bIsPromotable && !Options.bSeparateGamePromotion)
-                    {
-                        AddDependency(GameAggregatePromotableNode.StaticGetFullName(BranchConfig.Branch.BaseEngineProject));
-                    }
-                }
-                foreach (var CodeProj in BranchConfig.Branch.CodeProjects)
-                {
-                    var Options = CodeProj.Options(HostPlatform);
-                    if (!Options.bSeparateGamePromotion)
-                    {
-                        if (Options.bIsPromotable)
-                        {
-                            AddDependency(GameAggregatePromotableNode.StaticGetFullName(CodeProj));
-                        }
-                        else if (Options.bTestWithShared)
-                        {
-							if (!Options.bIsNonCode)
-							{
-                                AddDependency(EditorGameNode.StaticGetFullName(HostPlatform, CodeProj)); // if we are just testing, we will still include the editor stuff
-                            }
-                        }
-                    }
-                }
-                if(HostPlatform == UnrealTargetPlatform.Win64 && InActivePlatforms.Contains(UnrealTargetPlatform.Linux))
-                {
-                    AddDependency(RootEditorCrossCompileLinuxNode.StaticGetFullName());
-                    AddDependency(ToolsCrossCompileNode.StaticGetFullName(HostPlatform));
-                }
-            }
-			if (!BranchConfig.BranchOptions.bNoInstalledEngine)
-			{
-				AddDependency(MakeFeaturePacksNode.StaticGetFullName(MakeFeaturePacksNode.GetDefaultBuildPlatform(BranchConfig.HostPlatforms)));
-			}
-        }
-		public override bool IsSeparatePromotable()
-		{
-			return true;
-		}
-        public static string StaticGetFullName()
-        {
-            return AggregatePromotableNode.StaticGetFullName("Shared");
-        }
-    }
-
 	public class FullGameAggregateNode : GUBPAggregateNode
 	{
 		string GameName;
@@ -1801,9 +1656,12 @@ partial class GUBP
 
     public class WaitToPackageSamplesNode : WaitForUserInput
     {
-        public WaitToPackageSamplesNode()
+        public WaitToPackageSamplesNode(IEnumerable<UnrealTargetPlatform> HostPlatforms)
         {
-			AddDependency(WaitForSharedPromotionUserInput.StaticGetFullName(false));
+			foreach(UnrealTargetPlatform HostPlatform in HostPlatforms)
+			{
+				AddDependency(RootEditorNode.StaticGetFullName(HostPlatform));
+			}
 		}
 
         public static string StaticGetFullName()
@@ -1827,309 +1685,6 @@ partial class GUBP
         }
     }
 
-    public class WaitForPromotionUserInput : WaitForUserInput
-    {
-        string PromotionLabelPrefix;
-        string PromotionLabelSuffix;
-        protected bool bLabelPromoted; // true if this is the promoted version
-
-        public WaitForPromotionUserInput(string InPromotionLabelPrefix, string InPromotionLabelSuffix, bool bInLabelPromoted)
-        {
-            PromotionLabelPrefix = InPromotionLabelPrefix;
-            PromotionLabelSuffix = InPromotionLabelSuffix;
-            bLabelPromoted = bInLabelPromoted;
-            if (bLabelPromoted)
-            {
-                AddDependency(LabelPromotableNode.StaticGetFullName(PromotionLabelPrefix, false));
-            }
-            else
-            {
-                AddDependency(AggregatePromotableNode.StaticGetFullName(PromotionLabelPrefix));
-            }
-        }
-        public static string StaticGetFullName(string InPromotionLabelPrefix, string InPromotionLabelSuffix, bool bInLabelPromoted)
-        {
-            return InPromotionLabelPrefix + (bInLabelPromoted ? "_WaitForPromotion" : "_WaitForPromotable") + InPromotionLabelSuffix;
-        }
-        public override string GetFullName()
-        {
-            return StaticGetFullName(PromotionLabelPrefix, PromotionLabelSuffix, bLabelPromoted);
-        }
-    }
-
-    public class WaitForGamePromotionUserInput : WaitForPromotionUserInput
-    {
-        BranchInfo.BranchUProject GameProj;
-        bool bCustomWorkflow;        
-        public WaitForGamePromotionUserInput(GUBP bp, BranchInfo.BranchUProject InGameProj, bool bInLabelPromoted)
-            : base(InGameProj.GameName, "", bInLabelPromoted)
-        {
-            GameProj = InGameProj;
-            var Options = InGameProj.Options(UnrealTargetPlatform.Win64);
-            bCustomWorkflow = Options.bCustomWorkflowForPromotion;       
-        }
-        public static string StaticGetFullName(BranchInfo.BranchUProject InGameProj, bool bInLabelPromoted)
-        {
-            return WaitForPromotionUserInput.StaticGetFullName(InGameProj.GameName, "", bInLabelPromoted);
-        }
-        public override string GameNameIfAnyForTempStorage()
-        {
-            return GameProj.GameName;
-        }
-        public override string GetTriggerDescText()
-        {
-            if (bLabelPromoted)
-            {
-                return GameProj.GameName + " is ready for promotion.";
-            }
-            return GameProj.GameName + " is ready to make a promotable label and begin testing.";
-        }
-        public override string GetTriggerActionText()
-        {
-            if (bLabelPromoted)
-            {
-                return "Promote " + GameProj.GameName + ".";
-            }
-            return "Make a promotable label for " + GameProj.GameName + " and begin testing.";
-        }
-        public override string GetTriggerStateName()
-        {
-            if (bCustomWorkflow)
-            {
-                return GetFullName();
-            }
-            return base.GetTriggerStateName();
-        }
-        public override bool TriggerRequiresRecursiveWorkflow()
-        {
-            if (bCustomWorkflow)
-            {
-                return !bLabelPromoted; // the promotable starts the hardcoded chain
-            }
-            return base.TriggerRequiresRecursiveWorkflow();
-        }
-    }
-
-    public class WaitForSharedPromotionUserInput : WaitForPromotionUserInput
-    {
-        
-        public WaitForSharedPromotionUserInput(GUBP bp, bool bInLabelPromoted)
-            : base("Shared", IsMainBranch(), bInLabelPromoted)
-        {
-        }
-        public override string GetTriggerDescText()
-        {
-            if (bLabelPromoted)
-            {
-                return "The shared promotable is ready for promotion.";
-            }
-            return "The shared promotable is ready to make a promotable label.";
-        }
-        public override string GetTriggerActionText()
-        {
-            if (bLabelPromoted)
-            {
-                return "Promote the shared promotable.";
-            }
-            return "Make the shared promotable label.";
-        }
-        public static string StaticGetFullName(bool bInLabelPromoted)
-        {            
-            return WaitForPromotionUserInput.StaticGetFullName("Shared", IsMainBranch(), bInLabelPromoted);
-        }
-        public static string IsMainBranch()
-        {
-            string isMain = "";
-            if (P4Enabled)
-            {
-                string CurrentBranch = P4Env.BuildRootP4;
-                if (CurrentBranch == "//depot/UE4")
-                {
-                    isMain = "_WithNightlys";
-                }
-            }
-            return isMain;
-        }
-        public override string GetTriggerStateName()
-        {
-            return GetFullName();
-        }
-        public override bool TriggerRequiresRecursiveWorkflow()
-        {
-            return !bLabelPromoted;
-        }
-    }   	
-
-    public class LabelPromotableNode : GUBPNode
-    {
-		GUBP.GUBPBranchConfig BranchConfig;
-        string PromotionLabelPrefix;        
-        protected bool bLabelPromoted; // true if this is the promoted version
-
-        public LabelPromotableNode(GUBP.GUBPBranchConfig InBranchConfig, string InPromotionLabelPrefix, string InPromotionLabelSuffix, bool bInLabelPromoted)
-        {
-			BranchConfig = InBranchConfig;
-            PromotionLabelPrefix = InPromotionLabelPrefix;            
-            bLabelPromoted = bInLabelPromoted;
-            AddDependency(WaitForPromotionUserInput.StaticGetFullName(PromotionLabelPrefix, InPromotionLabelSuffix, bLabelPromoted));
-        }
-        string LabelName(bool bLocalLabelPromoted)
-        {
-            string LabelPrefix = PromotionLabelPrefix;            
-            string CompleteLabelPrefix = (bLocalLabelPromoted ? "Promoted-" : "Promotable-") + LabelPrefix;
-            if (LabelPrefix == "Shared" && bLocalLabelPromoted)
-            {
-                // shared promotion has a shorter name
-                CompleteLabelPrefix = "Promoted";
-            }
-            if (LabelPrefix == "Shared" && !bLocalLabelPromoted)
-            {
-                //shared promotable has a shorter name
-                CompleteLabelPrefix = "Promotable";
-            }                   
-            return CompleteLabelPrefix;
-        }
-
-        public override bool IsSticky()
-        {
-            return true;
-        }
-        public override bool SendSuccessEmail()
-        {
-            return true;
-        }
-
-
-        public override void DoBuild(GUBP bp)
-        {
-            BuildProducts = new List<string>();
-
-			if (P4Enabled && !BranchConfig.JobInfo.IsPreflight)
-            {
-                if (AllDependencyBuildProducts.Count == 0)
-                {
-                    throw new AutomationException("{0} has no build products", GetFullName());
-
-                }
-
-                if (bLabelPromoted)
-                {
-					P4.MakeDownstreamLabelFromLabel(P4Env, LabelName(true), LabelName(false));
-                }
-                else
-                {
-					int WorkingCL = P4.CreateChange(P4Env.Client, String.Format("GUBP Node {0} built from changelist {1}", GetFullName(), BranchConfig.JobInfo.Changelist));
-                    LogConsole("Build from {0}    Working in {1}", BranchConfig.JobInfo.Changelist, WorkingCL);
-
-                    var ProductsToSubmit = new List<String>();
-
-                    foreach (var Product in AllDependencyBuildProducts)
-                    {
-                        // hacks to keep certain things out of P4
-                        if (
-                            !Product.EndsWith("version.h", StringComparison.InvariantCultureIgnoreCase) && 
-                            !Product.EndsWith("version.cpp", StringComparison.InvariantCultureIgnoreCase) &&
-							!Product.Replace('\\', '/').EndsWith("DotNetCommon/MetaData.cs", StringComparison.InvariantCultureIgnoreCase) &&
-                            !Product.EndsWith("_Success.log", StringComparison.InvariantCultureIgnoreCase) &&
-							!Product.Replace('\\', '/').Contains("/Intermediate/") &&
-							!Product.Replace('\\', '/').Contains("/Engine/Saved/") &&
-							!Product.Replace('\\', '/').Contains("/DerivedDataCache/") &&
-							!Product.EndsWith("build.version", StringComparison.InvariantCultureIgnoreCase) && 
-							!Product.EndsWith(".lib") &&
-							!Product.EndsWith(".a") && 
-							!Product.EndsWith(".bc")
-                            )
-                        {
-                            ProductsToSubmit.Add(Product);
-                        }
-                    }
-
-                    // Open files for add or edit
-                    UE4Build.AddBuildProductsToChangelist(WorkingCL, ProductsToSubmit);
-
-                    // Check everything in!
-                    int SubmittedCL;
-					P4.Submit(WorkingCL, out SubmittedCL, true, true);
-
-                    // Label it       
-					P4.MakeDownstreamLabel(P4Env, LabelName(false), null);
-                }
-            }
-            SaveRecordOfSuccessAndAddToBuildProducts();
-        }
-        public static string StaticGetFullName(string InPromotionLabelPrefix, bool bInLabelPromoted)
-        {
-            return InPromotionLabelPrefix + (bInLabelPromoted ? "_LabelPromoted" : "_LabelPromotable");
-        }
-        public override string GetFullName()
-        {
-            return StaticGetFullName(PromotionLabelPrefix, bLabelPromoted);
-        }
-    }
-
-    public class GameLabelPromotableNode : LabelPromotableNode
-    {
-        BranchInfo.BranchUProject GameProj;        
-        public GameLabelPromotableNode(GUBP.GUBPBranchConfig InBranchConfig, BranchInfo.BranchUProject InGameProj, bool bInLabelPromoted)
-			: base(InBranchConfig, InGameProj.GameName, "", bInLabelPromoted)
-        {
-            GameProj = InGameProj;
-        }
-
-        public static string StaticGetFullName(BranchInfo.BranchUProject InGameProj, bool bInLabelPromoted)
-        {
-            return LabelPromotableNode.StaticGetFullName(InGameProj.GameName, bInLabelPromoted);
-        }
-        public override string GameNameIfAnyForTempStorage()
-        {
-            return GameProj.GameName;
-        }
-    }
-
-    public class SharedLabelPromotableNode : LabelPromotableNode
-    {        
-        public SharedLabelPromotableNode(GUBP.GUBPBranchConfig InBranchConfig, bool bInLabelPromoted)
-            : base(InBranchConfig, "Shared", IsMainBranch(), bInLabelPromoted)
-        {
-        }
-
-        public static string StaticGetFullName(bool bInLabelPromoted)
-        {
-            return LabelPromotableNode.StaticGetFullName("Shared", bInLabelPromoted);
-        }
-        public static string IsMainBranch()
-        {
-            string isMain = "";
-            if (P4Enabled)
-            {
-                string CurrentBranch = P4Env.BuildRootP4;
-                if (CurrentBranch == "//depot/UE4")
-                {
-                    isMain = "_WithNightlys";
-                }
-            }
-            return isMain;
-        }
-    }
-
-    public class SharedLabelPromotableSuccessNode : GUBPAggregateNode
-    {        
-        public SharedLabelPromotableSuccessNode()
-        {
-			AddDependency(SharedLabelPromotableNode.StaticGetFullName(false));
-        }
-
-        public static string StaticGetFullName()
-        {
-            return SharedLabelPromotableNode.StaticGetFullName(false) + "Aggregate";
-        }
-
-		public override string GetFullName()
-		{
-			return StaticGetFullName();
-		}
-    }
-
     public class WaitForTestShared : GUBPAggregateNode
     {
         public WaitForTestShared(GUBP bp)
diff --git a/Engine/Source/Programs/AutomationTool/Scripts/RocketBuild.Automation.cs b/Engine/Source/Programs/AutomationTool/Scripts/RocketBuild.Automation.cs
--- a/Engine/Source/Programs/AutomationTool/Scripts/RocketBuild.Automation.cs
+++ b/Engine/Source/Programs/AutomationTool/Scripts/RocketBuild.Automation.cs
@@ -80,7 +80,7 @@ namespace Rocket
 				// Add the aggregate for making a rocket build
 				if(!BranchConfig.HasNode(WaitToMakeRocketBuild.StaticGetFullName()))
 				{
-					BranchConfig.AddNode(new WaitToMakeRocketBuild(BranchConfig.HostPlatforms));
+					BranchConfig.AddNode(new WaitToMakeRocketBuild(BranchConfig));
 				}
 
 				// Find all the target platforms for this host platform.
@@ -271,12 +271,18 @@ namespace Rocket
 
 	public class WaitToMakeRocketBuild : GUBP.WaitForUserInput
 	{
-        public WaitToMakeRocketBuild(List<UnrealTargetPlatform> HostPlatforms)
+        public WaitToMakeRocketBuild(GUBP.GUBPBranchConfig BranchConfig)
         {
-			foreach(UnrealTargetPlatform HostPlatform in HostPlatforms)
+			foreach(UnrealTargetPlatform HostPlatform in BranchConfig.HostPlatforms)
 			{
 				AddDependency(FilterRocketNode.StaticGetFullName(HostPlatform));
 				AddDependency(BuildDerivedDataCacheNode.StaticGetFullName(HostPlatform));
+
+				SingleTargetProperties BuildPatchTool = BranchConfig.Branch.FindProgram("BuildPatchTool");
+				if(BuildPatchTool.Rules != null)
+				{
+					AddDependency(GUBP.SingleInternalToolsNode.StaticGetFullName(HostPlatform, BuildPatchTool));
+				}
 			}
 		}
 
