Mac HTML5LaunchHelper script is now copied instead of created to keep file permissions (i.e. +x).
Mark HTML5LaunchHelper script as executable.
Script handles spaces when pushing working directory.
Fixes UE-19637 (again)

diff --git a/Engine/Source/Programs/AutomationTool/HTML5/HTML5Platform.Automation.cs b/Engine/Source/Programs/AutomationTool/HTML5/HTML5Platform.Automation.cs
--- a/Engine/Source/Programs/AutomationTool/HTML5/HTML5Platform.Automation.cs
+++ b/Engine/Source/Programs/AutomationTool/HTML5/HTML5Platform.Automation.cs
@@ -315,11 +315,20 @@ public class HTML5Platform : Platform
 
 		if (outputContents.Length > 0)
 		{
-			// Save the file
+			// Save the file. We Copy the template file to keep any permissions set to it.
 			try
 			{
 				Directory.CreateDirectory(Path.GetDirectoryName(InOutputFile));
-				File.WriteAllText(InOutputFile, outputContents.ToString(), Encoding.UTF8);
+				if (File.Exists(InOutputFile))
+				{
+					File.Delete(InOutputFile);
+				}
+				File.Copy(InTemplateFile, InOutputFile);
+				using (var CmdFile = File.OpenWrite(InOutputFile)) 
+				{
+					Byte[] BytesToWrite = new UTF8Encoding(true).GetBytes(outputContents.ToString());
+					CmdFile.Write(BytesToWrite, 0, BytesToWrite.Length);
+				}
 			}
 			catch (Exception)
 			{
