Refactoring out some dead code that Resharper pointed out.

diff --git a/Engine/Source/Programs/AutomationTool/AutomationUtils/Log.cs b/Engine/Source/Programs/AutomationTool/AutomationUtils/Log.cs
--- a/Engine/Source/Programs/AutomationTool/AutomationUtils/Log.cs
+++ b/Engine/Source/Programs/AutomationTool/AutomationUtils/Log.cs
@@ -109,20 +109,12 @@ namespace AutomationTool
         /// <returns>The newly created TraceListener, or null if it could not be created.</returns>
         private static TraceListener CreateLogFileListener(out string LogFilename)
         {
-            StreamWriter LogFile = null;
             const int MaxAttempts = 10;
             int Attempt = 0;
             var TempLogFolder = Path.GetTempPath();
             do
             {
-                if (Attempt == 0)
-                {
-                    LogFilename = CommandUtils.CombinePaths(TempLogFolder, "Log.txt");
-                }
-                else
-                {
-                    LogFilename = CommandUtils.CombinePaths(TempLogFolder, String.Format("Log_{0}.txt", Attempt));
-                }
+                LogFilename = CommandUtils.CombinePaths(TempLogFolder, Attempt == 0 ? "Log.txt" : string.Format("Log_{0}.txt", Attempt));
                 try
                 {
                     // We do not need to set AutoFlush on the StreamWriter because we set Trace.AutoFlush, which calls it for us.
@@ -136,12 +128,11 @@ namespace AutomationTool
                     if (Attempt == (MaxAttempts - 1))
                     {
                         // Clear out the LogFilename to indicate we were not able to write one.
+                        UnrealBuildTool.Log.TraceWarning("Unable to create log file {0}: {1}", LogFilename, Ex);
                         LogFilename = null;
-                        UnrealBuildTool.Log.TraceWarning("Unable to create log file: {0}", LogFilename);
-                        UnrealBuildTool.Log.TraceWarning(LogUtils.FormatException(Ex));
                     }
                 }
-            } while (LogFile == null && ++Attempt < MaxAttempts);
+            } while (++Attempt < MaxAttempts);
             return null;
         }
 
