Parallel zip to temp storage using local folder as a staging dir first.

diff --git a/Engine/Source/Programs/AutomationTool/GUBP/GUBP.Automation.csproj b/Engine/Source/Programs/AutomationTool/GUBP/GUBP.Automation.csproj
--- a/Engine/Source/Programs/AutomationTool/GUBP/GUBP.Automation.csproj
+++ b/Engine/Source/Programs/AutomationTool/GUBP/GUBP.Automation.csproj
@@ -45,6 +45,8 @@
     </Reference>
     <Reference Include="System" />
     <Reference Include="System.Core" />
+    <Reference Include="System.IO.Compression" />
+    <Reference Include="System.IO.Compression.FileSystem" />
     <Reference Include="System.Runtime.Serialization" />
     <Reference Include="System.Web" />
     <Reference Include="System.Web.Extensions" />
diff --git a/Engine/Source/Programs/AutomationTool/GUBP/GUBP.cs b/Engine/Source/Programs/AutomationTool/GUBP/GUBP.cs
--- a/Engine/Source/Programs/AutomationTool/GUBP/GUBP.cs
+++ b/Engine/Source/Programs/AutomationTool/GUBP/GUBP.cs
@@ -43,7 +43,6 @@ using System.Diagnostics;
 [Help("ForceIncrementalCompile", "make sure all compiles are incremental")]
 [Help("AutomatedTesting", "Allow automated testing, currently disabled.")]
 [Help("StompCheck", "Look for stomped build products.")]
-[Help("NoZipTempStorage", "By default temp storage is zipped up and copied as a single archive. When this flag is set, the directory structure is replicated instead.")]
 public partial class GUBP : BuildCommand
 {
 	const string StartedTempStorageSuffix = "_Started";
@@ -166,9 +165,8 @@ public partial class GUBP : BuildCommand
 		{
 			HostPlatforms.Add(UnrealTargetPlatform.Linux);
 		}
-        TempStorage.SetZipTempStorage(!ParseParam("NoZipTempStorage"));
 
-        bool bCleanLocalTempStorage = ParseParam("CleanLocal");
+	    bool bCleanLocalTempStorage = ParseParam("CleanLocal");
         bool bSkipTriggers = ParseParam("SkipTriggers");
         bool bFake = ParseParam("fake");
         bool bFakeEC = ParseParam("FakeEC");
@@ -237,7 +235,7 @@ public partial class GUBP : BuildCommand
 
         if (bCleanLocalTempStorage)  // shared temp storage can never be wiped
         {
-            TempStorage.DeleteLocalTempStorageManifests();
+            TempStorage.DeleteLocalTempStorage();
         }
 
 		int TimeIndex = ParseParamInt("TimeIndex", 0);
diff --git a/Engine/Source/Programs/AutomationTool/GUBP/TempStorage.cs b/Engine/Source/Programs/AutomationTool/GUBP/TempStorage.cs
--- a/Engine/Source/Programs/AutomationTool/GUBP/TempStorage.cs
+++ b/Engine/Source/Programs/AutomationTool/GUBP/TempStorage.cs
@@ -1,11 +1,15 @@
 ﻿// Copyright 1998-2015 Epic Games, Inc. All Rights Reserved.
 using System;
+using System.Collections.Concurrent;
 using System.Collections.Generic;
 using System.Text;
 using System.Linq;
 using System.IO;
+using System.Threading.Tasks;
 using System.Xml;
 using System.Xml.Linq;
+using System.IO.Compression;
+using System.Threading;
 
 namespace AutomationTool
 {
@@ -63,7 +67,7 @@ namespace AutomationTool
         /// <summary>
         /// Used by anything that needs the block info as a string. This should generally be legacy stuff, as we should only be using the temp storage path for temp storage!
         /// </summary>
-        /// <returns>{BuildRootForTempStorage}-{Changelist}{PreflightInfo}-{NodeStorageName}</returns>
+        /// <returns>{BranchNameForTempStorage}-{Changelist}{PreflightInfo}-{NodeStorageName}</returns>
         public string GetLegacyString()
         {
             if (string.IsNullOrEmpty(CachedString))
@@ -78,7 +82,7 @@ namespace AutomationTool
         /// NOTE: <see cref="FindTempStorageNodeCLsMatchingSuffix"/> is dependent on this directory layout!!!
         /// If this layout is changed, FindTempStorageNodeCLsMatchingSuffix needs to be changed as well.
         /// </summary>
-        /// <returns>{BuildRootForTempStorage}/{Changelist}{PreflightInfo}/{NodeStorageName}</returns>
+        /// <returns>{BranchNameForTempStorage}/{Changelist}{PreflightInfo}/{NodeStorageName}</returns>
         public string GetRelativeDirectory()
         {
             if (string.IsNullOrEmpty(CachedDir))
@@ -424,12 +428,11 @@ namespace AutomationTool
             /// Gets a flat list of files in the manifest, converting to full path names rooted at the given base dir.
             /// </summary>
             /// <param name="RootDir">Root dir to prepend to all the files in the manifest.</param>
-            /// <param name="bFilesMustExist">if true, one file must exist else the function fails.</param>
             /// <returns>Flat list of all files in the manifest re-rooted at RootDir.</returns>
-            public List<string> GetFiles(string RootDir, bool bFilesMustExist = true)
+            public List<string> GetFiles(string RootDir)
             {
                 // flatten the list of directories, pull the files out, set their root path, and ensure the path is not too long.
-                var Result = Directories.SelectMany(Dir=>Dir.Value).Select(FileInfo =>
+                return Directories.SelectMany(Dir=>Dir.Value).Select(FileInfo =>
                     {
                         var NewFilePath = CommandUtils.CombinePaths(RootDir, FileInfo.Name);
                         // create a FileInfo using the file, which will help us catch path too long exceptions early.
@@ -442,13 +445,6 @@ namespace AutomationTool
                             throw new AutomationException(Ex, "Path too long ... failed to create FileInfo for {0}", NewFilePath);
                         }
                     }).ToList();
-
-                if (bFilesMustExist)
-                {
-                    CommandUtils.Log("Checking that first file in manifest exists...");
-                    InternalUtils.Robust_FileExists(false, Result[0], "GetFiles in temp manifest failed because the first Rebased manifest file does not exist: {0}");
-                }
-                return Result;
             }
         }
 
@@ -728,21 +724,6 @@ namespace AutomationTool
             return CommandUtils.CombinePaths(ResolveSharedTempStorageDirectory(GameName), TempStorageNodeInfo.GetRelativeDirectory());
         }
 
-        /// <summary>
-        /// Private helper that saves a temp file manifest described by a list of fully qualified files rooted at RootDir to the given fully qualified path, returning the created TempStorageManifest instance.
-        /// </summary>
-        /// <param name="RootDir">Folder that all the given files are rooted from.</param>
-        /// <param name="FinalFilename">Fully qualified path to of the file that should be written.</param>
-        /// <param name="Files">Fully qualified names of files to reference in the manifest file.</param>
-        /// <returns>The created manifest instance (which has already been saved to disk).</returns>
-        private static TempStorageManifest SaveTempStorageManifest(string RootDir, string FinalFilename, List<string> Files)
-        {
-            var Saver = TempStorageManifest.Create(Files, RootDir);
-            CommandUtils.CreateDirectory(true, Path.GetDirectoryName(FinalFilename));
-            Saver.Save(FinalFilename);
-            return Saver;
-        }
-
         /// <summary>
         /// Gets the name of the local temp storage manifest file for the given temp storage node (ie, Engine\Saved\TmpStore\NodeInfoDirectory\NodeInfoFilename.TempManifest)
         /// </summary>
@@ -762,7 +743,11 @@ namespace AutomationTool
         /// <returns>The created manifest instance (which has already been saved to disk).</returns>
         private static TempStorageManifest SaveLocalTempStorageManifest(string RootDir, TempStorageNodeInfo TempStorageNodeInfo, List<string> Files)
         {
-            return SaveTempStorageManifest(RootDir, LocalTempStorageManifestFilename(TempStorageNodeInfo), Files);
+            string FinalFilename = LocalTempStorageManifestFilename(TempStorageNodeInfo);
+            var Saver = TempStorageManifest.Create(Files, RootDir);
+            CommandUtils.CreateDirectory(true, Path.GetDirectoryName(FinalFilename));
+            Saver.Save(FinalFilename);
+            return Saver;
         }
 
         /// <summary>
@@ -776,24 +761,12 @@ namespace AutomationTool
             return CommandUtils.CombinePaths(SharedTempStorageDirectory(TempStorageNodeInfo, GameName), TempStorageNodeInfo.GetManifestFilename());
         }
 
-        /// <summary>
-        /// Saves the list of fully qualified files (that should be rooted at the shared temp storage location for the game) to a shared temp storage manifest with the given temp storage node and game.
-        /// </summary>
-        /// <param name="TempStorageNodeInfo">Node info descibing the block of temp storage (essentially used to identify a subdirectory insides the game's temp storage folder).</param>
-        /// <param name="GameName">game name to determine the temp storage folder for. Empty is equivalent to "UE4".</param>
-        /// <param name="Files">Fully qualified names of files to reference in the manifest file.</param>
-        /// <returns>The created manifest instance (which has already been saved to disk).</returns>
-        private static TempStorageManifest SaveSharedTempStorageManifest(TempStorageNodeInfo TempStorageNodeInfo, string GameName, List<string> Files)
-        {
-            return SaveTempStorageManifest(SharedTempStorageDirectory(TempStorageNodeInfo, GameName), SharedTempStorageManifestFilename(TempStorageNodeInfo, GameName), Files);
-        }
-
         /// <summary>
         /// Deletes all temp storage manifests from the local storage location (Engine\Saved\TmpStore).
         /// Local temp storage logic only stores manifests, so this effectively deletes any local temp storage work, while not actually deleting the local files, which the temp 
         /// storage system doesn't really own.
         /// </summary>
-        public static void DeleteLocalTempStorageManifests()
+        public static void DeleteLocalTempStorage()
         {
             CommandUtils.DeleteDirectory(true, LocalTempStorageManifestDirectory());
         }
@@ -862,26 +835,142 @@ namespace AutomationTool
         }
 
         /// <summary>
-        /// Controls how many threads to use when performing parallel copies.
+        /// Stores some result info from a parallel zip or unzip operation.
         /// </summary>
-        /// <returns>The number of threads to use.</returns>
-        private static int ThreadsToCopyWith()
+        class ParallelZipResult
         {
-            return 8;
+            /// <summary>
+            /// Time taken for the zip portion of the operation (actually, everything BUT the copy time).
+            /// </summary>
+            public readonly TimeSpan ZipTime;
+
+            /// <summary>
+            /// Total size of the zip files that are created.
+            /// </summary>
+            public readonly long ZipFilesTotalSize;
+
+            public ParallelZipResult(TimeSpan ZipTime, long ZipFilesTotalSize)
+            {
+                this.ZipTime = ZipTime;
+                this.ZipFilesTotalSize = ZipFilesTotalSize;
+            }
         }
 
         /// <summary>
-        /// Controls whether temp storage for this job is zipped up or not.
+        /// Zips a set of files (that must be rooted at the given RootDir) to a set of zip files in the given OutputDir. The files will be prefixed with the given basename.
         /// </summary>
-        private static bool bZipTempStorage = true;
+        /// <param name="Files">Fully qualified list of files to zip (must be rooted at RootDir).</param>
+        /// <param name="RootDir">Root Directory where all files will be extracted.</param>
+        /// <param name="OutputDir">Location to place the set of zip files created.</param>
+        /// <param name="StagingDir">Location to create zip files before copying them to the OutputDir. If the OutputDir is on a remote file share, staging may be more efficient. Use null to avoid using a staging copy.</param>
+        /// <param name="ZipBasename">The basename of the set of zip files.</param>
+        /// <returns>Some metrics about the zip process.</returns>
+        /// <remarks>
+        /// This function tries to zip the files in parallel as fast as it can. It makes no guarantees about how many zip files will be created or which files will be in which zip,
+        /// but it does try to reasonably balance the file sizes.
+        /// </remarks>
+        private static ParallelZipResult ParallelZipFiles(IEnumerable<string> Files, string RootDir, string OutputDir, string StagingDir, string ZipBasename)
+        {
+            var ZipTimer = DateTimeStopwatch.Start();
+            // First get the sizes of all the files. We won't parallelize if there isn't enough data to keep the number of zips down.
+            var FilesInfo = Files
+                .Select(File => new { File, FileSize = new FileInfo(File).Length })
+                .ToList();
+
+            // Profiling results show that we can zip 100MB quite fast and it is not worth parallelizing that case and creating a bunch of zips that are relatively small.
+            const long MinFileSizeToZipInParallel = 1024 * 1024 * 100L;
+            var bZipInParallel = FilesInfo.Sum(FileInfo => FileInfo.FileSize) >= MinFileSizeToZipInParallel;
+
+            // order the files in descending order so our threads pick up the biggest ones first.
+            // We want to end with the smaller files to more effectively fill in the gaps
+            var FilesToZip = new ConcurrentQueue<string>(FilesInfo.OrderByDescending(FileInfo => FileInfo.FileSize).Select(FileInfo => FileInfo.File));
+
+            long ZipFilesTotalSize = 0L;
+            // We deliberately avoid Parallel.ForEach here because profiles have shown that dynamic partitioning creates
+            // too many zip files, and they can be of too varying size, creating uneven work when unzipping later,
+            // as ZipFile cannot unzip files in parallel from a single archive.
+            // We can safely assume the build system will not be doing more important things at the same time, so we simply use all our logical cores,
+            // which has shown to be optimal via profiling, and limits the number of resulting zip files to the number of logical cores.
+            var Threads = (
+                from CoreNum in Enumerable.Range(0, bZipInParallel ? Environment.ProcessorCount : 1)
+                let ZipFileName = Path.Combine(StagingDir ?? OutputDir, string.Format("{0}{1}.zip", ZipBasename, bZipInParallel ? "-" + CoreNum.ToString("00") : ""))
+                select new Thread(() =>
+                {
+                    // Create one zip per thread using the given basename
+                    using (var ZipArchive = ZipFile.Open(ZipFileName, ZipArchiveMode.Create))
+                    {
+
+                        // pull from the queue until we are out of files.
+                        string File;
+                        while (FilesToZip.TryDequeue(out File))
+                        {
+                            // use fastest compression. In our best case we are CPU bound, so this is a good tradeoff,
+                            // cutting overall time by 2/3 while only modestly increasing the compression ratio (22.7% -> 23.8% for RootEditor PDBs).
+                            // This is in cases of a super hot cache, so the operation was largely CPU bound.
+                            ZipArchive.CreateEntryFromFile(File, CommandUtils.StripBaseDirectory(File, RootDir), CompressionLevel.Fastest);
+                        }
+                    }
+                    Interlocked.Add(ref ZipFilesTotalSize, new FileInfo(ZipFileName).Length);
+                    // if we are using a staging dir, copy to the final location and delete the staged copy.
+                    if (StagingDir != null)
+                    {
+                        CommandUtils.CopyFile(ZipFileName, CommandUtils.MakeRerootedFilePath(ZipFileName, StagingDir, OutputDir));
+                        CommandUtils.DeleteFile(true, ZipFileName);
+                    }
+                })).ToList();
+
+            Threads.ForEach(thread => thread.Start());
+            Threads.ForEach(thread => thread.Join());
+            return new ParallelZipResult(ZipTimer.ElapsedTime, ZipFilesTotalSize);
+        }
 
         /// <summary>
-        /// Used to set the usage of temp storage as either zipped or having the entire directory structure copied. Defaults to true.
+        /// Unzips a set of zip files with a given basename in a given folder to a given RootDir.
         /// </summary>
-        /// <param name="Val">true to zip temp storage, false to mirror the directory tree when storing.</param>
-        public static void SetZipTempStorage(bool Val)
+        /// <param name="RootDir">Root Directory where all files will be extracted.</param>
+        /// <param name="FolderWithZipFiles">Folder containing the zip files to unzip. None of the zips should have the same file path in them.</param>
+        /// <param name="ZipBasename">The basename of the set of zip files to unzip.</param>
+        /// <returns>Some metrics about the unzip process.</returns>
+        /// <remarks>
+        /// The code is expected to be the used as the symmetrical inverse of <see cref="ParallelZipFiles"/>, but could be used independently, as long as the files in the zip do not overlap.
+        /// </remarks>
+        private static ParallelZipResult ParallelUnZipFiles(string RootDir, string FolderWithZipFiles, string ZipBasename)
         {
-            bZipTempStorage = Val;
+            var UnzipTimer = DateTimeStopwatch.Start();
+            long ZipFilesTotalSize = 0L;
+            Parallel.ForEach(Directory.EnumerateFiles(FolderWithZipFiles, ZipBasename + "*.zip").ToList(),
+                (ZipFilename) =>
+                {
+                    Interlocked.Add(ref ZipFilesTotalSize, new FileInfo(ZipFilename).Length);
+                    // unzip the files manually instead of caling ZipFile.ExtractToDirectory() because we need to overwrite readonly files. Because of this, creating the directories is up to us as well.
+                    using (var ZipArchive = ZipFile.OpenRead(ZipFilename))
+                    {
+                        foreach (var Entry in ZipArchive.Entries)
+                        {
+                            // Use CommandUtils.CombinePaths to ensure directory separators get converted correctly. On mono on *nix, if the path has backslashes it will not convert it.
+                            var ExtractedFilename = CommandUtils.CombinePaths(RootDir, Entry.FullName);
+                            // Zips can contain empty dirs. Ours usually don't have them, but we should support it.
+                            if (Path.GetFileName(ExtractedFilename).Length == 0)
+                            {
+                                Directory.CreateDirectory(ExtractedFilename);
+                            }
+                            else
+                            {
+                                // We must delete any existing file, even if it's readonly. .Net does not do this by default.
+                                if (File.Exists(ExtractedFilename))
+                                {
+                                    InternalUtils.SafeDeleteFile(ExtractedFilename, true);
+                                }
+                                else
+                                {
+                                    Directory.CreateDirectory(Path.GetDirectoryName(ExtractedFilename));
+                                }
+                                Entry.ExtractToFile(ExtractedFilename, true);
+                            }
+                        }
+                    }
+                });
+            return new ParallelZipResult(UnzipTimer.ElapsedTime, ZipFilesTotalSize);
         }
 
         /// <summary>
@@ -906,80 +995,31 @@ namespace AutomationTool
                 // save the manifest to local temp storage.
                 var Local = SaveLocalTempStorageManifest(RootDir, TempStorageNodeInfo, Files);
                 var LocalTotalSize = Local.GetTotalSize();
-                if (!bLocalOnly)
+                if (bLocalOnly)
                 {
-                    var BlockPath = SharedTempStorageDirectory(TempStorageNodeInfo, GameName);
-                    CommandUtils.LogConsole("Storing to {0}", BlockPath);
-                    if (CommandUtils.DirectoryExists_NoExceptions(BlockPath))
-                    {
-                        throw new AutomationException("Storage Block Already Exists! {0}", BlockPath);
-                    }
-                    CommandUtils.CreateDirectory(true, BlockPath);
-                    if (!CommandUtils.DirectoryExists_NoExceptions(BlockPath))
-                    {
-                        throw new AutomationException("Storage Block Could Not Be Created! {0}", BlockPath);
-                    }
-
-                    // We know the source files exist and are under RootDir because we created the manifest, which verifies it.
-                    // Now create the list of target files
-                    var DestFiles = Files.Select(Filename => CommandUtils.MakeRerootedFilePath(Filename, RootDir, BlockPath)).ToList();
-                    // make sure the target file doesn't exist, as we never expect this.
-                    foreach (var DestFile in DestFiles)
-                    {
-                        if (CommandUtils.FileExists_NoExceptions(true, DestFile))
-                        {
-                            throw new AutomationException("Dest file {0} already exists.", DestFile);
-                        }
-                    }
-
-                    if (bZipTempStorage)
-                    {
-                        var ZipTimer = DateTimeStopwatch.Start();
-                        var Zip = new Ionic.Zip.ZipFile
-                        {
-                            UseZip64WhenSaving = Ionic.Zip.Zip64Option.Always
-                        };
-                        foreach (string File in Files)
-                        {
-                            // strip the root dir off each file so we get a proper relative path.
-                            Zip.AddFile(Path.Combine(RootDir, File), Path.GetDirectoryName(CommandUtils.StripBaseDirectory(File, RootDir)));
-                        }
-                        var LocalManifest = LocalTempStorageManifestFilename(TempStorageNodeInfo);
-                        var SharedManifest = SharedTempStorageManifestFilename(TempStorageNodeInfo, GameName);
-                        var ZipFilename = Path.ChangeExtension(LocalManifest, "zip");
-                        Zip.Save(ZipFilename);
-                        var ZipFileSize = new FileInfo(ZipFilename).Length;
-                        var ZipTimeMS = (long)ZipTimer.ElapsedTime.TotalMilliseconds;
-                        var CopyTimer = DateTimeStopwatch.Start();
-                        InternalUtils.Robust_CopyFile(ZipFilename, CommandUtils.CombinePaths(BlockPath, Path.GetFileName(ZipFilename)));
-                        var CopyTimeMS = (long)CopyTimer.ElapsedTime.TotalMilliseconds;
-                        CommandUtils.DeleteFile(ZipFilename);
-                        // copy the local manifest to the shared location. We have to assume the zip is a good copy.
-                        CommandUtils.CopyFile(LocalManifest, SharedManifest);
-                        TelemetryStopwatch.Finish(string.Format("StoreToTempStorage.{0}.{1}.{2}.Remote.{3}.{4}.{5}", Files.Count, LocalTotalSize, ZipFileSize, CopyTimeMS, ZipTimeMS, TempStorageNodeInfo.NodeStorageName));
-                    }
-                    else
-                    {
-                        var CopyTimer = DateTimeStopwatch.Start();
-                        CommandUtils.ThreadedCopyFiles(Files, DestFiles, ThreadsToCopyWith());
-                        var CopyTimeMS = (long)CopyTimer.ElapsedTime.TotalMilliseconds;
-
-                        // save the shared temp storage manifest file.
-                        var Shared = SaveSharedTempStorageManifest(TempStorageNodeInfo, GameName, DestFiles);
-                        // These manifests better have the same files, timestamps, and sizes, else something failed in the copy.
-                        if (!Local.Compare(Shared))
-                        {
-                            // we will rename this so it can't be used, but leave it around for inspection
-                            var SharedTempManifestFilename = SharedTempStorageManifestFilename(TempStorageNodeInfo, GameName);
-                            CommandUtils.RenameFile_NoExceptions(SharedTempManifestFilename, SharedTempManifestFilename + ".broken");
-                            throw new AutomationException("Shared and Local manifest mismatch. Wrote out {0} for inspection.", SharedTempManifestFilename + ".broken");
-                        }
-                        TelemetryStopwatch.Finish(string.Format("StoreToTempStorage.{0}.{1}.{2}.Remote.{3}.{4}.{5}", Files.Count, LocalTotalSize, 0L, CopyTimeMS, 0L, TempStorageNodeInfo.NodeStorageName));
-                    }
+                    TelemetryStopwatch.Finish(string.Format("StoreToTempStorage.{0}.{1}.{2}.Local.{3}.{4}.{5}", Files.Count, LocalTotalSize, 0L, 0L, 0L, TempStorageNodeInfo.NodeStorageName));
                 }
                 else
                 {
-                    TelemetryStopwatch.Finish(string.Format("StoreToTempStorage.{0}.{1}.{2}.Local.{3}.{4}.{5}", Files.Count, LocalTotalSize, 0L, 0L, 0L, TempStorageNodeInfo.NodeStorageName));
+                    var SharedStorageNodeDir = SharedTempStorageDirectory(TempStorageNodeInfo, GameName);
+                    CommandUtils.LogConsole("Storing to {0}", SharedStorageNodeDir);
+                    // this folder should not already exist, else we have concurrency or job duplication problems.
+                    if (CommandUtils.DirectoryExists_NoExceptions(SharedStorageNodeDir))
+                    {
+                        throw new AutomationException("Storage Block Already Exists! {0}", SharedStorageNodeDir);
+                    }
+                    CommandUtils.CreateDirectory(true, SharedStorageNodeDir);
+
+                    var LocalManifestFilename = LocalTempStorageManifestFilename(TempStorageNodeInfo);
+                    var SharedManifestFilename = SharedTempStorageManifestFilename(TempStorageNodeInfo, GameName);
+                    var StagingDir = Path.GetDirectoryName(LocalManifestFilename);
+                    var ZipBasename = Path.GetFileNameWithoutExtension(LocalManifestFilename);
+                    // initiate the parallel zip operation.
+                    var ZipResult = ParallelZipFiles(Files, RootDir, SharedStorageNodeDir, StagingDir, ZipBasename);
+
+                    // copy the local manifest to the shared location. We have to assume the zip is a good copy.
+                    CommandUtils.CopyFile(LocalManifestFilename, SharedManifestFilename);
+                    TelemetryStopwatch.Finish(string.Format("StoreToTempStorage.{0}.{1}.{2}.Remote.{3}.{4}.{5}", Files.Count, LocalTotalSize, ZipResult.ZipFilesTotalSize, 0, (long)ZipResult.ZipTime.TotalMilliseconds, TempStorageNodeInfo.NodeStorageName));
                 }
             }
         }
@@ -1026,66 +1066,38 @@ namespace AutomationTool
                 WasLocal = false;
 
                 // We couldn't find the node storage locally, so get it from the shared location.
-                var BlockPath = SharedTempStorageDirectory(TempStorageNodeInfo, GameName);
+                var SharedStorageNodeDir = SharedTempStorageDirectory(TempStorageNodeInfo, GameName);
 
-                CommandUtils.LogVerbose("Attempting to retrieve from {0}", BlockPath);
-                if (!CommandUtils.DirectoryExists_NoExceptions(BlockPath))
+                CommandUtils.LogVerbose("Attempting to retrieve from {0}", SharedStorageNodeDir);
+                if (!CommandUtils.DirectoryExists_NoExceptions(SharedStorageNodeDir))
                 {
-                    throw new AutomationException("Storage Block Does Not Exists! {0}", BlockPath);
+                    throw new AutomationException("Storage Block Does Not Exists! {0}", SharedStorageNodeDir);
                 }
                 var SharedManifest = SharedTempStorageManifestFilename(TempStorageNodeInfo, GameName);
                 InternalUtils.Robust_FileExists(SharedManifest, "Storage Block Manifest Does Not Exists! {0}");
 
                 var Shared = TempStorageManifest.Load(SharedManifest);
-                var SharedFiles = Shared.GetFiles(BlockPath, !bZipTempStorage);
+                var SharedFiles = Shared.GetFiles(SharedStorageNodeDir);
 
                 // We know the source files exist and are under RootDir because we created the manifest, which verifies it.
                 // Now create the list of target files
-                var DestFiles = SharedFiles.Select(Filename => CommandUtils.MakeRerootedFilePath(Filename, BlockPath, RootDir)).ToList();
+                var DestFiles = SharedFiles.Select(Filename => CommandUtils.MakeRerootedFilePath(Filename, SharedStorageNodeDir, RootDir)).ToList();
 
-                long ZipFileSize = 0L;
-                long CopyTimeMS = 0L;
-                long UnzipTimeMS = 0L;
-                if (bZipTempStorage)
-                {
-                    var LocalZipFilename = Path.ChangeExtension(LocalManifest, "zip");
-                    var SharedZipFilename = Path.ChangeExtension(SharedManifest, "zip");
-                    if (!CommandUtils.FileExists(SharedZipFilename))
-                    {
-                        throw new AutomationException("Storage block zip file does not exist! {0}", SharedZipFilename);
-                    }
-                    var CopyTimer = DateTimeStopwatch.Start();
-                    InternalUtils.Robust_CopyFile(SharedZipFilename, LocalZipFilename);
-                    CopyTimeMS = (long)CopyTimer.ElapsedTime.TotalMilliseconds;
-                    ZipFileSize = new FileInfo(LocalZipFilename).Length;
+                var ZipBasename = Path.GetFileNameWithoutExtension(LocalManifest);
 
-                    var UnzipTimer = DateTimeStopwatch.Start();
-                    using (Ionic.Zip.ZipFile Zip = new Ionic.Zip.ZipFile(LocalZipFilename))
-                    {
-                        Zip.ExtractAll(RootDir, Ionic.Zip.ExtractExistingFileAction.OverwriteSilently);
-                    }
-                    CommandUtils.DeleteFile(LocalZipFilename);
-                    UnzipTimeMS = (long)UnzipTimer.ElapsedTime.TotalMilliseconds;
-                }
-                else
-                {
-                    var CopyTimer = DateTimeStopwatch.Start();
-                    // If the local file already exists, it will be overwritten.
-                    foreach (var DestFile in DestFiles)
-                    {
-                        if (CommandUtils.FileExists_NoExceptions(true, DestFile))
-                        {
-                            CommandUtils.LogVerbose("Dest file {0} already exists, deleting and overwriting", DestFile);
-                            CommandUtils.DeleteFile(DestFile);
-                        }
-                    }
+                // now unzip in parallel, overwriting any existing local file.
+                var ZipResult = ParallelUnZipFiles(RootDir, SharedStorageNodeDir, ZipBasename);
 
-                    // Do the threaded copy to the local file system.
-                    CommandUtils.ThreadedCopyFiles(SharedFiles, DestFiles, ThreadsToCopyWith());
-                    CopyTimeMS = (long)CopyTimer.ElapsedTime.TotalMilliseconds;
+                var NewLocal = SaveLocalTempStorageManifest(RootDir, TempStorageNodeInfo, DestFiles);
+                // Now compare the created local files to ensure their attributes match the one we copied from the network.
+                if (!NewLocal.Compare(Shared))
+                {
+                    // we will rename this so it can't be used, but leave it around for inspection
+                    CommandUtils.RenameFile_NoExceptions(LocalManifest, LocalManifest + ".broken");
+                    throw new AutomationException("Shared and Local manifest mismatch.");
                 }
 
-                // Handle unix permissions/chmod issues.
+                // Handle unix permissions/chmod issues. This will touch the timestamp we check on the file, so do this after we've compared with the manifest attributes.
                 if (UnrealBuildTool.Utils.IsRunningOnMono)
                 {
                     foreach (string DestFile in DestFiles)
@@ -1093,15 +1105,8 @@ namespace AutomationTool
                         CommandUtils.FixUnixFilePermissions(DestFile);
                     }
                 }
-                var NewLocal = SaveLocalTempStorageManifest(RootDir, TempStorageNodeInfo, DestFiles);
-                // Now compare the created local files to ensure their attributes match the one we copied from the network.
-                if (!NewLocal.Compare(Shared))
-                {
-                    // we will rename this so it can't be used, but leave it around for inspection
-                    CommandUtils.RenameFile_NoExceptions(LocalManifest, LocalManifest + ".broken");
-                    throw new AutomationException("Shared and Local manifest mismatch.");
-                }
-                TelemetryStopwatch.Finish(string.Format("RetrieveFromTempStorage.{0}.{1}.{2}.Remote.{3}.{4}.{5}", DestFiles.Count, Shared.GetTotalSize(), ZipFileSize, CopyTimeMS, UnzipTimeMS, TempStorageNodeInfo.NodeStorageName));
+
+                TelemetryStopwatch.Finish(string.Format("RetrieveFromTempStorage.{0}.{1}.{2}.Remote.{3}.{4}.{5}", DestFiles.Count, Shared.GetTotalSize(), ZipResult.ZipFilesTotalSize, 0, (long)ZipResult.ZipTime.TotalMilliseconds, TempStorageNodeInfo.NodeStorageName));
                 return DestFiles;
             }
         }
@@ -1149,7 +1154,7 @@ namespace AutomationTool
             var TestTempStorageInfo = new TempStorageNodeInfo(new GUBP.JobInfo("Test", 0, 0, 0), "Test");
 
             // Delete any local and shared temp storage that may exist.
-            TempStorage.DeleteLocalTempStorageManifests();
+            TempStorage.DeleteLocalTempStorage();
             TempStorage.DeleteSharedTempStorage(TestTempStorageInfo, "UE4");
             if (TempStorage.TempStorageExists(TestTempStorageInfo, "UE4", false, false))
             {
@@ -1171,8 +1176,11 @@ namespace AutomationTool
                 .ToList();
             foreach (var TestFile in TestFiles)
             {
+                CommandUtils.DeleteFile(TestFile.FileName);
                 CommandUtils.Log("Test file {0}", TestFile.FileName);
                 File.WriteAllText(TestFile.FileName, TestFile.FileContents);
+                // we should be able to overwrite readonly files.
+                File.SetAttributes(TestFile.FileName, FileAttributes.ReadOnly);
             }
             
             // wrap the operation so we are sure to clean up afterward.
@@ -1191,7 +1199,7 @@ namespace AutomationTool
                     throw new AutomationException("shared storage should exist");
                 }
                 // Now delete the local manifest
-                TempStorage.DeleteLocalTempStorageManifests();
+                TempStorage.DeleteLocalTempStorage();
                 // It should no longer be there.
                 if (TempStorage.LocalTempStorageManifestExists(TestTempStorageInfo))
                 {
@@ -1202,6 +1210,13 @@ namespace AutomationTool
                 {
                     throw new AutomationException("some storage should exist");
                 }
+
+                // Now we should be able to retrieve the test files from shared storage, and it should overwrite our read-only files.
+                bool WasLocal;
+                TempStorage.RetrieveFromTempStorage(TestTempStorageInfo, out WasLocal, "UE4", CmdEnv.LocalRoot);
+                // Now delete the local manifest so we can try again with no files present (the usual case for restoring from temp storage).
+                TempStorage.DeleteLocalTempStorage();
+
                 // Ok, delete our test files locally.
                 foreach (var TestFile in TestFiles)
                 {
@@ -1209,7 +1224,6 @@ namespace AutomationTool
                 }
 
                 // Now we should be able to retrieve the test files from shared storage.
-                bool WasLocal;
                 TempStorage.RetrieveFromTempStorage(TestTempStorageInfo, out WasLocal, "UE4", CmdEnv.LocalRoot);
                 // the local manifest should be there, since we just retrieved from shared storage.
                 if (!TempStorage.LocalTempStorageManifestExists(TestTempStorageInfo))
@@ -1244,7 +1258,7 @@ namespace AutomationTool
                 }
 
                 // and now lets test tampering
-                TempStorage.DeleteLocalTempStorageManifests();
+                TempStorage.DeleteLocalTempStorage();
                 {
                     bool bFailedProperly = false;
                     var MissingFile = new List<string>(TestFiles.Select(TestFile => TestFile.FileName));
@@ -1269,31 +1283,20 @@ namespace AutomationTool
                 // store the test files to shared temp storage again.
                 TempStorage.StoreToTempStorage(TestTempStorageInfo, TestFiles.Select(TestFile => TestFile.FileName).ToList(), false, "UE4", CmdEnv.LocalRoot);
                 // force a load from shared by deleting the local manifest
-                TempStorage.DeleteLocalTempStorageManifests();
-                if (bZipTempStorage)
-                {
-                    var ZipFile = Path.ChangeExtension(TempStorage.SharedTempStorageManifestFilename(TestTempStorageInfo, "UE4"), "zip");
-                    // the shared storage file should exist.
-                    if (!CommandUtils.FileExists_NoExceptions(ZipFile))
-                    {
-                        throw new AutomationException("Shared file {0} did not exist", ZipFile);
-                    }
-                    // delete the shared file.
-                    CommandUtils.DeleteFile(ZipFile);
-                }
-                else
+                TempStorage.DeleteLocalTempStorage();
+                // delete our test files locally.
+                foreach (var TestFile in TestFiles)
                 {
-                    var SharedFile = CommandUtils.MakeRerootedFilePath(TestFiles[0].FileName, CmdEnv.LocalRoot, TempStorage.SharedTempStorageDirectory(TestTempStorageInfo, "UE4"));
-                    // the shared storage file should exist.
-                    if (!CommandUtils.FileExists_NoExceptions(SharedFile))
-                    {
-                        throw new AutomationException("Shared file {0} did not exist", SharedFile);
-                    }
-                    // delete the shared file.
-                    CommandUtils.DeleteFile(SharedFile);
+                    CommandUtils.DeleteFile(TestFile.FileName);
                 }
+
                 // now test that retrieving from shared temp storage properly balks that a file is missing.
                 {
+                    // tamper with the shared files.
+                    var RandomSharedZipFile = Directory.EnumerateFiles(TempStorage.SharedTempStorageDirectory(TestTempStorageInfo, "UE4"), "*.zip").First();
+                    // delete the shared file.
+                    CommandUtils.DeleteFile(RandomSharedZipFile);
+
                     bool bFailedProperly = false;
                     try
                     {
@@ -1308,6 +1311,11 @@ namespace AutomationTool
                         throw new AutomationException("Did not fail to load from missing file.");
                     }
                 }
+                // recreate our temp files.
+                foreach (var TestFile in TestFiles)
+                {
+                    File.WriteAllText(TestFile.FileName, TestFile.FileContents);
+                }
 
                 // clear the shared temp storage again.
                 TempStorage.DeleteSharedTempStorage(TestTempStorageInfo, "UE4");
@@ -1336,7 +1344,7 @@ namespace AutomationTool
             {
                 // Do a final cleanup.
                 TempStorage.DeleteSharedTempStorage(TestTempStorageInfo, "UE4");
-                TempStorage.DeleteLocalTempStorageManifests();
+                TempStorage.DeleteLocalTempStorage();
                 foreach (var TestFile in TestFiles)
                 {
                     CommandUtils.DeleteFile(TestFile.FileName);
