Remove precompiled binaries for BootstrapPackagedGame, and build it automatically whenever it's needed to package. Also fix CrashReportClient being built in Development rather than Shipping, and remove binaries for that.

diff --git a/Engine/Source/Programs/AutomationTool/Scripts/BuildProjectCommand.Automation.cs b/Engine/Source/Programs/AutomationTool/Scripts/BuildProjectCommand.Automation.cs
--- a/Engine/Source/Programs/AutomationTool/Scripts/BuildProjectCommand.Automation.cs
+++ b/Engine/Source/Programs/AutomationTool/Scripts/BuildProjectCommand.Automation.cs
@@ -88,14 +88,24 @@ public partial class Project : CommandUtils
 				}
 			}
 		}
+		if (!Params.NoBootstrapExe && !Params.Rocket)
+		{
+			UnrealBuildTool.UnrealTargetPlatform[] BootstrapPackagedGamePlatforms = { UnrealBuildTool.UnrealTargetPlatform.Win32, UnrealBuildTool.UnrealTargetPlatform.Win64 };
+			foreach(UnrealBuildTool.UnrealTargetPlatform BootstrapPackagedGamePlatform in BootstrapPackagedGamePlatforms)
+			{
+				if(Params.ClientTargetPlatforms.Contains(BootstrapPackagedGamePlatform))
+				{
+					Agenda.AddTarget("BootstrapPackagedGame", BootstrapPackagedGamePlatform, UnrealBuildTool.UnrealTargetConfiguration.Shipping);
+				}
+			}
+		}
 		if (Params.CrashReporter && !Params.Rocket)
 		{
-			var CrashReportClientTarget = new[] { "CrashReportClient" };
 			foreach (var CrashReportPlatform in CrashReportPlatforms)
 			{
 				if (UnrealBuildTool.UnrealBuildTool.PlatformSupportsCrashReporter(CrashReportPlatform))
 				{
-					Agenda.AddTargets(CrashReportClientTarget, CrashReportPlatform, UnrealTargetConfiguration.Development);
+					Agenda.AddTarget("CrashReportClient", CrashReportPlatform, UnrealTargetConfiguration.Shipping);
 				}
 			}
 		}
