Fixed chunking issue where secondary chunks wouldn't be generated because file paths didn't, because they were compaired with / without extensions.

diff --git a/Engine/Source/Programs/AutomationTool/Scripts/CopyBuildToStagingDirectory.Automation.cs b/Engine/Source/Programs/AutomationTool/Scripts/CopyBuildToStagingDirectory.Automation.cs
--- a/Engine/Source/Programs/AutomationTool/Scripts/CopyBuildToStagingDirectory.Automation.cs
+++ b/Engine/Source/Programs/AutomationTool/Scripts/CopyBuildToStagingDirectory.Automation.cs
@@ -937,7 +937,15 @@ public partial class Project : CommandUtils
 			bool bAddedToChunk = false;
 			for (int ChunkIndex = 0; !bAddedToChunk && ChunkIndex < ChunkResponseFiles.Length; ++ChunkIndex)
 			{
-				if (ChunkResponseFiles[ChunkIndex].Contains(StagingFile.Key) || ChunkResponseFiles[ChunkIndex].Contains(StagingFile.Key.Replace('/','\\')))
+                string OriginalFilename = StagingFile.Key;
+                string NoExtension = CombinePaths(Path.GetDirectoryName(OriginalFilename), Path.GetFileNameWithoutExtension(OriginalFilename));
+                string OriginalReplaceSlashes = OriginalFilename.Replace('/', '\\');
+                string NoExtensionReplaceSlashes = NoExtension.Replace('/', '\\');
+
+				if (ChunkResponseFiles[ChunkIndex].Contains(OriginalFilename) || 
+                    ChunkResponseFiles[ChunkIndex].Contains(OriginalReplaceSlashes) ||
+                    ChunkResponseFiles[ChunkIndex].Contains(NoExtension) ||
+                    ChunkResponseFiles[ChunkIndex].Contains(NoExtensionReplaceSlashes))
 				{
 					PakResponseFiles[ChunkIndex].Add(StagingFile.Key, StagingFile.Value);
 					bAddedToChunk = true;
