Adding more logs to help diagnose the mac timestamp failure.

diff --git a/Engine/Source/Programs/AutomationTool/GUBP/TempStorage.cs b/Engine/Source/Programs/AutomationTool/GUBP/TempStorage.cs
--- a/Engine/Source/Programs/AutomationTool/GUBP/TempStorage.cs
+++ b/Engine/Source/Programs/AutomationTool/GUBP/TempStorage.cs
@@ -952,6 +952,7 @@ namespace AutomationTool
                         {
                             // Use CommandUtils.CombinePaths to ensure directory separators get converted correctly. On mono on *nix, if the path has backslashes it will not convert it.
                             var ExtractedFilename = CommandUtils.CombinePaths(RootDir, Entry.FullName);
+                            CommandUtils.Log("{0}: Zip entry extracting to {1}.", Entry, ExtractedFilename);
                             // Zips can contain empty dirs. Ours usually don't have them, but we should support it.
                             if (Path.GetFileName(ExtractedFilename).Length == 0)
                             {
@@ -962,10 +963,12 @@ namespace AutomationTool
                                 // We must delete any existing file, even if it's readonly. .Net does not do this by default.
                                 if (File.Exists(ExtractedFilename))
                                 {
+                                    CommandUtils.Log("{0}: Destination already exists {1}. Deleting then extracting.", Entry, ExtractedFilename);
                                     InternalUtils.SafeDeleteFile(ExtractedFilename, true);
                                 }
                                 else
                                 {
+                                    CommandUtils.Log("{0}: Destination did not exist {1}. Extracting.", Entry, ExtractedFilename);
                                     Directory.CreateDirectory(Path.GetDirectoryName(ExtractedFilename));
                                 }
                                 Entry.ExtractToFile(ExtractedFilename, true);
@@ -1054,7 +1057,7 @@ namespace AutomationTool
                 var LocalManifest = LocalTempStorageManifestFilename(TempStorageNodeInfo);
                 if (CommandUtils.FileExists_NoExceptions(LocalManifest))
                 {
-                    CommandUtils.LogVerbose("Found local manifest {0}", LocalManifest);
+                    CommandUtils.Log("Found local manifest {0}", LocalManifest);
                     var Local = TempStorageManifest.Load(LocalManifest);
                     var Files = Local.GetFiles(RootDir);
                     var LocalTest = TempStorageManifest.Create(Files, RootDir);
@@ -1071,7 +1074,7 @@ namespace AutomationTool
                 // We couldn't find the node storage locally, so get it from the shared location.
                 var SharedStorageNodeDir = SharedTempStorageDirectory(TempStorageNodeInfo, GameName);
 
-                CommandUtils.LogVerbose("Attempting to retrieve from {0}", SharedStorageNodeDir);
+                CommandUtils.Log("Attempting to retrieve from {0}", SharedStorageNodeDir);
                 if (!CommandUtils.DirectoryExists_NoExceptions(SharedStorageNodeDir))
                 {
                     throw new AutomationException("Storage Block Does Not Exists! {0}", SharedStorageNodeDir);
