Don't add build.version as a build product; it's never meant to be checked in (legacy code from when we used to submit build.properties).

diff --git a/Engine/Source/Programs/AutomationTool/AutomationUtils/UE4Build.cs b/Engine/Source/Programs/AutomationTool/AutomationUtils/UE4Build.cs
--- a/Engine/Source/Programs/AutomationTool/AutomationUtils/UE4Build.cs
+++ b/Engine/Source/Programs/AutomationTool/AutomationUtils/UE4Build.cs
@@ -482,11 +482,6 @@ namespace AutomationTool
 				{
 					LogVerbose("{0} will not be updated because P4 is not enabled.", VerFile);
 				}
-				if (IsBuildMachine)
-				{
-					// Only BuildMachines can check this file in
-					AddBuildProduct(VerFile);
-				}
 				Result.Add(VerFile);
 			}
 
