Delegate archiving build products to temp storage to the BuildNode class, so legacy nodes can perform any fixups they need to.

diff --git a/Engine/Source/Programs/AutomationTool/GUBP/GUBP.cs b/Engine/Source/Programs/AutomationTool/GUBP/GUBP.cs
--- a/Engine/Source/Programs/AutomationTool/GUBP/GUBP.cs
+++ b/Engine/Source/Programs/AutomationTool/GUBP/GUBP.cs
@@ -1582,17 +1582,17 @@ public partial class GUBP : BuildCommand
 
             var TempStorageNodeInfo = new TempStorageNodeInfo(JobInfo, NodeToDo.Name);
             
-            string GameNameIfAny = NodeToDo.GameNameIfAnyForTempStorage;
-            string StorageRootIfAny = NodeToDo.RootIfAnyForTempStorage;
-					
-            if (bFake)
-            {
-                StorageRootIfAny = ""; // we don't rebase fake runs since those are entirely "records of success", which are always in the logs folder
-            }
-            if (string.IsNullOrEmpty(StorageRootIfAny))
-            {
-                StorageRootIfAny = CmdEnv.LocalRoot;
-            }
+			string GameNameIfAny = NodeToDo.GameNameIfAnyForTempStorage;
+			string StorageRootIfAny = NodeToDo.RootIfAnyForTempStorage;
+
+			if (bFake)
+			{
+				StorageRootIfAny = ""; // we don't rebase fake runs since those are entirely "records of success", which are always in the logs folder
+			}
+			if (string.IsNullOrEmpty(StorageRootIfAny))
+			{
+				StorageRootIfAny = CmdEnv.LocalRoot;
+			}
 
             // this is kinda complicated
             bool SaveSuccessRecords = (IsBuildMachine || bFakeEC) && // no real reason to make these locally except for fakeEC tests
@@ -1601,36 +1601,7 @@ public partial class GUBP : BuildCommand
             LogConsole("***** Running GUBP Node {0} -> {1} : {2}", NodeToDo.Name, GameNameIfAny, TempStorageNodeInfo.GetRelativeDirectory());
             if (NodeToDo.IsComplete)
             {
-                if (NodeToDo.Name == VersionFilesNode.StaticGetFullName() && !IsBuildMachine)
-                {
-                    LogConsole("***** NOT ****** Retrieving GUBP Node {0} from {1}; it is the version files.", NodeToDo.Name, TempStorageNodeInfo.GetRelativeDirectory());
-                    NodeToDo.Node.BuildProducts = new List<string>();
-
-                }
-                else
-                {
-                    LogConsole("***** Retrieving GUBP Node {0} from {1}", NodeToDo.Name, TempStorageNodeInfo.GetRelativeDirectory());
-                    bool WasLocal;
-					try
-					{
-                        NodeToDo.Node.BuildProducts = TempStorage.RetrieveFromTempStorage(TempStorageNodeInfo, out WasLocal, GameNameIfAny, StorageRootIfAny);
-					}
-					catch (Exception Ex)
-					{
-						if(GameNameIfAny != "")
-						{
-                            NodeToDo.Node.BuildProducts = TempStorage.RetrieveFromTempStorage(TempStorageNodeInfo, out WasLocal, "", StorageRootIfAny);
-						}
-						else
-						{
-							throw new AutomationException(Ex, "Build Products cannot be found for node {0}", NodeToDo.Name);
-						}
-					}
-                    if (!WasLocal)
-                    {
-                        NodeToDo.Node.PostLoadFromSharedTempStorage(this);
-                    }
-                }
+				NodeToDo.RetrieveBuildProducts(GameNameIfAny, StorageRootIfAny, TempStorageNodeInfo);
             }
             else
             {
diff --git a/Engine/Source/Programs/AutomationTool/GUBP/Nodes/BuildNode.cs b/Engine/Source/Programs/AutomationTool/GUBP/Nodes/BuildNode.cs
--- a/Engine/Source/Programs/AutomationTool/GUBP/Nodes/BuildNode.cs
+++ b/Engine/Source/Programs/AutomationTool/GUBP/Nodes/BuildNode.cs
@@ -49,6 +49,8 @@ namespace AutomationTool
 			get;
 		}
 
+		public abstract void RetrieveBuildProducts(string GameNameIfAny, string StorageRootIfAny, TempStorageNodeInfo TempStorageNodeInfo);
+
 		public abstract void DoBuild();
 
 		public abstract void DoFakeBuild();
@@ -123,6 +125,40 @@ namespace AutomationTool
 			get { return Node.FullNamesOfPseudodependencies; }
 		}
 
+		public override void RetrieveBuildProducts(string GameNameIfAny, string StorageRootIfAny, TempStorageNodeInfo TempStorageNodeInfo)
+		{
+			if (Name == GUBP.VersionFilesNode.StaticGetFullName() && !CommandUtils.IsBuildMachine)
+			{
+				CommandUtils.LogConsole("***** NOT ****** Retrieving GUBP Node {0} from {1}; it is the version files.", Name, TempStorageNodeInfo.GetRelativeDirectory());
+				Node.BuildProducts = new List<string>();
+
+			}
+			else
+			{
+				CommandUtils.LogConsole("***** Retrieving GUBP Node {0} from {1}", Name, TempStorageNodeInfo.GetRelativeDirectory());
+				bool WasLocal;
+				try
+				{
+					Node.BuildProducts = TempStorage.RetrieveFromTempStorage(TempStorageNodeInfo, out WasLocal, GameNameIfAny, StorageRootIfAny);
+				}
+				catch (Exception Ex)
+				{
+					if (GameNameIfAny != "")
+					{
+						Node.BuildProducts = TempStorage.RetrieveFromTempStorage(TempStorageNodeInfo, out WasLocal, "", StorageRootIfAny);
+					}
+					else
+					{
+						throw new AutomationException(Ex, "Build Products cannot be found for node {0}", Name);
+					}
+				}
+				if (!WasLocal)
+				{
+					Node.PostLoadFromSharedTempStorage(Owner);
+				}
+			}
+		}
+
 		public override void DoBuild()
 		{
 			Node.AllDependencyBuildProducts = new List<string>();
