Move archiving into a function on BuildNode.

diff --git a/Engine/Source/Programs/AutomationTool/GUBP/GUBP.cs b/Engine/Source/Programs/AutomationTool/GUBP/GUBP.cs
--- a/Engine/Source/Programs/AutomationTool/GUBP/GUBP.cs
+++ b/Engine/Source/Programs/AutomationTool/GUBP/GUBP.cs
@@ -1633,7 +1633,7 @@ public partial class GUBP : BuildCommand
 						BuildDuration = (DateTime.UtcNow - StartTime).TotalMilliseconds / 1000;
                     }
 
-                    TempStorage.StoreToTempStorage(TempStorageNodeInfo, NodeToDo.Node.BuildProducts, !bSaveSharedTempStorage, GameNameIfAny, StorageRootIfAny);
+					NodeToDo.ArchiveBuildProducts(GameNameIfAny, StorageRootIfAny, TempStorageNodeInfo, !bSaveSharedTempStorage);
                 }
                 catch (Exception Ex)
                 {
diff --git a/Engine/Source/Programs/AutomationTool/GUBP/Nodes/BuildNode.cs b/Engine/Source/Programs/AutomationTool/GUBP/Nodes/BuildNode.cs
--- a/Engine/Source/Programs/AutomationTool/GUBP/Nodes/BuildNode.cs
+++ b/Engine/Source/Programs/AutomationTool/GUBP/Nodes/BuildNode.cs
@@ -49,6 +49,8 @@ namespace AutomationTool
 			get;
 		}
 
+		public abstract void ArchiveBuildProducts(string GameNameIfAny, string StorageRootIfAny, TempStorageNodeInfo TempStorageNodeInfo, bool bLocalOnly);
+
 		public abstract void RetrieveBuildProducts(string GameNameIfAny, string StorageRootIfAny, TempStorageNodeInfo TempStorageNodeInfo);
 
 		public abstract void DoBuild();
@@ -125,6 +127,11 @@ namespace AutomationTool
 			get { return Node.FullNamesOfPseudodependencies; }
 		}
 
+		public override void ArchiveBuildProducts(string GameNameIfAny, string StorageRootIfAny, TempStorageNodeInfo TempStorageNodeInfo, bool bLocalOnly)
+		{
+			TempStorage.StoreToTempStorage(TempStorageNodeInfo, Node.BuildProducts, bLocalOnly, GameNameIfAny, StorageRootIfAny);
+		}
+
 		public override void RetrieveBuildProducts(string GameNameIfAny, string StorageRootIfAny, TempStorageNodeInfo TempStorageNodeInfo)
 		{
 			if (Name == GUBP.VersionFilesNode.StaticGetFullName() && !CommandUtils.IsBuildMachine)
