Remove the last GUBPNode dependency from GUBP.cs (the list of build products). Core GUBP now only deals with abstract BuildNode classes.

Also remove methods to handle removing stale DLLs; we don't need to do this now that binaries are not in Perforce, and remove the special-case checks for the VerisonFiles node.

diff --git a/Engine/Source/Programs/AutomationTool/GUBP/GUBP.cs b/Engine/Source/Programs/AutomationTool/GUBP/GUBP.cs
--- a/Engine/Source/Programs/AutomationTool/GUBP/GUBP.cs
+++ b/Engine/Source/Programs/AutomationTool/GUBP/GUBP.cs
@@ -1575,7 +1575,7 @@ public partial class GUBP : BuildCommand
         Dictionary<string, BuildNode> BuildProductToNodeMap = new Dictionary<string, BuildNode>();
 		foreach (BuildNode NodeToDo in OrderedToDo)
         {
-            if (NodeToDo.Node.BuildProducts != null)
+            if (NodeToDo.BuildProducts != null)
             {
                 throw new AutomationException("topological sort error");
             }
@@ -1737,7 +1737,7 @@ public partial class GUBP : BuildCommand
 					EC.UpdateECBuildTime(NodeToDo, BuildDuration);
                 }
             }
-            foreach (string Product in NodeToDo.Node.BuildProducts)
+            foreach (string Product in NodeToDo.BuildProducts)
             {
                 if (BuildProductToNodeMap.ContainsKey(Product))
                 {
diff --git a/Engine/Source/Programs/AutomationTool/GUBP/LegacyNodes.cs b/Engine/Source/Programs/AutomationTool/GUBP/LegacyNodes.cs
--- a/Engine/Source/Programs/AutomationTool/GUBP/LegacyNodes.cs
+++ b/Engine/Source/Programs/AutomationTool/GUBP/LegacyNodes.cs
@@ -45,9 +45,6 @@ partial class GUBP
             BuildProducts = new List<string>();
             SaveRecordOfSuccessAndAddToBuildProducts();
         }
-        public virtual void PostLoadFromSharedTempStorage(GUBP bp)
-        {
-        }
         public virtual void DoFakeBuild(GUBP bp) // this is used to more rapidly test a build system, it does nothing but save a record of success as a build product
         {
             BuildProducts = new List<string>();
@@ -199,8 +196,16 @@ partial class GUBP
 
         public override void DoBuild(GUBP bp)
         {
-            var UE4Build = new UE4Build(bp);
-            BuildProducts = UE4Build.UpdateVersionFiles(ActuallyUpdateVersionFiles: CommandUtils.P4Enabled && CommandUtils.AllowSubmit);
+			if (CommandUtils.P4Enabled && CommandUtils.AllowSubmit)
+			{
+				var UE4Build = new UE4Build(bp);
+				BuildProducts = UE4Build.UpdateVersionFiles(ActuallyUpdateVersionFiles: true);
+			}
+			else
+			{
+				BuildProducts = new List<string>();
+				SaveRecordOfSuccessAndAddToBuildProducts();
+			}
         }
         public override bool IsSticky()
         {
@@ -286,9 +291,6 @@ partial class GUBP
         public virtual void PostBuild(GUBP bp, UE4Build UE4Build)
         {
         }
-        public virtual void PostBuildProducts(GUBP bp)
-        {
-        }
         public virtual bool DeleteBuildProducts()
         {
             return false;
@@ -314,7 +316,6 @@ partial class GUBP
                     AddBuildProduct(Product);
                 }
                 RemoveOveralppingBuildProducts();
-                PostBuildProducts(bp);
             }
 			if (Agenda == null || (BuildProducts.Count == 0 && BranchConfig.bForceIncrementalCompile))
             {
@@ -489,76 +490,6 @@ partial class GUBP
             }
             return Agenda;
         }
-        void DeleteStaleDLLs(GUBP bp)
-        {
-			if (BranchConfig.bForceIncrementalCompile)
-            {
-                return;
-            }
-            var Targets = new List<string>{BranchConfig.Branch.BaseEngineProject.Properties.Targets[TargetRules.TargetType.Editor].TargetName};
-            foreach (var ProgramTarget in BranchConfig.Branch.BaseEngineProject.Properties.Programs)
-            {
-                if (ProgramTarget.Rules.GUBP_AlwaysBuildWithBaseEditor() && ProgramTarget.Rules.SupportsPlatform(HostPlatform))
-                {
-                    Targets.Add(ProgramTarget.TargetName);
-                }
-            }
-
-
-            foreach (var Target in Targets)
-            {
-                var EnginePlatformBinaries = CommandUtils.CombinePaths(CommandUtils.CmdEnv.LocalRoot, "Engine", "Binaries", HostPlatform.ToString());
-                var Wildcard = Target + "-*";
-                LogConsole("************Deleting stale editor DLLs, path {0} wildcard {1}", EnginePlatformBinaries, Wildcard);
-                foreach (var DiskFile in FindFiles(Wildcard, true, EnginePlatformBinaries))
-                {
-                    bool IsBuildProduct = false;
-                    foreach (var Product in BuildProducts)
-                    {
-                        if (Product.Equals(DiskFile, StringComparison.InvariantCultureIgnoreCase))
-                        {
-                            IsBuildProduct = true;
-                            break;
-                        }
-                    }
-                    if (!IsBuildProduct)
-                    {
-                        DeleteFile(DiskFile);
-                    }
-                }
-                var EnginePluginBinaries = CommandUtils.CombinePaths(CommandUtils.CmdEnv.LocalRoot, "Engine", "Plugins");
-                var HostSubstring = CommandUtils.CombinePaths("/", HostPlatform.ToString(), "/");
-                LogConsole("************Deleting stale editor DLLs, path {0} wildcard {1} host {2}", EnginePluginBinaries, Wildcard, HostSubstring);
-                foreach (var DiskFile in FindFiles(Wildcard, true, EnginePluginBinaries))
-                {
-                    if (DiskFile.IndexOf(HostSubstring, StringComparison.InvariantCultureIgnoreCase) < 0)
-                    {
-                        continue;
-                    }
-                    bool IsBuildProduct = false;
-                    foreach (var Product in BuildProducts)
-                    {
-                        if (Product.Equals(DiskFile, StringComparison.InvariantCultureIgnoreCase))
-                        {
-                            IsBuildProduct = true;
-                            break;
-                        }
-                    }
-                    if (!IsBuildProduct)
-                    {
-                        DeleteFile(DiskFile);
-                    }
-                }
-            }
-        }
-        public override void PostLoadFromSharedTempStorage(GUBP bp)
-        {
-            DeleteStaleDLLs(bp);
-        }
-        public override void PostBuildProducts(GUBP bp)
-        {
-            DeleteStaleDLLs(bp);
-        }
     }
 	public class RootEditorCrossCompileLinuxNode : CompileNode
 	{
@@ -597,76 +528,6 @@ partial class GUBP
             }
 			return Agenda;
 		}
-        void DeleteStaleDLLs(GUBP bp)
-        {
-			if (BranchConfig.bForceIncrementalCompile)
-            {
-                return;
-            }
-            var Targets = new List<string> { BranchConfig.Branch.BaseEngineProject.Properties.Targets[TargetRules.TargetType.Editor].TargetName };
-            foreach (var ProgramTarget in BranchConfig.Branch.BaseEngineProject.Properties.Programs)
-            {
-                if (ProgramTarget.Rules.GUBP_AlwaysBuildWithBaseEditor() && ProgramTarget.Rules.SupportsPlatform(UnrealTargetPlatform.Linux))
-                {
-                    Targets.Add(ProgramTarget.TargetName);
-                }
-            }
-
-
-            foreach (var Target in Targets)
-            {
-                var EnginePlatformBinaries = CommandUtils.CombinePaths(CommandUtils.CmdEnv.LocalRoot, "Engine", "Binaries", UnrealTargetPlatform.Linux.ToString());
-                var Wildcard = Target + "-*";
-                LogConsole("************Deleting stale editor DLLs, path {0} wildcard {1}", EnginePlatformBinaries, Wildcard);
-                foreach (var DiskFile in FindFiles(Wildcard, true, EnginePlatformBinaries))
-                {
-                    bool IsBuildProduct = false;
-                    foreach (var Product in BuildProducts)
-                    {
-                        if (Product.Equals(DiskFile, StringComparison.InvariantCultureIgnoreCase))
-                        {
-                            IsBuildProduct = true;
-                            break;
-                        }
-                    }
-                    if (!IsBuildProduct)
-                    {
-                        DeleteFile(DiskFile);
-                    }
-                }
-                var EnginePluginBinaries = CommandUtils.CombinePaths(CommandUtils.CmdEnv.LocalRoot, "Engine", "Plugins");
-                var HostSubstring = CommandUtils.CombinePaths("/", UnrealTargetPlatform.Linux.ToString(), "/");
-                LogConsole("************Deleting stale editor DLLs, path {0} wildcard {1} host {2}", EnginePluginBinaries, Wildcard, HostSubstring);
-                foreach (var DiskFile in FindFiles(Wildcard, true, EnginePluginBinaries))
-                {
-                    if (DiskFile.IndexOf(HostSubstring, StringComparison.InvariantCultureIgnoreCase) < 0)
-                    {
-                        continue;
-                    }
-                    bool IsBuildProduct = false;
-                    foreach (var Product in BuildProducts)
-                    {
-                        if (Product.Equals(DiskFile, StringComparison.InvariantCultureIgnoreCase))
-                        {
-                            IsBuildProduct = true;
-                            break;
-                        }
-                    }
-                    if (!IsBuildProduct)
-                    {
-                        DeleteFile(DiskFile);
-                    }
-                }
-            }
-        }
-        public override void PostLoadFromSharedTempStorage(GUBP bp)
-        {
-            DeleteStaleDLLs(bp);
-        }
-        public override void PostBuildProducts(GUBP bp)
-        {
-            DeleteStaleDLLs(bp);        
-		}
 	}
     public class ToolsNode : CompileNode
     {
diff --git a/Engine/Source/Programs/AutomationTool/GUBP/Nodes/BuildNode.cs b/Engine/Source/Programs/AutomationTool/GUBP/Nodes/BuildNode.cs
--- a/Engine/Source/Programs/AutomationTool/GUBP/Nodes/BuildNode.cs
+++ b/Engine/Source/Programs/AutomationTool/GUBP/Nodes/BuildNode.cs
@@ -19,7 +19,6 @@ namespace AutomationTool
 		public int TimeoutInMinutes;
 		public float Priority;
 
-		public GUBP.GUBPNode Node;
 		public HashSet<BuildNode> InputDependencies;
 		public HashSet<BuildNode> OrderDependencies;
 		public TriggerNode[] ControllingTriggers;
@@ -29,6 +28,8 @@ namespace AutomationTool
 		public bool SendSuccessEmail;
 		public bool IsParallelAgentShareEditor;
 
+		public List<string> BuildProducts;
+
 		public BuildNode(string InName)
 		{
 			Name = InName;
@@ -99,6 +100,7 @@ namespace AutomationTool
 	public class LegacyBuildNode : BuildNode
 	{
 		GUBP Owner;
+		GUBP.GUBPNode Node;
 
 		public LegacyBuildNode(GUBP InOwner, GUBP.GUBPNode InNode) : base(InNode.GetFullName())
 		{
@@ -129,41 +131,29 @@ namespace AutomationTool
 
 		public override void ArchiveBuildProducts(string GameNameIfAny, string StorageRootIfAny, TempStorageNodeInfo TempStorageNodeInfo, bool bLocalOnly)
 		{
-			TempStorage.StoreToTempStorage(TempStorageNodeInfo, Node.BuildProducts, bLocalOnly, GameNameIfAny, StorageRootIfAny);
+			TempStorage.StoreToTempStorage(TempStorageNodeInfo, BuildProducts, bLocalOnly, GameNameIfAny, StorageRootIfAny);
 		}
 
 		public override void RetrieveBuildProducts(string GameNameIfAny, string StorageRootIfAny, TempStorageNodeInfo TempStorageNodeInfo)
 		{
-			if (Name == GUBP.VersionFilesNode.StaticGetFullName() && !CommandUtils.IsBuildMachine)
+			CommandUtils.LogConsole("***** Retrieving GUBP Node {0} from {1}", Name, TempStorageNodeInfo.GetRelativeDirectory());
+			bool WasLocal;
+			try
 			{
-				CommandUtils.LogConsole("***** NOT ****** Retrieving GUBP Node {0} from {1}; it is the version files.", Name, TempStorageNodeInfo.GetRelativeDirectory());
-				Node.BuildProducts = new List<string>();
-
+				BuildProducts = TempStorage.RetrieveFromTempStorage(TempStorageNodeInfo, out WasLocal, GameNameIfAny, StorageRootIfAny);
 			}
-			else
+			catch (Exception Ex)
 			{
-				CommandUtils.LogConsole("***** Retrieving GUBP Node {0} from {1}", Name, TempStorageNodeInfo.GetRelativeDirectory());
-				bool WasLocal;
-				try
-				{
-					Node.BuildProducts = TempStorage.RetrieveFromTempStorage(TempStorageNodeInfo, out WasLocal, GameNameIfAny, StorageRootIfAny);
-				}
-				catch (Exception Ex)
+				if (GameNameIfAny != "")
 				{
-					if (GameNameIfAny != "")
-					{
-						Node.BuildProducts = TempStorage.RetrieveFromTempStorage(TempStorageNodeInfo, out WasLocal, "", StorageRootIfAny);
-					}
-					else
-					{
-						throw new AutomationException(Ex, "Build Products cannot be found for node {0}", Name);
-					}
+					BuildProducts = TempStorage.RetrieveFromTempStorage(TempStorageNodeInfo, out WasLocal, "", StorageRootIfAny);
 				}
-				if (!WasLocal)
+				else
 				{
-					Node.PostLoadFromSharedTempStorage(Owner);
+					throw new AutomationException(Ex, "Build Products cannot be found for node {0}", Name);
 				}
 			}
+			Node.BuildProducts = BuildProducts;
 		}
 
 		public override void DoBuild()
@@ -171,18 +161,20 @@ namespace AutomationTool
 			Node.AllDependencyBuildProducts = new List<string>();
 			foreach (BuildNode InputDependency in InputDependencies)
 			{
-				foreach (string BuildProduct in InputDependency.Node.BuildProducts)
+				foreach (string BuildProduct in InputDependency.BuildProducts)
 				{
 					Node.AddDependentBuildProduct(BuildProduct);
 				}
 			}
 
 			Node.DoBuild(Owner);
+			BuildProducts = Node.BuildProducts;
 		}
 
 		public override void DoFakeBuild()
 		{
 			Node.DoFakeBuild(Owner);
+			BuildProducts = Node.BuildProducts;
 		}
 
 		public override bool IsTest
