Adding documentation for SkipCookingEditorContent in UAT.

diff --git a/Engine/Source/Programs/AutomationTool/AutomationUtils/ProjectParams.cs b/Engine/Source/Programs/AutomationTool/AutomationUtils/ProjectParams.cs
--- a/Engine/Source/Programs/AutomationTool/AutomationUtils/ProjectParams.cs
+++ b/Engine/Source/Programs/AutomationTool/AutomationUtils/ProjectParams.cs
@@ -1135,6 +1135,7 @@ namespace AutomationTool
         /// <summary>
         /// Cook: Skip cooking editor content 
         /// </summary>
+		[Help("SkipCookingEditorContent", "Skips content under /Engine/Editor when cooking")]
         public bool SkipCookingEditorContent;
 
         /// <summary>
