Add a UAT command to sync and build a project from Perforce, similar to UGS.

diff --git a/Engine/Source/Programs/AutomationTool/AutomationUtils/P4Utils.cs b/Engine/Source/Programs/AutomationTool/AutomationUtils/P4Utils.cs
--- a/Engine/Source/Programs/AutomationTool/AutomationUtils/P4Utils.cs
+++ b/Engine/Source/Programs/AutomationTool/AutomationUtils/P4Utils.cs
@@ -2221,6 +2221,24 @@ namespace AutomationTool
 			return Records.ToArray();
 		}
 
+		/// <summary>
+		/// Determines whether a file exists in the depot.
+		/// </summary>
+		/// <param name="DepotFile">Depot path</param>
+		/// <returns>List of records describing the file's mapping. Usually just one, but may be more.</returns>
+		public bool FileExistsInDepot(string DepotFile, bool AllowSpew = true)
+		{
+			string CommandLine = String.Format("-z tag fstat {0}", CommandUtils.MakePathSafeToUseWithCommandLine(DepotFile));
+
+			string Output;
+			if(!LogP4Output(out Output, CommandLine, AllowSpew: false) || Output.Contains("no such file(s)"))
+			{
+				return false;
+			}
+
+			return true;
+		}
+
 		/// <summary>
 		/// Gets file stats.
 		/// </summary>
diff --git a/Engine/Source/Programs/AutomationTool/AutomationUtils/UE4Build.cs b/Engine/Source/Programs/AutomationTool/AutomationUtils/UE4Build.cs
--- a/Engine/Source/Programs/AutomationTool/AutomationUtils/UE4Build.cs
+++ b/Engine/Source/Programs/AutomationTool/AutomationUtils/UE4Build.cs
@@ -451,7 +451,7 @@ namespace AutomationTool
 			string ChangelistString = String.Empty;
 			if (bDoUpdateVersionFiles)
 			{
-				ChangelistNumber = ChangelistNumberOverride.GetValueOrDefault(P4Env.Changelist);
+				ChangelistNumber = ChangelistNumberOverride.HasValue? ChangelistNumberOverride.Value : P4Env.Changelist;
 				ChangelistString = ChangelistNumber.ToString();
 			}
 
