Increase verbosity level of FixUnixFilePermissions. Too spammy on build logs, and we're not debugging this stuff right now.

diff --git a/Engine/Source/Programs/AutomationTool/AutomationUtils/BuildUtils.cs b/Engine/Source/Programs/AutomationTool/AutomationUtils/BuildUtils.cs
--- a/Engine/Source/Programs/AutomationTool/AutomationUtils/BuildUtils.cs
+++ b/Engine/Source/Programs/AutomationTool/AutomationUtils/BuildUtils.cs
@@ -107,7 +107,7 @@ namespace AutomationTool
 		public static void FixUnixFilePermissions(string Filename)
 		{
 			string Permissions = IsProbablyAMacOrIOSExe(Filename) ? "0755" : "0644";
-			var Result = CommandUtils.Run("sh", string.Format("-c 'chmod {0} \"{1}\"'", Permissions, Filename.Replace("'", "'\"'\"'")));
+			var Result = CommandUtils.Run("sh", string.Format("-c 'chmod {0} \"{1}\"'", Permissions, Filename.Replace("'", "'\"'\"'")), Options:ERunOptions.SpewIsVerbose);
 			if (Result.ExitCode != 0)
 			{
 				throw new AutomationException(String.Format("Failed to chmod \"{0}\"", Filename));
