Fix Zip temp storage functionality on Macs.
* Mono doesn't correctly set the timestamp on files it archives, so we can't rely on the timestamp matching the manifest. Instead we still rely on filesizes, but we only check the timestamps when checking local manifests for tampering.
* Also reduced the logspam I added earlier to help debug this.
* Also fixed a bug in the manifest comparison code that only checked the sizes/timestamps of the first directory in the manfest. This bug has been around forever, so it may expose other issues, but preflight builds looked ok so far.

diff --git a/Engine/Source/Programs/AutomationTool/GUBP/TempStorage.cs b/Engine/Source/Programs/AutomationTool/GUBP/TempStorage.cs
--- a/Engine/Source/Programs/AutomationTool/GUBP/TempStorage.cs
+++ b/Engine/Source/Programs/AutomationTool/GUBP/TempStorage.cs
@@ -140,8 +140,9 @@ namespace AutomationTool
             /// Also uses a relaxed timestamp matching to allow for filesystems with limited granularity in their timestamps.
             /// </summary>
             /// <param name="Other"></param>
+            /// <param name="bCheckTimestamps">If true, checks timestamps of the two manifests. This is done when checking the local manifest against tampering.</param>
             /// <returns></returns>
-            public bool Compare(TempStorageFileInfo Other)
+            public bool Compare(TempStorageFileInfo Other, bool bCheckTimestamps)
             {
                 bool bOk = true;
                 if (!Name.Equals(Other.Name, StringComparison.InvariantCultureIgnoreCase))
@@ -188,7 +189,7 @@ namespace AutomationTool
 
                     // on FAT filesystems writetime has a two seconds resolution
                     // cf. http://msdn.microsoft.com/en-us/library/windows/desktop/ms724290%28v=vs.85%29.aspx
-                    if (!((Timestamp - Other.Timestamp).TotalSeconds < 2 && (Timestamp - Other.Timestamp).TotalSeconds > -2))
+                    if (bCheckTimestamps && !((Timestamp - Other.Timestamp).TotalSeconds < 2 && (Timestamp - Other.Timestamp).TotalSeconds > -2))
                     {
                         CommandUtils.LogWithVerbosity(LogType, "File date mismatch {0} {1} {2} {3}", Name, Timestamp, Other.Name, Other.Timestamp);
                         bOk = bOkToBeDifferent || bSilentOkToBeDifferent;
@@ -247,6 +248,7 @@ namespace AutomationTool
             {
                 var Directories = new Dictionary<string, List<TempStorageFileInfo>>();
 
+                CommandUtils.LogConsole("Creating TempStorageManifest at {0}", DateTime.UtcNow.ToString("o"));
                 foreach (string Filename in InFiles)
                 {
                     // use this to warm up the file shared on Mac.
@@ -277,8 +279,9 @@ namespace AutomationTool
             /// Compares this temp manifest with another one. It allows certain files that we know will differ to differe, while requiring the rest of the important details remain the same.
             /// </summary>
             /// <param name="Other"></param>
+            /// <param name="bCheckTimestamps">If true, checks timestamps of the two manifests. This is done when checking the local manifest against tampering.</param>
             /// <returns></returns>
-            public bool Compare(TempStorageManifest Other)
+            public bool Compare(TempStorageManifest Other, bool bCheckTimestamps)
             {
                 if (Directories.Count != Other.Directories.Count)
                 {
@@ -304,6 +307,7 @@ namespace AutomationTool
                     return false;
                 }
 
+                bool bResult = true;
                 foreach (KeyValuePair<string, List<TempStorageFileInfo>> Directory in Directories)
                 {
                     List<TempStorageFileInfo> OtherDirectory;
@@ -325,20 +329,17 @@ namespace AutomationTool
                         }
                         return false;
                     }
-                    bool bResult = true;
                     for (int FileIndex = 0; FileIndex < Directory.Value.Count; ++FileIndex)
                     {
                         TempStorageFileInfo File = Directory.Value[FileIndex];
                         TempStorageFileInfo OtherFile = OtherDirectory[FileIndex];
-                        if (File.Compare(OtherFile) == false)
+                        if (File.Compare(OtherFile, bCheckTimestamps) == false)
                         {
                             bResult = false;
                         }
                     }
-                    return bResult;
                 }
-
-                return true;
+                return bResult;
             }
 
             /// <summary>
@@ -359,7 +360,7 @@ namespace AutomationTool
                             DirElement => DirElement.Attribute(NameAttributeName).Value,
                             DirElement => DirElement.Elements(FileElementName).Select(FileElement => new TempStorageFileInfo
                                             {
-                                                Timestamp = new DateTime(long.Parse(FileElement.Attribute(TimestampAttributeName).Value)),
+                                                Timestamp = new DateTime(long.Parse(FileElement.Attribute(TimestampAttributeName).Value)/*, DateTimeKind.Utc*/),
                                                 Size = long.Parse(FileElement.Attribute(SizeAttributeName).Value),
                                                 Name = FileElement.Value,
                                             })
@@ -910,6 +911,8 @@ namespace AutomationTool
                             // use fastest compression. In our best case we are CPU bound, so this is a good tradeoff,
                             // cutting overall time by 2/3 while only modestly increasing the compression ratio (22.7% -> 23.8% for RootEditor PDBs).
                             // This is in cases of a super hot cache, so the operation was largely CPU bound.
+                            // Also, sadly, mono appears to have a bug where nothing you can do will properly set the LastWriteTime on the created entry,
+                            // so we have to ignore timestamps on files extracted from a zip, since it may have been created on a Mac.
                             ZipArchive.CreateEntryFromFile(File, CommandUtils.StripBaseDirectory(File, RootDir), CompressionLevel.Fastest);
                         }
                     }
@@ -952,7 +955,6 @@ namespace AutomationTool
                         {
                             // Use CommandUtils.CombinePaths to ensure directory separators get converted correctly. On mono on *nix, if the path has backslashes it will not convert it.
                             var ExtractedFilename = CommandUtils.CombinePaths(RootDir, Entry.FullName);
-                            CommandUtils.LogConsole("{0}: Zip entry extracting to {1}.", Entry.FullName, ExtractedFilename);
                             // Zips can contain empty dirs. Ours usually don't have them, but we should support it.
                             if (Path.GetFileName(ExtractedFilename).Length == 0)
                             {
@@ -963,12 +965,10 @@ namespace AutomationTool
                                 // We must delete any existing file, even if it's readonly. .Net does not do this by default.
                                 if (File.Exists(ExtractedFilename))
                                 {
-                                    CommandUtils.LogConsole("{0}: Destination already exists {1}. Deleting then extracting.", Entry.FullName, ExtractedFilename);
                                     InternalUtils.SafeDeleteFile(ExtractedFilename, true);
                                 }
                                 else
                                 {
-                                    CommandUtils.LogConsole("{0}: Destination did not exist {1}. Extracting.", Entry.FullName, ExtractedFilename);
                                     Directory.CreateDirectory(Path.GetDirectoryName(ExtractedFilename));
                                 }
                                 Entry.ExtractToFile(ExtractedFilename, true);
@@ -1061,7 +1061,7 @@ namespace AutomationTool
                     var Local = TempStorageManifest.Load(LocalManifest);
                     var Files = Local.GetFiles(RootDir);
                     var LocalTest = TempStorageManifest.Create(Files, RootDir);
-                    if (!Local.Compare(LocalTest))
+                    if (!Local.Compare(LocalTest, true))
                     {
                         throw new AutomationException("Local files in manifest {0} were tampered with.", LocalManifest);
                     }
@@ -1096,7 +1096,8 @@ namespace AutomationTool
 
                 var NewLocal = SaveLocalTempStorageManifest(RootDir, TempStorageNodeInfo, DestFiles);
                 // Now compare the created local files to ensure their attributes match the one we copied from the network.
-                if (!NewLocal.Compare(Shared))
+                // Don't check the timestamps because temp storage zips created on mono do not correctly store the timestamps.
+                if (!NewLocal.Compare(Shared, false))
                 {
                     // we will rename this so it can't be used, but leave it around for inspection
                     CommandUtils.RenameFile_NoExceptions(LocalManifest, LocalManifest + ".broken");
