Modified GetParamValueIfNotSpecified to return Default value if the param is not part of the command line

diff --git a/Engine/Source/Programs/AutomationTool/AutomationUtils/ProjectParams.cs b/Engine/Source/Programs/AutomationTool/AutomationUtils/ProjectParams.cs
--- a/Engine/Source/Programs/AutomationTool/AutomationUtils/ProjectParams.cs
+++ b/Engine/Source/Programs/AutomationTool/AutomationUtils/ProjectParams.cs
@@ -1,4 +1,4 @@
-﻿// Copyright 1998-2015 Epic Games, Inc. All Rights Reserved.
+// Copyright 1998-2015 Epic Games, Inc. All Rights Reserved.
 using System;
 using System.Collections.Generic;
 using System.Linq;
@@ -32,17 +32,15 @@ namespace AutomationTool
 			}
 			else if (Command != null)
 			{
-				bool Value = false;
 				foreach (var Param in ParamNames)
 				{
-					Value = Value || Command.ParseParam(Param);
+					if (Command.ParseParam(Param))
+					{
+						return true;
+					}
 				}
-				return Value;
-			}
-			else
-			{
-				return Default;
 			}
+			return Default;
 		}
 
 		/// <summary>
