Enabled diffing cooked builds for orion to help find deterministic cook issues.

diff --git a/Engine/Source/Programs/AutomationTool/Scripts/CookCommand.Automation.cs b/Engine/Source/Programs/AutomationTool/Scripts/CookCommand.Automation.cs
--- a/Engine/Source/Programs/AutomationTool/Scripts/CookCommand.Automation.cs
+++ b/Engine/Source/Programs/AutomationTool/Scripts/CookCommand.Automation.cs
@@ -299,26 +299,46 @@ public partial class Project : CommandUtils
 
             Platform CurrentPlatform = Params.GetTargetPlatformInstance(PlatformsToCook[CookPlatformIndex]);
 
-            if (Params.DiffCookedContentPath.EndsWith(".pak"))
-            {
-                string TemporaryPakFilename = CombinePaths( TemporaryPakPath, Path.GetFileName(Params.DiffCookedContentPath) );
-                File.Copy(Params.DiffCookedContentPath, TemporaryPakFilename);
-            }
-            else
+            string SourceCookedContentPath = Params.DiffCookedContentPath;
+
+            List<string> PakFiles = new List<string>();
+
+            if (Path.HasExtension(SourceCookedContentPath) && (!SourceCookedContentPath.EndsWith(".pak")))
             {
+                // must be a per platform pkg file try this
                 CurrentPlatform.ExtractPackage(Params, Params.DiffCookedContentPath, TemporaryPakPath);
+
+                // find the pak file
+                PakFiles = Directory.EnumerateFiles(TemporaryPakPath, "*.pak").ToList();
             }
 
-            
+            string CookPlatformString = CurrentPlatform.GetCookPlatform(false, Params.HasDedicatedServerAndClient, Params.CookFlavor);
 
-            string FullCookPath = CombinePaths(CookedSandboxesPath, PlatformsToCook[CookPlatformIndex].ToString());
+            if (!Path.HasExtension(SourceCookedContentPath))
+            {
+                // try find the pak or pkg file
+                string SourceCookedContentPlatformPath = CombinePaths(SourceCookedContentPath, CookPlatformString);
 
-            
+                foreach (var PakName in Directory.EnumerateFiles(SourceCookedContentPlatformPath, "*.pak"))
+                {
+                    string TemporaryPakFilename = CombinePaths(TemporaryPakPath, Path.GetFileName(PakName ));
+                    File.Copy(PakName , TemporaryPakFilename);
+                    PakFiles.Add(TemporaryPakFilename);
+                }
+            }
+            else if (SourceCookedContentPath.EndsWith(".pak"))
+            {
+                string TemporaryPakFilename = CombinePaths(TemporaryPakPath, Path.GetFileName(SourceCookedContentPath));
+                File.Copy(SourceCookedContentPath, TemporaryPakFilename);
+                PakFiles.Add(TemporaryPakFilename);
+            }
+
+
+            string FullCookPath = CombinePaths(CookedSandboxesPath, CookPlatformString);
 
             var UnrealPakExe = CombinePaths(CmdEnv.LocalRoot, "Engine/Binaries/Win64/UnrealPak.exe");
 
-            // find the pak file
-            var PakFiles = Directory.EnumerateFiles(TemporaryPakPath, "*.pak");
+            
             foreach (var Name in PakFiles)
             {
                 string UnrealPakParams = Name + " -Extract " + " " + TemporaryFilesPath;
@@ -326,6 +346,11 @@ public partial class Project : CommandUtils
                 RunAndLog(CmdEnv, UnrealPakExe, UnrealPakParams, Options: ERunOptions.Default | ERunOptions.UTF8Output);
             }
 
+            const string RootFailedContentDirectory = "\\\\epicgames.net\\root\\Developers\\Daniel.Lamb\\";
+
+            string FailedContentDirectory = CombinePaths( RootFailedContentDirectory, CommandUtils.P4Env.BuildRootP4, Params.ShortProjectName, CookPlatformString );
+
+
             // diff the content
             List<string> AllFiles = Directory.EnumerateFiles(FullCookPath, "*.uasset", System.IO.SearchOption.AllDirectories).ToList();
             AllFiles.AddRange(Directory.EnumerateFiles(FullCookPath, "*.map", System.IO.SearchOption.AllDirectories).ToList());
@@ -347,6 +372,11 @@ public partial class Project : CommandUtils
                         if ( SourceFile[Index] != DestFile[Index] )
                         {
 					        LogConsole("Diff cooked content failed on file " +SourceFilename + " when comparing against "+DestFilename + " at offset " + Index.ToString() );
+                            string SavedSourceFilename = CombinePaths( FailedContentDirectory, "Source" + Path.GetFileName(SourceFilename));
+                            string SavedDestFilename = CombinePaths( FailedContentDirectory, "Dest" + Path.GetFileName(DestFilename));
+                            File.Copy(SourceFilename, SavedSourceFilename);
+                            File.Copy(DestFilename, SavedDestFilename);
+                            LogConsole("Content temporarily saved to " +SavedSourceFilename + " and "+SavedDestFilename + " at offset " + Index.ToString() );
                             break;
                         }
                     }
