Changed the directory which failed diffed content is saved to.

diff --git a/Engine/Source/Programs/AutomationTool/Scripts/CookCommand.Automation.cs b/Engine/Source/Programs/AutomationTool/Scripts/CookCommand.Automation.cs
--- a/Engine/Source/Programs/AutomationTool/Scripts/CookCommand.Automation.cs
+++ b/Engine/Source/Programs/AutomationTool/Scripts/CookCommand.Automation.cs
@@ -348,7 +348,7 @@ public partial class Project : CommandUtils
 
             const string RootFailedContentDirectory = "\\\\epicgames.net\\root\\Developers\\Daniel.Lamb\\";
 
-            string FailedContentDirectory = CombinePaths( RootFailedContentDirectory, CommandUtils.P4Env.BuildRootP4, Params.ShortProjectName, CookPlatformString );
+            string FailedContentDirectory = CombinePaths( RootFailedContentDirectory, CommandUtils.P4Env.BuildRootP4 + CommandUtils.P4Env.ChangelistString, Params.ShortProjectName, CookPlatformString );
 
 
             // diff the content
