Change temp manifest verification to just simply ignore all timestamps of the set that it was ignoring before. Fixing the bug in manifest comparison to check all directories exposed a ton of additional areas where we were throwing warning on tamper checks, which was causing warnings on PC builds that never used to be there. There will still be warnings in other binaries folders (like plugins) that I am trying to suss out.
* Also removed unnecessary using clauses, which Resharper took as an opportunity to remove redundant qualifiers as well for some reason, thus the larger diff than necessary.

diff --git a/Engine/Source/Programs/AutomationTool/GUBP/TempStorage.cs b/Engine/Source/Programs/AutomationTool/GUBP/TempStorage.cs
--- a/Engine/Source/Programs/AutomationTool/GUBP/TempStorage.cs
+++ b/Engine/Source/Programs/AutomationTool/GUBP/TempStorage.cs
@@ -1,15 +1,15 @@
 ﻿// Copyright 1998-2015 Epic Games, Inc. All Rights Reserved.
+
 using System;
 using System.Collections.Concurrent;
 using System.Collections.Generic;
-using System.Text;
-using System.Linq;
 using System.IO;
-using System.Threading.Tasks;
-using System.Xml;
-using System.Xml.Linq;
 using System.IO.Compression;
+using System.Linq;
 using System.Threading;
+using System.Threading.Tasks;
+using System.Xml.Linq;
+using UnrealBuildTool;
 
 namespace AutomationTool
 {
@@ -79,8 +79,8 @@ namespace AutomationTool
 
         /// <summary>
         /// Turns the block info into a relative path to a folder where temp storage should be placed.
-        /// NOTE: <see cref="FindTempStorageNodeCLsMatchingSuffix"/> is dependent on this directory layout!!!
-        /// If this layout is changed, FindTempStorageNodeCLsMatchingSuffix needs to be changed as well.
+        /// NOTE: <see cref="TempStorage.FindMatchingSharedTempStorageNodeCLs"/> is dependent on this directory layout!!!
+        /// If this layout is changed, TempStorage.FindMatchingSharedTempStorageNodeCLs needs to be changed as well.
         /// </summary>
         /// <returns>{BranchNameForTempStorage}/{Changelist}{PreflightInfo}/{NodeStorageName}</returns>
         public string GetRelativeDirectory()
@@ -154,45 +154,21 @@ namespace AutomationTool
                 {
                     // this is a bit of a hack, but UAT itself creates these, so we need to allow them to be 
                     bool bOkToBeDifferent = Name.Contains("Engine/Binaries/DotNET/");
+                    // Lets just allow all mac warnings to be silent
+                    bOkToBeDifferent = bOkToBeDifferent || Name.Contains("Engine/Binaries/Mac/");
                     // this is a problem with mac compiles
                     bOkToBeDifferent = bOkToBeDifferent || Name.EndsWith("MacOS/libogg.dylib");
                     bOkToBeDifferent = bOkToBeDifferent || Name.EndsWith("MacOS/libvorbis.dylib");
                     bOkToBeDifferent = bOkToBeDifferent || Name.EndsWith("Contents/MacOS/UE4Editor");
 
-                    //temp hack until the mac build products work correctly
-                    bOkToBeDifferent = bOkToBeDifferent || Name.Contains("Engine/Binaries/Mac/UE4Editor.app/Contents/MacOS/");
-
-
-                    // DotNETUtilities.dll is built by a tons of other things
-                    bool bSilentOkToBeDifferent = (Name == "Engine/Binaries/DotNET/DotNETUtilities.dll");
-                    bSilentOkToBeDifferent = bSilentOkToBeDifferent || (Name == "Engine/Binaries/DotNET/DotNETUtilities.pdb");
-                    // RPCUtility is build by IPP and maybe other things
-                    bSilentOkToBeDifferent = bSilentOkToBeDifferent || (Name == "Engine/Binaries/DotNET/RPCUtility.exe");
-                    bSilentOkToBeDifferent = bSilentOkToBeDifferent || (Name == "Engine/Binaries/DotNET/RPCUtility.pdb");
-                    bSilentOkToBeDifferent = bSilentOkToBeDifferent || (Name == "Engine/Binaries/DotNET/AutomationTool.exe");
-                    bSilentOkToBeDifferent = bSilentOkToBeDifferent || (Name == "Engine/Binaries/DotNET/AutomationTool.exe.config");
-                    bSilentOkToBeDifferent = bSilentOkToBeDifferent || (Name == "Engine/Binaries/DotNET/AutomationUtils.Automation.dll");
-                    bSilentOkToBeDifferent = bSilentOkToBeDifferent || (Name == "Engine/Binaries/DotNET/AutomationUtils.Automation.pdb");
-                    bSilentOkToBeDifferent = bSilentOkToBeDifferent || (Name == "Engine/Binaries/DotNET/UnrealBuildTool.exe");
-                    bSilentOkToBeDifferent = bSilentOkToBeDifferent || (Name == "Engine/Binaries/DotNET/UnrealBuildTool.exe.config");					
-                    bSilentOkToBeDifferent = bSilentOkToBeDifferent || (Name == "Engine/Binaries/DotNET/EnvVarsToXML.exe");
-                    bSilentOkToBeDifferent = bSilentOkToBeDifferent || (Name == "Engine/Binaries/DotNET/EnvVarsToXML.exe.config");					
-
-                    // Lets just allow all mac warnings to be silent
-                    bSilentOkToBeDifferent = bSilentOkToBeDifferent || Name.Contains("Engine/Binaries/Mac");
-
-                    UnrealBuildTool.LogEventType LogType = bOkToBeDifferent ? UnrealBuildTool.LogEventType.Warning : UnrealBuildTool.LogEventType.Error;
-                    if (bSilentOkToBeDifferent)
-                    {
-                        LogType = UnrealBuildTool.LogEventType.Console;
-                    }
+                    LogEventType LogType = bOkToBeDifferent ? LogEventType.Console : LogEventType.Error;
 
                     // on FAT filesystems writetime has a two seconds resolution
                     // cf. http://msdn.microsoft.com/en-us/library/windows/desktop/ms724290%28v=vs.85%29.aspx
                     if (bCheckTimestamps && !((Timestamp - Other.Timestamp).TotalSeconds < 2 && (Timestamp - Other.Timestamp).TotalSeconds > -2))
                     {
                         CommandUtils.LogWithVerbosity(LogType, "File date mismatch {0} {1} {2} {3}", Name, Timestamp, Other.Name, Other.Timestamp);
-                        bOk = bOkToBeDifferent || bSilentOkToBeDifferent;
+                        bOk = bOkToBeDifferent;
                     }
                     if (Size != Other.Size)
                     {
@@ -595,7 +571,7 @@ namespace AutomationTool
         /// <param name="GameNames">List of game names to clean temp folders for.</param>
         public static void CleanSharedTempStorageDirectory(IEnumerable<string> GameNames)
         {
-            if (!CommandUtils.IsBuildMachine || UnrealBuildTool.Utils.IsRunningOnMono)  // saw a hang on this, anyway it isn't necessary to clean with macs, they are slow anyway
+            if (!CommandUtils.IsBuildMachine || Utils.IsRunningOnMono)  // saw a hang on this, anyway it isn't necessary to clean with macs, they are slow anyway
             {
                 return;
             }
@@ -634,7 +610,7 @@ namespace AutomationTool
                         return true;
                     }
                     CommandUtils.FindDirectories_NoExceptions(false, "*", false, UE4TempStorageDir); // there is some internet evidence to suggest this might perk up the mac share
-                    System.Threading.Thread.Sleep(5000);
+                    Thread.Sleep(5000);
                     Retries++;
                 }
             }
@@ -889,7 +865,7 @@ namespace AutomationTool
             long ZipFilesTotalSize = 0L;
             // This is to work around OOM errors on Mono on Macs. They apparently run out of memory more easily than a windows machine, so we can't go quite as wide as some of these files we
             // are zipping are really big.
-            var MaxParallelism = UnrealBuildTool.Utils.IsRunningOnMono ? Environment.ProcessorCount / 4 : Environment.ProcessorCount;
+            var MaxParallelism = Utils.IsRunningOnMono ? Environment.ProcessorCount / 4 : Environment.ProcessorCount;
             // We deliberately avoid Parallel.ForEach here because profiles have shown that dynamic partitioning creates
             // too many zip files, and they can be of too varying size, creating uneven work when unzipping later,
             // as ZipFile cannot unzip files in parallel from a single archive.
@@ -1105,7 +1081,7 @@ namespace AutomationTool
                 }
 
                 // Handle unix permissions/chmod issues. This will touch the timestamp we check on the file, so do this after we've compared with the manifest attributes.
-                if (UnrealBuildTool.Utils.IsRunningOnMono)
+                if (Utils.IsRunningOnMono)
                 {
                     foreach (string DestFile in DestFiles)
                     {
@@ -1161,9 +1137,9 @@ namespace AutomationTool
             var TestTempStorageInfo = new TempStorageNodeInfo(new GUBP.JobInfo("Test", 0, 0, 0), "Test");
 
             // Delete any local and shared temp storage that may exist.
-            TempStorage.DeleteLocalTempStorage();
-            TempStorage.DeleteSharedTempStorage(TestTempStorageInfo, "UE4");
-            if (TempStorage.TempStorageExists(TestTempStorageInfo, "UE4", false, false))
+            DeleteLocalTempStorage();
+            DeleteSharedTempStorage(TestTempStorageInfo, "UE4");
+            if (TempStorageExists(TestTempStorageInfo, "UE4", false, false))
             {
                 throw new AutomationException("storage should not exist");
             }
@@ -1194,35 +1170,35 @@ namespace AutomationTool
             try
             {
                 // Store the test file to temp storage.
-                TempStorage.StoreToTempStorage(TestTempStorageInfo, TestFiles.Select(TestFile => TestFile.FileName).ToList(), false, "UE4", CmdEnv.LocalRoot);
+                StoreToTempStorage(TestTempStorageInfo, TestFiles.Select(TestFile => TestFile.FileName).ToList(), false, "UE4", CmdEnv.LocalRoot);
                 // The manifest should exist locally.
-                if (!TempStorage.LocalTempStorageManifestExists(TestTempStorageInfo))
+                if (!LocalTempStorageManifestExists(TestTempStorageInfo))
                 {
                     throw new AutomationException("local storage should exist");
                 }
                 // The manifest should exist on the shared drive.
-                if (!TempStorage.SharedTempStorageManifestExists(TestTempStorageInfo, "UE4", false))
+                if (!SharedTempStorageManifestExists(TestTempStorageInfo, "UE4", false))
                 {
                     throw new AutomationException("shared storage should exist");
                 }
                 // Now delete the local manifest
-                TempStorage.DeleteLocalTempStorage();
+                DeleteLocalTempStorage();
                 // It should no longer be there.
-                if (TempStorage.LocalTempStorageManifestExists(TestTempStorageInfo))
+                if (LocalTempStorageManifestExists(TestTempStorageInfo))
                 {
                     throw new AutomationException("local storage should not exist");
                 }
                 // But the shared storage should still exist.
-                if (!TempStorage.TempStorageExists(TestTempStorageInfo, "UE4", false, false))
+                if (!TempStorageExists(TestTempStorageInfo, "UE4", false, false))
                 {
                     throw new AutomationException("some storage should exist");
                 }
 
                 // Now we should be able to retrieve the test files from shared storage, and it should overwrite our read-only files.
                 bool WasLocal;
-                TempStorage.RetrieveFromTempStorage(TestTempStorageInfo, out WasLocal, "UE4", CmdEnv.LocalRoot);
+                RetrieveFromTempStorage(TestTempStorageInfo, out WasLocal, "UE4", CmdEnv.LocalRoot);
                 // Now delete the local manifest so we can try again with no files present (the usual case for restoring from temp storage).
-                TempStorage.DeleteLocalTempStorage();
+                DeleteLocalTempStorage();
 
                 // Ok, delete our test files locally.
                 foreach (var TestFile in TestFiles)
@@ -1231,14 +1207,14 @@ namespace AutomationTool
                 }
 
                 // Now we should be able to retrieve the test files from shared storage.
-                TempStorage.RetrieveFromTempStorage(TestTempStorageInfo, out WasLocal, "UE4", CmdEnv.LocalRoot);
+                RetrieveFromTempStorage(TestTempStorageInfo, out WasLocal, "UE4", CmdEnv.LocalRoot);
                 // the local manifest should be there, since we just retrieved from shared storage.
-                if (!TempStorage.LocalTempStorageManifestExists(TestTempStorageInfo))
+                if (!LocalTempStorageManifestExists(TestTempStorageInfo))
                 {
                     throw new AutomationException("local storage should exist");
                 }
                 // The shared manifest should also still be there.
-                if (!TempStorage.SharedTempStorageManifestExists(TestTempStorageInfo, "UE4", false))
+                if (!SharedTempStorageManifestExists(TestTempStorageInfo, "UE4", false))
                 {
                     throw new AutomationException("shared storage should exist");
                 }
@@ -1251,21 +1227,21 @@ namespace AutomationTool
                     }
                 }
                 // Now delete the shared temp storage
-                TempStorage.DeleteSharedTempStorage(TestTempStorageInfo, "UE4");
+                DeleteSharedTempStorage(TestTempStorageInfo, "UE4");
                 // Shared temp storage manifest should no longer exist.
-                if (TempStorage.SharedTempStorageManifestExists(TestTempStorageInfo, "UE4", false))
+                if (SharedTempStorageManifestExists(TestTempStorageInfo, "UE4", false))
                 {
                     throw new AutomationException("shared storage should not exist");
                 }
                 // Retrieving temp storage should now just retrieve from local
-                TempStorage.RetrieveFromTempStorage(TestTempStorageInfo, out WasLocal, "UE4", CmdEnv.LocalRoot);
-                if (!WasLocal || !TempStorage.LocalTempStorageManifestExists(TestTempStorageInfo))
+                RetrieveFromTempStorage(TestTempStorageInfo, out WasLocal, "UE4", CmdEnv.LocalRoot);
+                if (!WasLocal || !LocalTempStorageManifestExists(TestTempStorageInfo))
                 {
                     throw new AutomationException("local storage should exist");
                 }
 
                 // and now lets test tampering
-                TempStorage.DeleteLocalTempStorage();
+                DeleteLocalTempStorage();
                 {
                     bool bFailedProperly = false;
                     var MissingFile = new List<string>(TestFiles.Select(TestFile => TestFile.FileName));
@@ -1273,7 +1249,7 @@ namespace AutomationTool
                     MissingFile.Add(CommandUtils.CombinePaths(CmdEnv.LocalRoot, "Engine", "SomeFileThatDoesntExist.txt"));
                     try
                     {
-                        TempStorage.StoreToTempStorage(TestTempStorageInfo, MissingFile, false, "UE4", CmdEnv.LocalRoot);
+                        StoreToTempStorage(TestTempStorageInfo, MissingFile, false, "UE4", CmdEnv.LocalRoot);
                     }
                     catch (AutomationException)
                     {
@@ -1286,11 +1262,11 @@ namespace AutomationTool
                 }
 
                 // clear the shared temp storage again.
-                TempStorage.DeleteSharedTempStorage(TestTempStorageInfo, "UE4");
+                DeleteSharedTempStorage(TestTempStorageInfo, "UE4");
                 // store the test files to shared temp storage again.
-                TempStorage.StoreToTempStorage(TestTempStorageInfo, TestFiles.Select(TestFile => TestFile.FileName).ToList(), false, "UE4", CmdEnv.LocalRoot);
+                StoreToTempStorage(TestTempStorageInfo, TestFiles.Select(TestFile => TestFile.FileName).ToList(), false, "UE4", CmdEnv.LocalRoot);
                 // force a load from shared by deleting the local manifest
-                TempStorage.DeleteLocalTempStorage();
+                DeleteLocalTempStorage();
                 // delete our test files locally.
                 foreach (var TestFile in TestFiles)
                 {
@@ -1300,14 +1276,14 @@ namespace AutomationTool
                 // now test that retrieving from shared temp storage properly balks that a file is missing.
                 {
                     // tamper with the shared files.
-                    var RandomSharedZipFile = Directory.EnumerateFiles(TempStorage.SharedTempStorageDirectory(TestTempStorageInfo, "UE4"), "*.zip").First();
+                    var RandomSharedZipFile = Directory.EnumerateFiles(SharedTempStorageDirectory(TestTempStorageInfo, "UE4"), "*.zip").First();
                     // delete the shared file.
                     CommandUtils.DeleteFile(RandomSharedZipFile);
 
                     bool bFailedProperly = false;
                     try
                     {
-                        TempStorage.RetrieveFromTempStorage(TestTempStorageInfo, out WasLocal, "UE4", CmdEnv.LocalRoot);
+                        RetrieveFromTempStorage(TestTempStorageInfo, out WasLocal, "UE4", CmdEnv.LocalRoot);
                     }
                     catch (AutomationException)
                     {
@@ -1325,9 +1301,9 @@ namespace AutomationTool
                 }
 
                 // clear the shared temp storage again.
-                TempStorage.DeleteSharedTempStorage(TestTempStorageInfo, "UE4");
+                DeleteSharedTempStorage(TestTempStorageInfo, "UE4");
                 // Copy the files to temp storage.
-                TempStorage.StoreToTempStorage(TestTempStorageInfo, TestFiles.Select(TestFile => TestFile.FileName).ToList(), false, "UE4", CmdEnv.LocalRoot);
+                StoreToTempStorage(TestTempStorageInfo, TestFiles.Select(TestFile => TestFile.FileName).ToList(), false, "UE4", CmdEnv.LocalRoot);
                 // Delete a local file.
                 CommandUtils.DeleteFile(TestFiles[0].FileName);
                 // retrieving from temp storage should use WasLocal, but should balk because a local file was deleted.
@@ -1335,7 +1311,7 @@ namespace AutomationTool
                     bool bFailedProperly = false;
                     try
                     {
-                        TempStorage.RetrieveFromTempStorage(TestTempStorageInfo, out WasLocal, "UE4", CmdEnv.LocalRoot);
+                        RetrieveFromTempStorage(TestTempStorageInfo, out WasLocal, "UE4", CmdEnv.LocalRoot);
                     }
                     catch (AutomationException)
                     {
@@ -1350,8 +1326,8 @@ namespace AutomationTool
             finally
             {
                 // Do a final cleanup.
-                TempStorage.DeleteSharedTempStorage(TestTempStorageInfo, "UE4");
-                TempStorage.DeleteLocalTempStorage();
+                DeleteSharedTempStorage(TestTempStorageInfo, "UE4");
+                DeleteLocalTempStorage();
                 foreach (var TestFile in TestFiles)
                 {
                     CommandUtils.DeleteFile(TestFile.FileName);
