Fix failed pull requests leaving the P4 workspace in a dirty state.

diff --git a/Engine/Source/Programs/AutomationTool/AutomationUtils/P4Utils.cs b/Engine/Source/Programs/AutomationTool/AutomationUtils/P4Utils.cs
--- a/Engine/Source/Programs/AutomationTool/AutomationUtils/P4Utils.cs
+++ b/Engine/Source/Programs/AutomationTool/AutomationUtils/P4Utils.cs
@@ -1273,10 +1273,10 @@ namespace AutomationTool
 		/// Invokes revert command.
 		/// </summary>
 		/// <param name="CommandLine">Commandline for the command.</param>
-		public void Revert(string CommandLine)
+		public void Revert(string CommandLine, bool AllowSpew = true)
 		{
 			CheckP4Enabled();
-			LogP4("revert " + CommandLine);
+			LogP4("revert " + CommandLine, AllowSpew: AllowSpew);
 		}
 
 		/// <summary>
