Turn Manifest verification errors into a warning until we can root out what set of redundant work is causing the files to change.

diff --git a/Engine/Source/Programs/AutomationTool/GUBP/TempStorage.cs b/Engine/Source/Programs/AutomationTool/GUBP/TempStorage.cs
--- a/Engine/Source/Programs/AutomationTool/GUBP/TempStorage.cs
+++ b/Engine/Source/Programs/AutomationTool/GUBP/TempStorage.cs
@@ -161,7 +161,7 @@ namespace AutomationTool
                     bOkToBeDifferent = bOkToBeDifferent || Name.EndsWith("MacOS/libvorbis.dylib");
                     bOkToBeDifferent = bOkToBeDifferent || Name.EndsWith("Contents/MacOS/UE4Editor");
 
-                    LogEventType LogType = bOkToBeDifferent ? LogEventType.Console : LogEventType.Error;
+                    LogEventType LogType = bOkToBeDifferent ? LogEventType.Console : LogEventType.Warning;
 
                     // on FAT filesystems writetime has a two seconds resolution
                     // cf. http://msdn.microsoft.com/en-us/library/windows/desktop/ms724290%28v=vs.85%29.aspx
