No, really, for real this time make timestamp mismatches in temp storage tamper check be a warning only.

diff --git a/Engine/Source/Programs/AutomationTool/GUBP/TempStorage.cs b/Engine/Source/Programs/AutomationTool/GUBP/TempStorage.cs
--- a/Engine/Source/Programs/AutomationTool/GUBP/TempStorage.cs
+++ b/Engine/Source/Programs/AutomationTool/GUBP/TempStorage.cs
@@ -168,7 +168,7 @@ namespace AutomationTool
                     if (bCheckTimestamps && !((Timestamp - Other.Timestamp).TotalSeconds < 2 && (Timestamp - Other.Timestamp).TotalSeconds > -2))
                     {
                         CommandUtils.LogWithVerbosity(LogType, "File date mismatch {0} {1} {2} {3}", Name, Timestamp, Other.Name, Other.Timestamp);
-                        bOk = bOkToBeDifferent;
+                        //bOk = bOkToBeDifferent;
                     }
                     if (Size != Other.Size)
                     {
