Added potential location for third party binaries to the build plugin command

diff --git a/Engine/Source/Programs/AutomationTool/Scripts/BuildPluginCommand.Automation.cs b/Engine/Source/Programs/AutomationTool/Scripts/BuildPluginCommand.Automation.cs
--- a/Engine/Source/Programs/AutomationTool/Scripts/BuildPluginCommand.Automation.cs
+++ b/Engine/Source/Programs/AutomationTool/Scripts/BuildPluginCommand.Automation.cs
@@ -156,6 +156,7 @@ class BuildPlugin : BuildCommand
 		FileFilter Filter = new FileFilter();
 		Filter.AddRuleForFile(PluginFileName, PluginDirectory, FileFilterType.Include);
 		Filter.AddRuleForFiles(BuildProducts.Select(x => x.Path), PluginDirectory, FileFilterType.Include);
+		Filter.Include("/Binaries/ThirdParty/...");
 		Filter.Include("/Resources/...");
 		Filter.Include("/Content/...");
 		Filter.Include("/Intermediate/Build/.../Inc/...");
