Move RequireFilesToExist out of target receipts; it's just internal state for UAT, and doesn't belong there.

diff --git a/Engine/Source/Programs/AutomationTool/AutomationUtils/DeploymentContext.cs b/Engine/Source/Programs/AutomationTool/AutomationUtils/DeploymentContext.cs
--- a/Engine/Source/Programs/AutomationTool/AutomationUtils/DeploymentContext.cs
+++ b/Engine/Source/Programs/AutomationTool/AutomationUtils/DeploymentContext.cs
@@ -14,6 +14,12 @@ public enum StagedFileType
 	DebugNonUFS
 }
 
+public struct StageTarget
+{
+	public TargetReceipt Receipt;
+	public bool RequireFilesExist;
+}
+
 public class DeploymentContext //: ProjectParams
 {
 	/// <summary>
@@ -52,9 +58,9 @@ public class DeploymentContext //: ProjectParams
 	public List<UnrealTargetConfiguration> StageTargetConfigurations;
 
 	/// <summary>
-	/// Receipts for the build targets that should be staged
+	/// Receipts for the build targets that should be staged.
 	/// </summary>
-	public List<TargetReceipt> StageTargetReceipts;
+	public List<StageTarget> StageTargets;
 
 	/// <summary>
 	///  this is the root directory that contains the engine: d:\a\UE4\
@@ -234,7 +240,7 @@ public class DeploymentContext //: ProjectParams
 		Platform InSourcePlatform,
         Platform InTargetPlatform,
 		List<UnrealTargetConfiguration> InTargetConfigurations,
-		IEnumerable<TargetReceipt> InStageTargetReceipts,
+		IEnumerable<StageTarget> InStageTargets,
 		List<String> InStageExecutables,
 		bool InServer,
 		bool InCooked,
@@ -254,7 +260,7 @@ public class DeploymentContext //: ProjectParams
         CookSourcePlatform = InSourcePlatform;
 		StageTargetPlatform = InTargetPlatform;
 		StageTargetConfigurations = new List<UnrealTargetConfiguration>(InTargetConfigurations);
-		StageTargetReceipts = new List<TargetReceipt>(InStageTargetReceipts);
+		StageTargets = new List<StageTarget>(InStageTargets);
 		StageExecutables = InStageExecutables;
         IsCodeBasedProject = ProjectUtils.IsCodeBasedUProjectFile(RawProjectPath);
 		ShortProjectName = ProjectUtils.GetShortProjectName(RawProjectPath);
@@ -385,13 +391,13 @@ public class DeploymentContext //: ProjectParams
 		}
 	}
 
-	public void StageBuildProductsFromReceipt(TargetReceipt Receipt)
+	public void StageBuildProductsFromReceipt(TargetReceipt Receipt, bool RequireDependenciesToExist)
 	{
 		// Stage all the build products needed at runtime
 		foreach(BuildProduct BuildProduct in Receipt.BuildProducts)
 		{
 			// allow missing files if needed
-			if (Receipt.bRequireDependenciesToExist == false && File.Exists(BuildProduct.Path) == false)
+			if (RequireDependenciesToExist == false && File.Exists(BuildProduct.Path) == false)
 			{
 				continue;
 			}
@@ -407,13 +413,13 @@ public class DeploymentContext //: ProjectParams
 		}
 	}
 
-	public void StageRuntimeDependenciesFromReceipt(TargetReceipt Receipt)
+	public void StageRuntimeDependenciesFromReceipt(TargetReceipt Receipt, bool RequireDependenciesToExist)
 	{
 		// Also stage any additional runtime dependencies, like ThirdParty DLLs
 		foreach(RuntimeDependency RuntimeDependency in Receipt.RuntimeDependencies)
 		{
 			// allow missing files if needed
-			if (Receipt.bRequireDependenciesToExist == false && File.Exists(RuntimeDependency.Path) == false)
+			if (RequireDependenciesToExist == false && File.Exists(RuntimeDependency.Path) == false)
 			{
 				continue;
 			}
diff --git a/Engine/Source/Programs/AutomationTool/Mac/MacPlatform.Automation.cs b/Engine/Source/Programs/AutomationTool/Mac/MacPlatform.Automation.cs
--- a/Engine/Source/Programs/AutomationTool/Mac/MacPlatform.Automation.cs
+++ b/Engine/Source/Programs/AutomationTool/Mac/MacPlatform.Automation.cs
@@ -174,9 +174,9 @@ public class MacPlatform : Platform
 		// Stage the bootstrap executable
 		if(!Params.NoBootstrapExe)
 		{
-			foreach(TargetReceipt Receipt in SC.StageTargetReceipts)
+			foreach(StageTarget Target in SC.StageTargets)
 			{
-				BuildProduct Executable = Receipt.BuildProducts.FirstOrDefault(x => x.Type == BuildProductType.Executable);
+				BuildProduct Executable = Target.Receipt.BuildProducts.FirstOrDefault(x => x.Type == BuildProductType.Executable);
 				if(Executable != null)
 				{
 					// only create bootstraps for executables
@@ -195,7 +195,7 @@ public class MacPlatform : Platform
 						}
 						else if(Params.IsCodeBasedProject)
 						{
-							BootstrapExeName = Receipt.GetProperty("TargetName", SC.ShortProjectName) + ".app";
+							BootstrapExeName = Target.Receipt.GetProperty("TargetName", SC.ShortProjectName) + ".app";
 						}
 						else
 						{
diff --git a/Engine/Source/Programs/AutomationTool/Scripts/CopyBuildToStagingDirectory.Automation.cs b/Engine/Source/Programs/AutomationTool/Scripts/CopyBuildToStagingDirectory.Automation.cs
--- a/Engine/Source/Programs/AutomationTool/Scripts/CopyBuildToStagingDirectory.Automation.cs
+++ b/Engine/Source/Programs/AutomationTool/Scripts/CopyBuildToStagingDirectory.Automation.cs
@@ -268,9 +268,9 @@ public partial class Project : CommandUtils
 
 
 		// Stage any extra runtime dependencies from the receipts
-		foreach(TargetReceipt Receipt in SC.StageTargetReceipts)
+		foreach(StageTarget Target in SC.StageTargets)
 		{
-			SC.StageRuntimeDependenciesFromReceipt(Receipt);
+			SC.StageRuntimeDependenciesFromReceipt(Target.Receipt, Target.RequireFilesExist);
 		}
 
 		// Get the build.properties file
@@ -1528,7 +1528,7 @@ public partial class Project : CommandUtils
 
 			string EngineDir = CommandUtils.CombinePaths(CommandUtils.CmdEnv.LocalRoot, "Engine");
 
-			List<TargetReceipt> TargetsToStage = new List<TargetReceipt>();
+			List<StageTarget> TargetsToStage = new List<StageTarget>();
 			foreach(string Target in ListToProcess)
 			{
 				foreach(UnrealTargetConfiguration Config in ConfigsToProcess)
@@ -1578,8 +1578,7 @@ public partial class Project : CommandUtils
 
 						// Convert the paths to absolute
 						Receipt.ExpandPathVariables(EngineDir, Path.GetDirectoryName(Params.RawProjectPath));
-						Receipt.SetDependenciesToBeRequired(bRequireStagedFilesToExist);
-						TargetsToStage.Add(Receipt);
+						TargetsToStage.Add(new StageTarget{ Receipt = Receipt, RequireFilesExist = bRequireStagedFilesToExist });
 					}
 				}
 			}
diff --git a/Engine/Source/Programs/AutomationTool/Win/WinPlatform.Automation.cs b/Engine/Source/Programs/AutomationTool/Win/WinPlatform.Automation.cs
--- a/Engine/Source/Programs/AutomationTool/Win/WinPlatform.Automation.cs
+++ b/Engine/Source/Programs/AutomationTool/Win/WinPlatform.Automation.cs
@@ -36,9 +36,9 @@ public abstract class BaseWinPlatform : Platform
 		}
 
 		// Stage all the build products
-		foreach(TargetReceipt Receipt in SC.StageTargetReceipts)
+		foreach(StageTarget Target in SC.StageTargets)
 		{
-			SC.StageBuildProductsFromReceipt(Receipt);
+			SC.StageBuildProductsFromReceipt(Target.Receipt, Target.RequireFilesExist);
 		}
 
 		// Copy the splash screen, windows specific
@@ -47,9 +47,9 @@ public abstract class BaseWinPlatform : Platform
 		// Stage the bootstrap executable
 		if(!Params.NoBootstrapExe)
 		{
-			foreach(TargetReceipt Receipt in SC.StageTargetReceipts)
+			foreach(StageTarget Target in SC.StageTargets)
 			{
-				BuildProduct Executable = Receipt.BuildProducts.FirstOrDefault(x => x.Type == BuildProductType.Executable);
+				BuildProduct Executable = Target.Receipt.BuildProducts.FirstOrDefault(x => x.Type == BuildProductType.Executable);
 				if(Executable != null)
 				{
 					// only create bootstraps for executables
@@ -68,7 +68,7 @@ public abstract class BaseWinPlatform : Platform
 						}
 						else if(Params.IsCodeBasedProject)
 						{
-							BootstrapExeName = Receipt.GetProperty("TargetName", SC.ShortProjectName) + ".exe";
+							BootstrapExeName = Target.Receipt.GetProperty("TargetName", SC.ShortProjectName) + ".exe";
 						}
 						else
 						{
