Limit parallelism on unzip for mac. Too much memory is being used apparently when it starts up this many threads. :(

diff --git a/Engine/Source/Programs/AutomationTool/GUBP/TempStorage.cs b/Engine/Source/Programs/AutomationTool/GUBP/TempStorage.cs
--- a/Engine/Source/Programs/AutomationTool/GUBP/TempStorage.cs
+++ b/Engine/Source/Programs/AutomationTool/GUBP/TempStorage.cs
@@ -920,7 +920,10 @@ namespace AutomationTool
         {
             var UnzipTimer = DateTimeStopwatch.Start();
             long ZipFilesTotalSize = 0L;
+            // Limit mac parallelism. Unzipping files seems to also be potentially causing too much memory commit.
+            var ParallelOptions = Utils.IsRunningOnMono ? new ParallelOptions { MaxDegreeOfParallelism = Environment.ProcessorCount / 4 } : new ParallelOptions();
             Parallel.ForEach(Directory.EnumerateFiles(FolderWithZipFiles, ZipBasename + "*.zip").ToList(),
+                ParallelOptions,
                 (ZipFilename) =>
                 {
                     Interlocked.Add(ref ZipFilesTotalSize, new FileInfo(ZipFilename).Length);
