Switching back to Ionic.Zip for temp storage saving because mono is running the macs out of memory.
* The problem is mono stages all files to a memory stream before copying it out, creating intense working memory pressure while creating the zips. By doing this all in parallel, we really hurt the working set size for large storage requests like RootEditor and DDC.
* Also turned back on manifest comparison for remote retrieval, now that timestamp comparison failures won't cause a failure. I'll use the warnings to track down build stages that are redundantly regenerating build artifacts, but we don't need to fail the builds on them.

diff --git a/Engine/Source/Programs/AutomationTool/GUBP/TempStorage.cs b/Engine/Source/Programs/AutomationTool/GUBP/TempStorage.cs
--- a/Engine/Source/Programs/AutomationTool/GUBP/TempStorage.cs
+++ b/Engine/Source/Programs/AutomationTool/GUBP/TempStorage.cs
@@ -4,12 +4,12 @@ using System;
 using System.Collections.Concurrent;
 using System.Collections.Generic;
 using System.IO;
-using System.IO.Compression;
 using System.Linq;
 using System.Threading;
 using System.Threading.Tasks;
 using System.Xml.Linq;
 using UnrealBuildTool;
+using Ionic.Zip;
 
 namespace AutomationTool
 {
@@ -168,6 +168,7 @@ namespace AutomationTool
                     if (bCheckTimestamps && !((Timestamp - Other.Timestamp).TotalSeconds < 2 && (Timestamp - Other.Timestamp).TotalSeconds > -2))
                     {
                         CommandUtils.LogWithVerbosity(LogType, "File date mismatch {0} {1} {2} {3}", Name, Timestamp, Other.Name, Other.Timestamp);
+                        // WRH 2015/08/19 - don't actually fail any builds due to a timestamp check. there's lots of reasons why the timestamps are not preserved right now, mostly due to zipfile stuff.
                         //bOk = bOkToBeDifferent;
                     }
                     if (Size != Other.Size)
@@ -863,21 +864,18 @@ namespace AutomationTool
             var FilesToZip = new ConcurrentQueue<string>(FilesInfo.OrderByDescending(FileInfo => FileInfo.FileSize).Select(FileInfo => FileInfo.File));
 
             long ZipFilesTotalSize = 0L;
-            // This is to work around OOM errors on Mono on Macs. They apparently run out of memory more easily than a windows machine, so we can't go quite as wide as some of these files we
-            // are zipping are really big.
-            var MaxParallelism = Utils.IsRunningOnMono ? Environment.ProcessorCount / 4 : Environment.ProcessorCount;
             // We deliberately avoid Parallel.ForEach here because profiles have shown that dynamic partitioning creates
             // too many zip files, and they can be of too varying size, creating uneven work when unzipping later,
             // as ZipFile cannot unzip files in parallel from a single archive.
             // We can safely assume the build system will not be doing more important things at the same time, so we simply use all our logical cores,
             // which has shown to be optimal via profiling, and limits the number of resulting zip files to the number of logical cores.
             var Threads = (
-                from CoreNum in Enumerable.Range(0, bZipInParallel ? MaxParallelism : 1)
+                from CoreNum in Enumerable.Range(0, bZipInParallel ? Environment.ProcessorCount : 1)
                 let ZipFileName = Path.Combine(StagingDir ?? OutputDir, string.Format("{0}{1}.zip", ZipBasename, bZipInParallel ? "-" + CoreNum.ToString("00") : ""))
                 select new Thread(() =>
                 {
                     // Create one zip per thread using the given basename
-                    using (var ZipArchive = ZipFile.Open(ZipFileName, ZipArchiveMode.Create))
+                    using (var ZipArchive = new ZipFile(ZipFileName) { CompressionLevel = Ionic.Zlib.CompressionLevel.BestSpeed })
                     {
 
                         // pull from the queue until we are out of files.
@@ -887,10 +885,9 @@ namespace AutomationTool
                             // use fastest compression. In our best case we are CPU bound, so this is a good tradeoff,
                             // cutting overall time by 2/3 while only modestly increasing the compression ratio (22.7% -> 23.8% for RootEditor PDBs).
                             // This is in cases of a super hot cache, so the operation was largely CPU bound.
-                            // Also, sadly, mono appears to have a bug where nothing you can do will properly set the LastWriteTime on the created entry,
-                            // so we have to ignore timestamps on files extracted from a zip, since it may have been created on a Mac.
-                            ZipArchive.CreateEntryFromFile(File, CommandUtils.StripBaseDirectory(File, RootDir), CompressionLevel.Fastest);
+                            ZipArchive.AddFile(File, Path.GetDirectoryName(CommandUtils.StripBaseDirectory(File, RootDir)));
                         }
+                        ZipArchive.Save();
                     }
                     Interlocked.Add(ref ZipFilesTotalSize, new FileInfo(ZipFileName).Length);
                     // if we are using a staging dir, copy to the final location and delete the staged copy.
@@ -920,39 +917,13 @@ namespace AutomationTool
         {
             var UnzipTimer = DateTimeStopwatch.Start();
             long ZipFilesTotalSize = 0L;
-            // Limit mac parallelism. Unzipping files seems to also be potentially causing too much memory commit.
-            var ParallelOptions = Utils.IsRunningOnMono ? new ParallelOptions { MaxDegreeOfParallelism = Environment.ProcessorCount / 4 } : new ParallelOptions();
             Parallel.ForEach(Directory.EnumerateFiles(FolderWithZipFiles, ZipBasename + "*.zip").ToList(),
-                ParallelOptions,
                 (ZipFilename) =>
                 {
                     Interlocked.Add(ref ZipFilesTotalSize, new FileInfo(ZipFilename).Length);
-                    // unzip the files manually instead of caling ZipFile.ExtractToDirectory() because we need to overwrite readonly files. Because of this, creating the directories is up to us as well.
-                    using (var ZipArchive = ZipFile.OpenRead(ZipFilename))
+                    using (var ZipArchive = ZipFile.Read(ZipFilename))
                     {
-                        foreach (var Entry in ZipArchive.Entries)
-                        {
-                            // Use CommandUtils.CombinePaths to ensure directory separators get converted correctly. On mono on *nix, if the path has backslashes it will not convert it.
-                            var ExtractedFilename = CommandUtils.CombinePaths(RootDir, Entry.FullName);
-                            // Zips can contain empty dirs. Ours usually don't have them, but we should support it.
-                            if (Path.GetFileName(ExtractedFilename).Length == 0)
-                            {
-                                Directory.CreateDirectory(ExtractedFilename);
-                            }
-                            else
-                            {
-                                // We must delete any existing file, even if it's readonly. .Net does not do this by default.
-                                if (File.Exists(ExtractedFilename))
-                                {
-                                    InternalUtils.SafeDeleteFile(ExtractedFilename, true);
-                                }
-                                else
-                                {
-                                    Directory.CreateDirectory(Path.GetDirectoryName(ExtractedFilename));
-                                }
-                                Entry.ExtractToFile(ExtractedFilename, true);
-                            }
-                        }
+                        ZipArchive.ExtractAll(RootDir, ExtractExistingFileAction.OverwriteSilently);
                     }
                 });
             return new ParallelZipResult(UnzipTimer.ElapsedTime, ZipFilesTotalSize);
@@ -1075,8 +1046,7 @@ namespace AutomationTool
 
                 var NewLocal = SaveLocalTempStorageManifest(RootDir, TempStorageNodeInfo, DestFiles);
                 // Now compare the created local files to ensure their attributes match the one we copied from the network.
-                // Don't check the timestamps because temp storage zips created on mono do not correctly store the timestamps.
-                if (!NewLocal.Compare(Shared, false))
+                if (!NewLocal.Compare(Shared, true))
                 {
                     // we will rename this so it can't be used, but leave it around for inspection
                     CommandUtils.RenameFile_NoExceptions(LocalManifest, LocalManifest + ".broken");
