Reduce temp storage spam. Print UTC time in low when UAT starts.

diff --git a/Engine/Source/Programs/AutomationTool/GUBP/TempStorage.cs b/Engine/Source/Programs/AutomationTool/GUBP/TempStorage.cs
--- a/Engine/Source/Programs/AutomationTool/GUBP/TempStorage.cs
+++ b/Engine/Source/Programs/AutomationTool/GUBP/TempStorage.cs
@@ -225,7 +225,6 @@ namespace AutomationTool
             {
                 var Directories = new Dictionary<string, List<TempStorageFileInfo>>();
 
-                CommandUtils.LogConsole("Creating TempStorageManifest at {0}", DateTime.UtcNow.ToString("o"));
                 foreach (string Filename in InFiles)
                 {
                     // use this to warm up the file shared on Mac.
diff --git a/Engine/Source/Programs/AutomationTool/Program.cs b/Engine/Source/Programs/AutomationTool/Program.cs
--- a/Engine/Source/Programs/AutomationTool/Program.cs
+++ b/Engine/Source/Programs/AutomationTool/Program.cs
@@ -29,7 +29,7 @@ namespace AutomationTool
                 AssemblyUtils.InstallAssemblyResolver(Path.GetDirectoryName(Assembly.GetEntryAssembly().GetOriginalLocation()));
                 HostPlatform.Initialize();
 
-                Log.TraceVerbose("Running on {0} as a {1}-bit process.", HostPlatform.Current.GetType().Name, Environment.Is64BitProcess ? 64 : 32);
+                Log.TraceInformation("{2}: Running on {0} as a {1}-bit process.", HostPlatform.Current.GetType().Name, Environment.Is64BitProcess ? 64 : 32, DateTime.UtcNow.ToString("o"));
 
                 XmlConfigLoader.Init();
 
