Add a BuildLocalization node, which runs the UAT Localise script.

diff --git a/Engine/Source/Programs/AutomationTool/GUBP/LegacyBranchSetup.cs b/Engine/Source/Programs/AutomationTool/GUBP/LegacyBranchSetup.cs
--- a/Engine/Source/Programs/AutomationTool/GUBP/LegacyBranchSetup.cs
+++ b/Engine/Source/Programs/AutomationTool/GUBP/LegacyBranchSetup.cs
@@ -1114,6 +1114,11 @@ partial class GUBP
 			}
 		}
 
+		if(HostPlatforms.Contains(UnrealTargetPlatform.Win64) && BranchConfig.HasNode(RootEditorNode.StaticGetFullName(UnrealTargetPlatform.Win64)))
+		{
+			BranchConfig.AddNode(new BuildLocalization());
+		}
+
         BranchConfig.AddNode(new WaitForTestShared(this));
 		BranchConfig.AddNode(new WaitToPackageSamplesNode(BranchConfig.HostPlatforms));
 
diff --git a/Engine/Source/Programs/AutomationTool/GUBP/LegacyNodes.cs b/Engine/Source/Programs/AutomationTool/GUBP/LegacyNodes.cs
--- a/Engine/Source/Programs/AutomationTool/GUBP/LegacyNodes.cs
+++ b/Engine/Source/Programs/AutomationTool/GUBP/LegacyNodes.cs
@@ -2597,6 +2597,40 @@ partial class GUBP
         }
     }
 
+	public class BuildLocalization : HostPlatformNode
+	{
+		public BuildLocalization()
+			: base(UnrealTargetPlatform.Win64)
+		{
+			AddDependency(RootEditorNode.StaticGetFullName(HostPlatform));
+		}
+
+		public override void DoBuild(GUBP bp)
+		{
+			// Run the localise script.
+			//CommandUtils.RunUAT(CommandUtils.CmdEnv, "Localise");
+
+			// Don't pass on any build products to other build nodes at the moment.
+			BuildProducts = new List<string>();
+			SaveRecordOfSuccessAndAddToBuildProducts();
+		}
+
+		public override int CISFrequencyQuantumShift(GUBP.GUBPBranchConfig BranchConfig)
+		{
+			return base.CISFrequencyQuantumShift(BranchConfig) + 6;
+		}
+
+		public static string StaticGetFullName()
+		{
+			return "BuildLocalization";
+		}
+
+		public override string GetFullName()
+		{
+			return StaticGetFullName();
+		}
+	}
+
     public class GameAggregateNode : HostPlatformAggregateNode
     {
         BranchInfo.BranchUProject GameProj;
diff --git a/Engine/Source/Programs/AutomationTool/Scripts/Localisation.Automation.cs b/Engine/Source/Programs/AutomationTool/Scripts/Localisation.Automation.cs
--- a/Engine/Source/Programs/AutomationTool/Scripts/Localisation.Automation.cs
+++ b/Engine/Source/Programs/AutomationTool/Scripts/Localisation.Automation.cs
@@ -30,15 +30,6 @@ class Localise : BuildCommand
 	{
 		var EditorExe = CombinePaths(CmdEnv.LocalRoot, @"Engine/Binaries/Win64/UE4Editor-Cmd.exe");
 
-		if (P4Enabled)
-		{
-			LogConsole("Sync necessary content to head revision");
-			P4.Sync(P4Env.BuildRootP4 + "/Engine/Config/...");
-			P4.Sync(P4Env.BuildRootP4 + "/Engine/Content/...");
-			P4.Sync(P4Env.BuildRootP4 + "/Engine/Source/...");
-			LogConsole("Localize from label {0}", P4Env.LabelToSync);
-		}
-
         OneSkyConfigData OneSkyConfig = OneSkyConfigHelper.Find("OneSkyConfig_EpicGames");
         var oneSkyService = new OneSkyService(OneSkyConfig.ApiKey, OneSkyConfig.ApiSecret);
         var projectGroup = GetProjectGroup(oneSkyService, "Unreal Engine");
