Separate node templates from their implementations. Prelude to switching towards the build graph being specified through markup.

diff --git a/Engine/Source/Programs/AutomationTool/GUBP/GUBP.Automation.csproj b/Engine/Source/Programs/AutomationTool/GUBP/GUBP.Automation.csproj
--- a/Engine/Source/Programs/AutomationTool/GUBP/GUBP.Automation.csproj
+++ b/Engine/Source/Programs/AutomationTool/GUBP/GUBP.Automation.csproj
@@ -66,6 +66,7 @@
     <Compile Include="LegacyNodes.cs" />
     <Compile Include="Nodes\AggregateNode.cs" />
     <Compile Include="Nodes\BuildNode.cs" />
+    <Compile Include="Nodes\LegacyNode.cs" />
     <Compile Include="Nodes\TriggerNode.cs" />
     <Compile Include="TempStorage.cs">
       <SubType>Code</SubType>
diff --git a/Engine/Source/Programs/AutomationTool/GUBP/GUBP.cs b/Engine/Source/Programs/AutomationTool/GUBP/GUBP.cs
--- a/Engine/Source/Programs/AutomationTool/GUBP/GUBP.cs
+++ b/Engine/Source/Programs/AutomationTool/GUBP/GUBP.cs
@@ -223,15 +223,18 @@ public partial class GUBP : BuildCommand
 			ExplicitTriggerName = ParseParamValue("TriggerNode", "");
 		}
 
-		List<BuildNode> AllNodes;
-		List<AggregateNode> AllAggregates;
+		List<BuildNodeTemplate> AllBuildNodeTemplates;
+		List<AggregateNodeTemplate> AllAggregateNodeTemplates;
 		int TimeQuantum = 20;
-		AddNodesForBranch(HostPlatforms, JobInfo, out AllNodes, out AllAggregates, ref TimeQuantum);
+		AddNodesForBranch(HostPlatforms, JobInfo, out AllBuildNodeTemplates, out AllAggregateNodeTemplates, ref TimeQuantum);
 
-		LinkGraph(AllAggregates, AllNodes);
-		FindControllingTriggers(AllNodes);
-        FindCompletionState(AllNodes, JobInfo, LocalOnly);
-		ComputeDependentFrequencies(AllNodes);
+		List<AggregateNode> AllAggregateNodes;
+		List<BuildNode> AllBuildNodes;
+		LinkGraph(AllAggregateNodeTemplates, AllBuildNodeTemplates, out AllAggregateNodes, out AllBuildNodes);
+
+		FindControllingTriggers(AllBuildNodes);
+        FindCompletionState(AllBuildNodes, JobInfo, LocalOnly);
+		ComputeDependentFrequencies(AllBuildNodes);
 
         if (bCleanLocalTempStorage)  // shared temp storage can never be wiped
         {
@@ -249,7 +252,7 @@ public partial class GUBP : BuildCommand
 			Log("Setting TimeIndex to {0}", TimeIndex);
 		}
 
-		HashSet<BuildNode> NodesToDo = ParseNodesToDo(AllNodes, AllAggregates);
+		HashSet<BuildNode> NodesToDo = ParseNodesToDo(AllBuildNodes, AllAggregateNodes);
 		CullNodesForTimeIndex(NodesToDo, TimeIndex);
         if (JobInfo.IsPreflight)
         {
@@ -261,7 +264,7 @@ public partial class GUBP : BuildCommand
         {
             if (!String.IsNullOrEmpty(ExplicitTriggerName))
             {
-                foreach (TriggerNode Trigger in AllNodes.OfType<TriggerNode>())
+                foreach (TriggerNode Trigger in AllBuildNodes.OfType<TriggerNode>())
                 {
 					if (Trigger.Name.Equals(ExplicitTriggerName, StringComparison.InvariantCultureIgnoreCase))
                     {
@@ -279,7 +282,7 @@ public partial class GUBP : BuildCommand
             {
                 if (bSkipTriggers)
                 {
-                    foreach (TriggerNode Trigger in AllNodes.OfType<TriggerNode>())
+                    foreach (TriggerNode Trigger in AllBuildNodes.OfType<TriggerNode>())
                     {
 						Trigger.Activate();
                     }
@@ -291,7 +294,7 @@ public partial class GUBP : BuildCommand
 
 		List<TriggerNode> UnfinishedTriggers = FindUnfinishedTriggers(bSkipTriggers, ExplicitTrigger, OrderedToDo);
 
-		PrintNodes(this, OrderedToDo, AllAggregates, UnfinishedTriggers, TimeQuantum);
+		PrintNodes(this, OrderedToDo, AllAggregateNodes, UnfinishedTriggers, TimeQuantum);
 
         //check sorting
 		CheckSortOrder(OrderedToDo);
@@ -301,7 +304,7 @@ public partial class GUBP : BuildCommand
 		string ShowHistoryParam = ParseParamValue("ShowHistory", null);
 		if(ShowHistoryParam != null)
 		{
-			BuildNode Node = AllNodes.FirstOrDefault(x => x.Name.Equals(ShowHistoryParam, StringComparison.InvariantCultureIgnoreCase));
+			BuildNode Node = AllBuildNodes.FirstOrDefault(x => x.Name.Equals(ShowHistoryParam, StringComparison.InvariantCultureIgnoreCase));
 			if(Node == null)
 			{
 				throw new AutomationException("Couldn't find node {0}", ShowHistoryParam);
@@ -320,7 +323,7 @@ public partial class GUBP : BuildCommand
 			string FakeFail = ParseParamValue("FakeFail");
 			if(CommanderSetup)
 			{
-				DoCommanderSetup(EC, AllNodes, AllAggregates, OrderedToDo, TimeIndex, TimeQuantum, bSkipTriggers, bFake, bFakeEC, ExplicitTrigger, UnfinishedTriggers, FakeFail, JobInfo.IsPreflight);
+				DoCommanderSetup(EC, AllBuildNodes, AllAggregateNodes, OrderedToDo, TimeIndex, TimeQuantum, bSkipTriggers, bFake, bFakeEC, ExplicitTrigger, UnfinishedTriggers, FakeFail, JobInfo.IsPreflight);
 			}
 			else if(!ParseParam("ListOnly"))
 			{
@@ -1114,33 +1117,40 @@ public partial class GUBP : BuildCommand
 	/// </summary>
 	/// <param name="AggregateNameToInfo">Map of aggregate names to their info objects</param>
 	/// <param name="NodeNameToInfo">Map of node names to their info objects</param>
-	private static void LinkGraph(IEnumerable<AggregateNode> Aggregates, IEnumerable<BuildNode> Nodes)
+	private static void LinkGraph(List<AggregateNodeTemplate> AggregateNodeTemplates, List<BuildNodeTemplate> BuildNodeTemplates, out List<AggregateNode> AggregateNodes, out List<BuildNode> BuildNodes)
 	{
-		Dictionary<string, BuildNode> NodeNameToInfo = new Dictionary<string,BuildNode>();
-		foreach(BuildNode Node in Nodes)
+		// Create all the build nodes, and a dictionary to find them by name
+		Dictionary<string, BuildNodePair> BuildNodeNameToPair = new Dictionary<string, BuildNodePair>(StringComparer.InvariantCultureIgnoreCase);
+		foreach (BuildNodeTemplate Template in BuildNodeTemplates)
 		{
-			NodeNameToInfo.Add(Node.Name, Node);
+			BuildNodeNameToPair.Add(Template.Name, new BuildNodePair(Template));
 		}
 
-		Dictionary<string, AggregateNode> AggregateNameToInfo = new Dictionary<string, AggregateNode>();
-		foreach(AggregateNode Aggregate in Aggregates)
+		// Create all the aggregate nodes, and add them to a dictionary too
+		Dictionary<string, AggregateNodePair> AggregateNodeNameToPair = new Dictionary<string, AggregateNodePair>(StringComparer.InvariantCultureIgnoreCase);
+		foreach (AggregateNodeTemplate Template in AggregateNodeTemplates)
 		{
-			AggregateNameToInfo.Add(Aggregate.Name, Aggregate);
+			AggregateNodeNameToPair.Add(Template.Name, new AggregateNodePair(Template));
 		}
 
+		// Link the graph
 		int NumErrors = 0;
-		foreach (AggregateNode AggregateNode in AggregateNameToInfo.Values)
+		foreach (AggregateNodePair Pair in AggregateNodeNameToPair.Values)
 		{
-			LinkAggregate(AggregateNode, AggregateNameToInfo, NodeNameToInfo, ref NumErrors);
+			LinkAggregate(Pair, AggregateNodeNameToPair, BuildNodeNameToPair, ref NumErrors);
 		}
-		foreach (BuildNode BuildNode in NodeNameToInfo.Values)
+		foreach (BuildNodePair Pair in BuildNodeNameToPair.Values)
 		{
-			LinkNode(BuildNode, new List<BuildNode>(), AggregateNameToInfo, NodeNameToInfo, ref NumErrors);
+			LinkNode(Pair, new List<BuildNodePair>(), AggregateNodeNameToPair, BuildNodeNameToPair, ref NumErrors);
 		}
 		if(NumErrors > 0)
 		{
 			throw new AutomationException("Failed to link graph ({0} errors).", NumErrors);
 		}
+
+		// Set the output lists
+		AggregateNodes = AggregateNodeNameToPair.Values.Select(x => x.Node).ToList();
+		BuildNodes = BuildNodeNameToPair.Values.Select(x => x.Node).ToList();
 	}
 
 	/// <summary>
@@ -1150,98 +1160,99 @@ public partial class GUBP : BuildCommand
 	/// <param name="AggregateNameToInfo">Map of other aggregate names to their corresponding instance.</param>
 	/// <param name="NodeNameToInfo">Map from node names to their corresponding instance.</param>
 	/// <param name="NumErrors">The number of errors output so far. Incremented if resolving this aggregate fails.</param>
-	private static void LinkAggregate(AggregateNode Aggregate, Dictionary<string, AggregateNode> AggregateNameToInfo, Dictionary<string, BuildNode> NodeNameToInfo, ref int NumErrors)
+	private static void LinkAggregate(AggregateNodePair Pair, Dictionary<string, AggregateNodePair> AggregateNodeNameToPair, Dictionary<string, BuildNodePair> BuildNodeNameToPair, ref int NumErrors)
 	{
-		if (Aggregate.Dependencies == null)
+		if (Pair.Node.Dependencies == null)
 		{
-			Aggregate.Dependencies = new BuildNode[0];
+			Pair.Node.Dependencies = new BuildNode[0];
 
 			HashSet<BuildNode> Dependencies = new HashSet<BuildNode>();
-			foreach (string DependencyName in Aggregate.DependencyNames)
+			foreach (string DependencyName in ParseSemicolonDelimitedList(Pair.Template.DependencyNames))
 			{
-				AggregateNode AggregateDependency;
-				if(AggregateNameToInfo.TryGetValue(DependencyName, out AggregateDependency))
+				AggregateNodePair AggregateNodeDependency;
+				if(AggregateNodeNameToPair.TryGetValue(DependencyName, out AggregateNodeDependency))
 				{
-					LinkAggregate(AggregateDependency, AggregateNameToInfo, NodeNameToInfo, ref NumErrors);
-					Dependencies.UnionWith(AggregateDependency.Dependencies);
+					LinkAggregate(AggregateNodeDependency, AggregateNodeNameToPair, BuildNodeNameToPair, ref NumErrors);
+					Dependencies.UnionWith(AggregateNodeDependency.Node.Dependencies);
 					continue;
 				}
 
-				BuildNode Dependency;
-				if(NodeNameToInfo.TryGetValue(DependencyName, out Dependency))
+				BuildNodePair BuildNodeDependency;
+				if(BuildNodeNameToPair.TryGetValue(DependencyName, out BuildNodeDependency))
 				{
-					Dependencies.Add(Dependency);
+					Dependencies.Add(BuildNodeDependency.Node);
 					continue;
 				}
 
-				CommandUtils.LogError("Node {0} is not in the graph. It is a dependency of {1}.", DependencyName, Aggregate.Name);
+				CommandUtils.LogError("Node {0} is not in the graph. It is a dependency of {1}.", DependencyName, Pair.Template.Name);
 				NumErrors++;
 			}
-			Aggregate.Dependencies = Dependencies.ToArray();
+
+			Pair.Node.Dependencies = Dependencies.ToArray();
 		}
 	}
 
-	private static void LinkNode(BuildNode Node, List<BuildNode> NodeStack, Dictionary<string, AggregateNode> NameToAggregateNode, Dictionary<string, BuildNode> NameToBuildNode, ref int NumErrors)
+	private static void LinkNode(BuildNodePair Pair, List<BuildNodePair> Stack, Dictionary<string, AggregateNodePair> AggregateNodeNameToPair, Dictionary<string, BuildNodePair> BuildNodeNameToPair, ref int NumErrors)
 	{
-		if(Node.OrderDependencies == null)
+		if(Pair.Node.OrderDependencies == null)
 		{
-			NodeStack.Add(Node);
+			Stack.Add(Pair);
 
-			int NodeStackIdx = NodeStack.IndexOf(Node);
-			if (NodeStackIdx != NodeStack.Count - 1)
+			int StackIdx = Stack.IndexOf(Pair);
+			if (StackIdx != Stack.Count - 1)
 			{
 				// There's a cycle in the build graph. Print out the chain.
-				CommandUtils.LogError("Build graph contains a cycle ({0})", String.Join(" -> ", NodeStack.Skip(NodeStackIdx).Select(x => x.Name)));
+				CommandUtils.LogError("Build graph contains a cycle ({0})", String.Join(" -> ", Stack.Skip(StackIdx).Select(x => x.Template.Name)));
 				NumErrors++;
-				Node.OrderDependencies = new HashSet<BuildNode>();
-				Node.InputDependencies = new HashSet<BuildNode>();
+				Pair.Node.OrderDependencies = new HashSet<BuildNode>();
+				Pair.Node.InputDependencies = new HashSet<BuildNode>();
 			}
 			else
 			{
 				// Find all the direct input dependencies
 				HashSet<BuildNode> DirectInputDependencies = new HashSet<BuildNode>();
-				foreach (string InputDependencyName in Node.InputDependencyNames)
+				foreach (string InputDependencyName in ParseSemicolonDelimitedList(Pair.Template.InputDependencyNames))
 				{
-					if (!ResolveDependencies(InputDependencyName, NameToAggregateNode, NameToBuildNode, DirectInputDependencies))
+					if (!ResolveDependencies(InputDependencyName, AggregateNodeNameToPair, BuildNodeNameToPair, DirectInputDependencies))
 					{
-						CommandUtils.LogError("Node {0} is not in the graph. It is an input dependency of {1}.", InputDependencyName, Node.Name);
+						CommandUtils.LogError("Node {0} is not in the graph. It is an input dependency of {1}.", InputDependencyName, Pair.Template.Name);
 						NumErrors++;
 					}
 				}
 
 				// Find all the direct order dependencies. All the input dependencies are also order dependencies.
 				HashSet<BuildNode> DirectOrderDependencies = new HashSet<BuildNode>(DirectInputDependencies);
-				foreach (string OrderDependencyName in Node.OrderDependencyNames)
+				foreach (string OrderDependencyName in ParseSemicolonDelimitedList(Pair.Template.OrderDependencyNames))
 				{
-					if (!ResolveDependencies(OrderDependencyName, NameToAggregateNode, NameToBuildNode, DirectOrderDependencies))
+					if (!ResolveDependencies(OrderDependencyName, AggregateNodeNameToPair, BuildNodeNameToPair, DirectOrderDependencies))
 					{
-						CommandUtils.LogError("Node {0} is not in the graph. It is an order dependency of {1}.", OrderDependencyName, Node.Name);
+						CommandUtils.LogError("Node {0} is not in the graph. It is an order dependency of {1}.", OrderDependencyName, Pair.Template.Name);
 						NumErrors++;
 					}
 				}
 
 				// Link everything
-				foreach(BuildNode DirectOrderDependency in DirectOrderDependencies)
+				foreach (BuildNode DirectOrderDependency in DirectOrderDependencies)
 				{
-					LinkNode(DirectOrderDependency, NodeStack, NameToAggregateNode, NameToBuildNode, ref NumErrors);
+					LinkNode(BuildNodeNameToPair[DirectOrderDependency.Name], Stack, AggregateNodeNameToPair, BuildNodeNameToPair, ref NumErrors);
 				}
 
 				// Recursively include all the input dependencies
-				Node.InputDependencies = new HashSet<BuildNode>(DirectInputDependencies);
+				Pair.Node.InputDependencies = new HashSet<BuildNode>(DirectInputDependencies);
 				foreach (BuildNode DirectInputDependency in DirectInputDependencies)
 				{
-					Node.InputDependencies.UnionWith(DirectInputDependency.InputDependencies);
+					Pair.Node.InputDependencies.UnionWith(DirectInputDependency.InputDependencies);
 				}
 
 				// Same for the order dependencies
-				Node.OrderDependencies = new HashSet<BuildNode>(DirectOrderDependencies);
+				Pair.Node.OrderDependencies = new HashSet<BuildNode>(DirectOrderDependencies);
 				foreach (BuildNode DirectOrderDependency in DirectOrderDependencies)
 				{
-					Node.OrderDependencies.UnionWith(DirectOrderDependency.OrderDependencies);
+					Pair.Node.OrderDependencies.UnionWith(DirectOrderDependency.OrderDependencies);
 				}
 			}
 
-			NodeStack.RemoveAt(NodeStack.Count - 1);
+			Stack.RemoveAt(Stack.Count - 1);
 		}
 	}
 
@@ -1253,25 +1264,35 @@ public partial class GUBP : BuildCommand
 	/// <param name="NodeNameToInfo">Map from node names to their corresponding info instance.</param>
 	/// <param name="Dependencies">The set of dependencies to add to.</param>
 	/// <returns>True if the name was found (and the dependencies list was updated), false otherwise.</returns>
-	private static bool ResolveDependencies(string Name, Dictionary<string, AggregateNode> AggregateNameToInfo, Dictionary<string, BuildNode> NodeNameToInfo, HashSet<BuildNode> Dependencies)
+	private static bool ResolveDependencies(string Name, Dictionary<string, AggregateNodePair> AggregateNodeNameToPair, Dictionary<string, BuildNodePair> BuildNodeNameToPair, HashSet<BuildNode> Dependencies)
 	{
-		AggregateNode AggregateDependency;
-		if (AggregateNameToInfo.TryGetValue(Name, out AggregateDependency))
+		AggregateNodePair AggregateDependency;
+		if (AggregateNodeNameToPair.TryGetValue(Name, out AggregateDependency))
 		{
-			Dependencies.UnionWith(AggregateDependency.Dependencies);
+			Dependencies.UnionWith(AggregateDependency.Node.Dependencies);
 			return true;
 		}
 
-		BuildNode NodeDependency;
-		if (NodeNameToInfo.TryGetValue(Name, out NodeDependency))
+		BuildNodePair NodeDependency;
+		if (BuildNodeNameToPair.TryGetValue(Name, out NodeDependency))
 		{
-			Dependencies.Add(NodeDependency);
+			Dependencies.Add(NodeDependency.Node);
 			return true;
 		}
 
 		return false;
 	}
 
+	/// <summary>
+	/// Parses a semi-colon delimited list into a sequence of strings.
+	/// </summary>
+	/// <param name="StringList">Semicolon delimited list of strings</param>
+	/// <returns>Sequence of strings</returns>
+	private static IEnumerable<string> ParseSemicolonDelimitedList(string StringList)
+	{
+		return StringList.Split(new char[] { ';' }, StringSplitOptions.RemoveEmptyEntries);
+	}
+
 	/// <summary>
 	/// Determine which nodes to build from the command line, or return everything if there's nothing specified explicitly.
 	/// </summary>
diff --git a/Engine/Source/Programs/AutomationTool/GUBP/LegacyBranchSetup.cs b/Engine/Source/Programs/AutomationTool/GUBP/LegacyBranchSetup.cs
--- a/Engine/Source/Programs/AutomationTool/GUBP/LegacyBranchSetup.cs
+++ b/Engine/Source/Programs/AutomationTool/GUBP/LegacyBranchSetup.cs
@@ -213,7 +213,7 @@ partial class GUBP
 		}
     }
 
-	private void FindEmailsForNodes(string BranchName, IEnumerable<BuildNode> Nodes)
+	private void FindEmailsForNodes(string BranchName, IEnumerable<BuildNodeTemplate> Nodes)
 	{
 		List<GUBPEmailHacker> EmailHackers = new List<GUBPEmailHacker>();
 
@@ -234,13 +234,13 @@ partial class GUBP
 			}
 		}
 
-		foreach (BuildNode Node in Nodes)
+		foreach (BuildNodeTemplate Node in Nodes)
 		{
-			Node.RecipientsForFailureEmails = HackEmails(EmailHackers, BranchName, Node, out Node.AddSubmittersToFailureEmails);
+			Node.RecipientsForFailureEmails = String.Join(";", HackEmails(EmailHackers, BranchName, Node, out Node.AddSubmittersToFailureEmails));
 		}
 	}
 
-    private string[] HackEmails(List<GUBPEmailHacker> EmailHackers, string Branch, BuildNode NodeInfo, out bool bIncludeCausers)
+	private string[] HackEmails(List<GUBPEmailHacker> EmailHackers, string Branch, BuildNodeTemplate NodeInfo, out bool bIncludeCausers)
     {
 		string OnlyEmail = ParseParamValue("OnlyEmail");
         if (!String.IsNullOrEmpty(OnlyEmail))
@@ -295,7 +295,7 @@ partial class GUBP
         }
     }
 
-	void FindFrequenciesForNodes(GUBP.GUBPBranchConfig BranchConfig, Dictionary<GUBP.GUBPNode, BuildNode> LegacyToNewNodes, Dictionary<string, int> FrequencyOverrides)
+	void FindFrequenciesForNodes(GUBP.GUBPBranchConfig BranchConfig, Dictionary<GUBP.GUBPNode, BuildNodeTemplate> LegacyToNewNodes, Dictionary<string, int> FrequencyOverrides)
 	{
         List<GUBPFrequencyHacker> FrequencyHackers = new List<GUBPFrequencyHacker>();
 
@@ -316,9 +316,9 @@ partial class GUBP
             }
         }
 
-		foreach (KeyValuePair<GUBP.GUBPNode, BuildNode> LegacyToNewNodePair in LegacyToNewNodes)
+		foreach (KeyValuePair<GUBP.GUBPNode, BuildNodeTemplate> LegacyToNewNodePair in LegacyToNewNodes)
 		{
-			BuildNode Node = LegacyToNewNodePair.Value;
+			BuildNodeTemplate Node = LegacyToNewNodePair.Value;
 			Node.FrequencyShift = LegacyToNewNodePair.Key.CISFrequencyQuantumShift(BranchConfig);
 
 			foreach(GUBPFrequencyHacker FrequencyHacker in FrequencyHackers)
@@ -345,7 +345,7 @@ partial class GUBP
         return AltHostPlatform;
     }
 
-    void AddNodesForBranch(List<UnrealTargetPlatform> InitialHostPlatforms, JobInfo JobInfo, out List<BuildNode> AllNodes, out List<AggregateNode> AllAggregates, ref int TimeQuantum)
+    void AddNodesForBranch(List<UnrealTargetPlatform> InitialHostPlatforms, JobInfo JobInfo, out List<BuildNodeTemplate> AllNodeDefinitions, out List<AggregateNodeTemplate> AllAggregateDefinitions, ref int TimeQuantum)
 	{
 		string BranchName;
 		if (P4Enabled)
@@ -1188,20 +1188,20 @@ partial class GUBP
 		Dictionary<string, int> FrequencyOverrides = ApplyFrequencyBarriers(BranchConfig.GUBPNodes, BranchConfig.GUBPAggregates, BranchOptions.FrequencyBarriers);
 
 		// Get the legacy node to new node mapping
-		Dictionary<GUBP.GUBPNode, BuildNode> LegacyToNewNodes = BranchConfig.GUBPNodes.Values.ToDictionary(x => x, x => x.GetBuildNode(this));
+		Dictionary<GUBP.GUBPNode, BuildNodeTemplate> LegacyToNewNodes = BranchConfig.GUBPNodes.Values.ToDictionary(x => x, x => x.GetTemplate(this));
 
 		// Convert the GUBPNodes and GUBPAggregates maps into lists
-		AllNodes = new List<BuildNode>();
-		AllNodes.AddRange(LegacyToNewNodes.Values);
+		AllNodeDefinitions = new List<BuildNodeTemplate>();
+		AllNodeDefinitions.AddRange(LegacyToNewNodes.Values);
 
-		AllAggregates = new List<AggregateNode>();
-		AllAggregates.AddRange(BranchConfig.GUBPAggregates.Values.Select(x => new LegacyAggregateNode(x)));
+		AllAggregateDefinitions = new List<AggregateNodeTemplate>();
+		AllAggregateDefinitions.AddRange(BranchConfig.GUBPAggregates.Values.Select(x => x.GetTemplate()));
 
 		// Calculate the frequencies for each node
 		FindFrequenciesForNodes(BranchConfig, LegacyToNewNodes, FrequencyOverrides);
 
 		// Get the email list for each node
-		FindEmailsForNodes(BranchName, AllNodes);
+		FindEmailsForNodes(BranchName, AllNodeDefinitions);
 	}
 
 	private static Dictionary<string, int> ApplyFrequencyBarriers(Dictionary<string, GUBPNode> GUBPNodes, Dictionary<string, GUBPAggregateNode> GUBPAggregates, Dictionary<string, sbyte> FrequencyBarriers)
diff --git a/Engine/Source/Programs/AutomationTool/GUBP/LegacyNodes.cs b/Engine/Source/Programs/AutomationTool/GUBP/LegacyNodes.cs
--- a/Engine/Source/Programs/AutomationTool/GUBP/LegacyNodes.cs
+++ b/Engine/Source/Programs/AutomationTool/GUBP/LegacyNodes.cs
@@ -36,9 +36,9 @@ partial class GUBP
         {
             return "";
         }
-		public virtual BuildNode GetBuildNode(GUBP bp)
+		public virtual BuildNodeTemplate GetTemplate(GUBP bp)
 		{
-			return new LegacyBuildNode(bp, this);
+			return new LegacyNodeTemplate(bp, this);
 		}
         public virtual void DoBuild(GUBP bp)
         {
@@ -242,11 +242,11 @@ partial class GUBP
 			}
             return "";
         }
-		public override BuildNode GetBuildNode(GUBP bp)
+		public override BuildNodeTemplate GetTemplate(GUBP bp)
 		{
-			BuildNode Node = base.GetBuildNode(bp);
-			Node.AgentPlatform = GetAgentPlatform();
-			return Node;
+			BuildNodeTemplate Template = base.GetTemplate(bp);
+			Template.AgentPlatform = GetAgentPlatform();
+			return Template;
 		}
         public virtual UnrealTargetPlatform GetAgentPlatform()
         {
@@ -358,14 +358,14 @@ partial class GUBP
             }
             return false;
         }
-		public override BuildNode GetBuildNode(GUBP bp)
+		public override BuildNodeTemplate GetTemplate(GUBP bp)
 		{
-			BuildNode Node = base.GetBuildNode(bp);
+			BuildNodeTemplate Template = base.GetTemplate(bp);
 			if(HostPlatform == UnrealTargetPlatform.Win64)
 			{
-				Node.IsParallelAgentShareEditor = true;
+				Template.IsParallelAgentShareEditor = true;
 			}
-			return Node;
+			return Template;
 		}
         public override bool DeleteBuildProducts()
         {
@@ -460,14 +460,14 @@ partial class GUBP
             }
             return false;
         }
-		public override BuildNode GetBuildNode(GUBP bp)
+		public override BuildNodeTemplate GetTemplate(GUBP bp)
 		{
-			BuildNode Node = base.GetBuildNode(bp);
+			BuildNodeTemplate Template = base.GetTemplate(bp);
 			if (HostPlatform == UnrealTargetPlatform.Win64)
 			{
-				Node.IsParallelAgentShareEditor = true;
+				Template.IsParallelAgentShareEditor = true;
 			}
-			return Node;
+			return Template;
 		}
 		public override UE4Build.BuildAgenda GetAgenda(GUBP bp)
         {
@@ -917,14 +917,14 @@ partial class GUBP
             }
             return false;
         }
-		public override BuildNode GetBuildNode(GUBP bp)
+		public override BuildNodeTemplate GetTemplate(GUBP bp)
 		{
-			BuildNode Node = base.GetBuildNode(bp);
+			BuildNodeTemplate Template = base.GetTemplate(bp);
 			if (HostPlatform == UnrealTargetPlatform.Win64)
 			{
-				Node.IsParallelAgentShareEditor = true;
+				Template.IsParallelAgentShareEditor = true;
 			}
-			return Node;
+			return Template;
 		}
 		public override string GameNameIfAnyForTempStorage()
         {
@@ -1398,6 +1398,16 @@ partial class GUBP
 			}
 		}
 
+		public AggregateNodeTemplate GetTemplate()
+		{
+			AggregateNodeTemplate Template = new AggregateNodeTemplate();
+			Template.Name = GetFullName();
+			Template.DependencyNames = String.Join(";", Dependencies);
+			Template.IsPromotableAggregate = IsPromotableAggregate();
+			Template.IsSeparatePromotable = IsSeparatePromotable();
+			return Template;
+		}
+
 		public abstract string GetFullName();
 
 		public virtual string GameNameIfAnyForTempStorage()
@@ -1492,9 +1502,9 @@ partial class GUBP
             BuildProducts = new List<string>();
             SaveRecordOfSuccessAndAddToBuildProducts();
         }
-        public override BuildNode GetBuildNode(GUBP bp)
+        public override BuildNodeTemplate GetTemplate(GUBP bp)
         {
-			return new TriggerNode(bp, this);
+			return new TriggerNodeTemplate(bp, this);
         }
         public virtual string GetTriggerStateName()
         {
diff --git a/Engine/Source/Programs/AutomationTool/GUBP/Nodes/AggregateNode.cs b/Engine/Source/Programs/AutomationTool/GUBP/Nodes/AggregateNode.cs
--- a/Engine/Source/Programs/AutomationTool/GUBP/Nodes/AggregateNode.cs
+++ b/Engine/Source/Programs/AutomationTool/GUBP/Nodes/AggregateNode.cs
@@ -7,54 +7,40 @@ using System.Text;
 namespace AutomationTool
 {
 	[DebuggerDisplay("{Name}")]
-	public abstract class AggregateNode
+	public class AggregateNodeTemplate
 	{
 		public string Name;
-		public BuildNode[] Dependencies;
-
-		public AggregateNode(string InName)
-		{
-			Name = InName;
-		}
-
-		public abstract IEnumerable<string> DependencyNames
-		{
-			get;
-		}
-
-		public abstract bool IsPromotableAggregate
-		{
-			get;
-		}
-
-		public abstract bool IsSeparatePromotable
-		{
-			get;
-		}
+		public bool IsPromotableAggregate;
+		public bool IsSeparatePromotable;
+		public string DependencyNames;
 	}
 
-	public class LegacyAggregateNode : AggregateNode
+	[DebuggerDisplay("{Name}")]
+	public class AggregateNode
 	{
-		public GUBP.GUBPAggregateNode Node;
-
-		public LegacyAggregateNode(GUBP.GUBPAggregateNode InNode) : base(InNode.GetFullName())
-		{
-			Node = InNode;
-		}
+		public string Name;
+		public bool IsPromotableAggregate;
+		public bool IsSeparatePromotable;
+		public BuildNode[] Dependencies;
 
-		public override IEnumerable<string> DependencyNames
+		public AggregateNode(AggregateNodeTemplate Template)
 		{
-			get { return Node.Dependencies; }
+			Name = Template.Name;
+			IsPromotableAggregate = Template.IsPromotableAggregate;
+			IsSeparatePromotable = Template.IsSeparatePromotable;
 		}
+	}
 
-		public override bool IsPromotableAggregate
-		{
-			get { return Node.IsPromotableAggregate(); }
-		}
+	[DebuggerDisplay("{Definition.Name}")]
+	class AggregateNodePair
+	{
+		public readonly AggregateNodeTemplate Template;
+		public readonly AggregateNode Node;
 
-		public override bool IsSeparatePromotable
+		public AggregateNodePair(AggregateNodeTemplate InTemplate)
 		{
-			get { return Node.IsSeparatePromotable(); }
+			Template = InTemplate;
+			Node = new AggregateNode(Template);
 		}
 	}
 }
diff --git a/Engine/Source/Programs/AutomationTool/GUBP/Nodes/BuildNode.cs b/Engine/Source/Programs/AutomationTool/GUBP/Nodes/BuildNode.cs
--- a/Engine/Source/Programs/AutomationTool/GUBP/Nodes/BuildNode.cs
+++ b/Engine/Source/Programs/AutomationTool/GUBP/Nodes/BuildNode.cs
@@ -7,6 +7,32 @@ using UnrealBuildTool;
 
 namespace AutomationTool
 {
+	[DebuggerDisplay("{Name}")]
+	public abstract class BuildNodeTemplate
+	{
+		public string Name;
+		public UnrealTargetPlatform AgentPlatform = UnrealTargetPlatform.Win64;
+		public string AgentRequirements;
+		public string AgentSharingGroup;
+		public int FrequencyShift;
+		public int AgentMemoryRequirement;
+		public int TimeoutInMinutes;
+		public float Priority;
+		public string InputDependencyNames;
+		public string OrderDependencyNames;
+		public string RecipientsForFailureEmails;
+		public bool AddSubmittersToFailureEmails;
+		public bool SendSuccessEmail;
+		public bool IsParallelAgentShareEditor;
+		public bool IsSticky;
+		public bool IsTest;
+		public string DisplayGroupName;
+		public string GameNameIfAnyForTempStorage;
+		public string RootIfAnyForTempStorage;
+
+		public abstract BuildNode Instantiate();
+	}
+
 	[DebuggerDisplay("{Name}")]
 	public abstract class BuildNode
 	{
@@ -27,27 +53,33 @@ namespace AutomationTool
 		public bool AddSubmittersToFailureEmails;
 		public bool SendSuccessEmail;
 		public bool IsParallelAgentShareEditor;
+		public bool IsSticky;
+		public bool IsTest;
+		public string DisplayGroupName;
+		public string GameNameIfAnyForTempStorage;
+		public string RootIfAnyForTempStorage;
 
 		public List<string> BuildProducts;
 
-		public BuildNode(string InName)
-		{
-			Name = InName;
-		}
-
-		public abstract bool IsSticky
-		{
-			get;
-		}
-
-		public abstract IEnumerable<string> InputDependencyNames
-		{
-			get;
-		}
-
-		public abstract IEnumerable<string> OrderDependencyNames
-		{
-			get;
+		public BuildNode(BuildNodeTemplate Template)
+		{
+			Name = Template.Name;
+			AgentPlatform = Template.AgentPlatform;
+			AgentRequirements = Template.AgentRequirements;
+			AgentSharingGroup = Template.AgentSharingGroup;
+			FrequencyShift = Template.FrequencyShift;
+			AgentMemoryRequirement = Template.AgentMemoryRequirement;
+			TimeoutInMinutes = Template.TimeoutInMinutes;
+			Priority = Template.Priority;
+			RecipientsForFailureEmails = Template.RecipientsForFailureEmails.Split(';');
+			AddSubmittersToFailureEmails = Template.AddSubmittersToFailureEmails;
+			SendSuccessEmail = Template.SendSuccessEmail;
+			IsParallelAgentShareEditor = Template.IsParallelAgentShareEditor;
+			IsSticky = Template.IsSticky;
+			IsTest = Template.IsTest;
+			DisplayGroupName = Template.DisplayGroupName;
+			GameNameIfAnyForTempStorage = Template.GameNameIfAnyForTempStorage;
+			RootIfAnyForTempStorage = Template.RootIfAnyForTempStorage;
 		}
 
 		public abstract void ArchiveBuildProducts(string GameNameIfAny, string StorageRootIfAny, TempStorageNodeInfo TempStorageNodeInfo, bool bLocalOnly);
@@ -73,128 +105,18 @@ namespace AutomationTool
 			System.Diagnostics.Trace.TraceWarning("Implicit conversion from NodeInfo to string\n{0}", Environment.StackTrace);
 			return Name;
 		}
-
-		// Legacy stuff that should probably be removed
-
-		public abstract bool IsTest
-		{
-			get;
-		}
-
-		public abstract string DisplayGroupName
-		{
-			get;
-		}
-
-		public abstract string GameNameIfAnyForTempStorage
-		{
-			get;
-		}
-
-		public abstract string RootIfAnyForTempStorage
-		{
-			get;
-		}
 	}
 
-	public class LegacyBuildNode : BuildNode
+	[DebuggerDisplay("{Definition.Name}")]
+	class BuildNodePair
 	{
-		GUBP Owner;
-		GUBP.GUBPNode Node;
-
-		public LegacyBuildNode(GUBP InOwner, GUBP.GUBPNode InNode) : base(InNode.GetFullName())
-		{
-			Owner = InOwner;
-			Node = InNode;
-			AgentRequirements = Node.ECAgentString();
-			AgentSharingGroup = Node.AgentSharingGroup;
-			AgentMemoryRequirement = Node.AgentMemoryRequirement();
-			TimeoutInMinutes = Node.TimeoutInMinutes();
-			SendSuccessEmail = Node.SendSuccessEmail();
-			Priority = Node.Priority();
-		}
-
-		public override bool IsSticky
-		{
-			get { return Node.IsSticky(); }
-		}
-
-		public override IEnumerable<string> InputDependencyNames
-		{
-			get { return Node.FullNamesOfDependencies; }
-		}
-
-		public override IEnumerable<string> OrderDependencyNames
-		{
-			get { return Node.FullNamesOfPseudodependencies; }
-		}
-
-		public override void ArchiveBuildProducts(string GameNameIfAny, string StorageRootIfAny, TempStorageNodeInfo TempStorageNodeInfo, bool bLocalOnly)
-		{
-			TempStorage.StoreToTempStorage(TempStorageNodeInfo, BuildProducts, bLocalOnly, GameNameIfAny, StorageRootIfAny);
-		}
-
-		public override void RetrieveBuildProducts(string GameNameIfAny, string StorageRootIfAny, TempStorageNodeInfo TempStorageNodeInfo)
-		{
-			CommandUtils.LogConsole("***** Retrieving GUBP Node {0} from {1}", Name, TempStorageNodeInfo.GetRelativeDirectory());
-			bool WasLocal;
-			try
-			{
-				BuildProducts = TempStorage.RetrieveFromTempStorage(TempStorageNodeInfo, out WasLocal, GameNameIfAny, StorageRootIfAny);
-			}
-			catch (Exception Ex)
-			{
-				if (GameNameIfAny != "")
-				{
-					BuildProducts = TempStorage.RetrieveFromTempStorage(TempStorageNodeInfo, out WasLocal, "", StorageRootIfAny);
-				}
-				else
-				{
-					throw new AutomationException(Ex, "Build Products cannot be found for node {0}", Name);
-				}
-			}
-			Node.BuildProducts = BuildProducts;
-		}
-
-		public override void DoBuild()
-		{
-			Node.AllDependencyBuildProducts = new List<string>();
-			foreach (BuildNode InputDependency in InputDependencies)
-			{
-				foreach (string BuildProduct in InputDependency.BuildProducts)
-				{
-					Node.AddDependentBuildProduct(BuildProduct);
-				}
-			}
-
-			Node.DoBuild(Owner);
-			BuildProducts = Node.BuildProducts;
-		}
-
-		public override void DoFakeBuild()
-		{
-			Node.DoFakeBuild(Owner);
-			BuildProducts = Node.BuildProducts;
-		}
-
-		public override bool IsTest
-		{
-			get { return Node.IsTest(); }
-		}
-
-		public override string DisplayGroupName
-		{
-			get { return Node.GetDisplayGroupName(); }
-		}
-
-		public override string GameNameIfAnyForTempStorage
-		{
-			get { return Node.GameNameIfAnyForTempStorage(); }
-		}
+		public readonly BuildNodeTemplate Template;
+		public readonly BuildNode Node;
 
-		public override string RootIfAnyForTempStorage
+		public BuildNodePair(BuildNodeTemplate InTemplate)
 		{
-			get { return Node.RootIfAnyForTempStorage(); }
+			Template = InTemplate;
+			Node = InTemplate.Instantiate();
 		}
 	}
 }
diff --git a/Engine/Source/Programs/AutomationTool/GUBP/Nodes/LegacyNode.cs b/Engine/Source/Programs/AutomationTool/GUBP/Nodes/LegacyNode.cs
new file mode 100644
--- /dev/null
+++ b/Engine/Source/Programs/AutomationTool/GUBP/Nodes/LegacyNode.cs
@@ -0,0 +1,107 @@
+﻿using System;
+using System.Collections.Generic;
+using System.Diagnostics;
+using System.Linq;
+using System.Text;
+using UnrealBuildTool;
+
+namespace AutomationTool
+{
+	public class LegacyNodeTemplate : BuildNodeTemplate
+	{
+		public GUBP Owner;
+		public GUBP.GUBPNode Node;
+
+		public LegacyNodeTemplate(GUBP InOwner, GUBP.GUBPNode InNode)
+		{
+			Owner = InOwner;
+			Node = InNode;
+
+			SetStandardProperties(InNode, this);
+		}
+
+		public static void SetStandardProperties(GUBP.GUBPNode Node, BuildNodeTemplate Definition)
+		{
+			Definition.Name = Node.GetFullName();
+			Definition.AgentRequirements = Node.ECAgentString();
+			Definition.AgentSharingGroup = Node.AgentSharingGroup;
+			Definition.AgentMemoryRequirement = Node.AgentMemoryRequirement();
+			Definition.TimeoutInMinutes = Node.TimeoutInMinutes();
+			Definition.SendSuccessEmail = Node.SendSuccessEmail();
+			Definition.Priority = Node.Priority();
+			Definition.IsSticky = Node.IsSticky();
+			Definition.InputDependencyNames = String.Join(";", Node.FullNamesOfDependencies);
+			Definition.OrderDependencyNames = String.Join(";", Node.FullNamesOfPseudodependencies);
+			Definition.IsTest = Node.IsTest();
+			Definition.DisplayGroupName = Node.GetDisplayGroupName();
+			Definition.GameNameIfAnyForTempStorage = Node.GameNameIfAnyForTempStorage();
+			Definition.RootIfAnyForTempStorage = Node.RootIfAnyForTempStorage();
+		}
+
+		public override BuildNode Instantiate()
+		{
+			return new LegacyNode(this);
+		}
+	}
+
+	public class LegacyNode : BuildNode
+	{
+		GUBP Owner;
+		GUBP.GUBPNode Node;
+
+		public LegacyNode(LegacyNodeTemplate Template)
+			: base(Template)
+		{
+			Owner = Template.Owner;
+			Node = Template.Node;
+		}
+
+		public override void ArchiveBuildProducts(string GameNameIfAny, string StorageRootIfAny, TempStorageNodeInfo TempStorageNodeInfo, bool bLocalOnly)
+		{
+			TempStorage.StoreToTempStorage(TempStorageNodeInfo, BuildProducts, bLocalOnly, GameNameIfAny, StorageRootIfAny);
+		}
+
+		public override void RetrieveBuildProducts(string GameNameIfAny, string StorageRootIfAny, TempStorageNodeInfo TempStorageNodeInfo)
+		{
+			CommandUtils.LogConsole("***** Retrieving GUBP Node {0} from {1}", Name, TempStorageNodeInfo.GetRelativeDirectory());
+			bool WasLocal;
+			try
+			{
+				BuildProducts = TempStorage.RetrieveFromTempStorage(TempStorageNodeInfo, out WasLocal, GameNameIfAny, StorageRootIfAny);
+			}
+			catch (Exception Ex)
+			{
+				if (GameNameIfAny != "")
+				{
+					BuildProducts = TempStorage.RetrieveFromTempStorage(TempStorageNodeInfo, out WasLocal, "", StorageRootIfAny);
+				}
+				else
+				{
+					throw new AutomationException(Ex, "Build Products cannot be found for node {0}", Name);
+				}
+			}
+			Node.BuildProducts = BuildProducts;
+		}
+
+		public override void DoBuild()
+		{
+			Node.AllDependencyBuildProducts = new List<string>();
+			foreach (BuildNode InputDependency in InputDependencies)
+			{
+				foreach (string BuildProduct in InputDependency.BuildProducts)
+				{
+					Node.AddDependentBuildProduct(BuildProduct);
+				}
+			}
+
+			Node.DoBuild(Owner);
+			BuildProducts = Node.BuildProducts;
+		}
+
+		public override void DoFakeBuild()
+		{
+			Node.DoFakeBuild(Owner);
+			BuildProducts = Node.BuildProducts;
+		}
+	}
+}
diff --git a/Engine/Source/Programs/AutomationTool/GUBP/Nodes/TriggerNode.cs b/Engine/Source/Programs/AutomationTool/GUBP/Nodes/TriggerNode.cs
--- a/Engine/Source/Programs/AutomationTool/GUBP/Nodes/TriggerNode.cs
+++ b/Engine/Source/Programs/AutomationTool/GUBP/Nodes/TriggerNode.cs
@@ -6,7 +6,22 @@ using System.Text;
 
 namespace AutomationTool
 {
-	public class TriggerNode : LegacyBuildNode
+	public class TriggerNodeTemplate : LegacyNodeTemplate
+	{
+		public new GUBP.WaitForUserInput Node;
+
+		public TriggerNodeTemplate(GUBP InOwner, GUBP.WaitForUserInput InNode) : base(InOwner, InNode)
+		{
+			Node = InNode;
+		}
+
+		public override BuildNode Instantiate()
+		{
+			return new TriggerNode(this);
+		}
+	}
+
+	public class TriggerNode : LegacyNode
 	{
         public bool IsTriggered;
 		public bool RequiresRecursiveWorkflow;
@@ -14,24 +29,20 @@ namespace AutomationTool
 		public string DescriptionText;
 		public string ActionText;
 
-		public TriggerNode(GUBP bp, GUBP.WaitForUserInput InNode) : base(bp, InNode)
+		public TriggerNode(TriggerNodeTemplate Template) : base(Template)
 		{
 			AddSubmittersToFailureEmails = false;
 
-			StateName = InNode.GetTriggerStateName();
-			DescriptionText = InNode.GetTriggerDescText();
-			ActionText = InNode.GetTriggerActionText();
-			RequiresRecursiveWorkflow = InNode.TriggerRequiresRecursiveWorkflow();
+			StateName = Template.Node.GetTriggerStateName();
+			DescriptionText = Template.Node.GetTriggerDescText();
+			ActionText = Template.Node.GetTriggerActionText();
+			RequiresRecursiveWorkflow = Template.Node.TriggerRequiresRecursiveWorkflow();
 		}
 
 		public void Activate()
-		{ 
-			IsTriggered = true;
-		}
-
-		public override bool IsSticky
 		{
-			get { return IsTriggered; }
+			IsTriggered = true;
+			IsSticky = true;
 		}
 	}
 }
