Move build graph into its own class.

diff --git a/Engine/Source/Programs/AutomationTool/GUBP/BuildGraph.cs b/Engine/Source/Programs/AutomationTool/GUBP/BuildGraph.cs
new file mode 100644
--- /dev/null
+++ b/Engine/Source/Programs/AutomationTool/GUBP/BuildGraph.cs
@@ -0,0 +1,279 @@
+﻿using System;
+using System.Collections.Generic;
+using System.Linq;
+using System.Text;
+using System.Threading.Tasks;
+
+namespace AutomationTool
+{
+	class BuildGraphTemplate
+	{
+		public List<BuildNodeTemplate> BuildNodeTemplates;
+		public List<AggregateNodeTemplate> AggregateNodeTemplates;
+	}
+
+	class BuildGraph
+	{
+		public List<BuildNode> BuildNodes;
+		public List<AggregateNode> AggregateNodes;
+
+		/// <summary>
+		/// Constructor. Creates build nodes from the given template.
+		/// </summary>
+		/// <param name="Template">Template for this build graph.</param>
+		public BuildGraph(BuildGraphTemplate Template)
+		{
+			LinkGraph(Template.AggregateNodeTemplates, Template.BuildNodeTemplates, out AggregateNodes, out BuildNodes);
+			FindControllingTriggers(BuildNodes);
+		}
+
+		/// <summary>
+		/// Resolves the names of each node and aggregates' dependencies, and links them together into the build graph.
+		/// </summary>
+		/// <param name="AggregateNameToInfo">Map of aggregate names to their info objects</param>
+		/// <param name="NodeNameToInfo">Map of node names to their info objects</param>
+		private static void LinkGraph(List<AggregateNodeTemplate> AggregateNodeTemplates, List<BuildNodeTemplate> BuildNodeTemplates, out List<AggregateNode> AggregateNodes, out List<BuildNode> BuildNodes)
+		{
+			// Create all the build nodes, and a dictionary to find them by name
+			Dictionary<string, BuildNodePair> BuildNodeNameToPair = new Dictionary<string, BuildNodePair>(StringComparer.InvariantCultureIgnoreCase);
+			foreach (BuildNodeTemplate Template in BuildNodeTemplates)
+			{
+				BuildNodeNameToPair.Add(Template.Name, new BuildNodePair(Template));
+			}
+
+			// Create all the aggregate nodes, and add them to a dictionary too
+			Dictionary<string, AggregateNodePair> AggregateNodeNameToPair = new Dictionary<string, AggregateNodePair>(StringComparer.InvariantCultureIgnoreCase);
+			foreach (AggregateNodeTemplate Template in AggregateNodeTemplates)
+			{
+				AggregateNodeNameToPair.Add(Template.Name, new AggregateNodePair(Template));
+			}
+
+			// Link the graph
+			int NumErrors = 0;
+			foreach (AggregateNodePair Pair in AggregateNodeNameToPair.Values)
+			{
+				LinkAggregate(Pair, AggregateNodeNameToPair, BuildNodeNameToPair, ref NumErrors);
+			}
+			foreach (BuildNodePair Pair in BuildNodeNameToPair.Values)
+			{
+				LinkNode(Pair, new List<BuildNodePair>(), AggregateNodeNameToPair, BuildNodeNameToPair, ref NumErrors);
+			}
+			if (NumErrors > 0)
+			{
+				throw new AutomationException("Failed to link graph ({0} errors).", NumErrors);
+			}
+
+			// Set the output lists
+			AggregateNodes = AggregateNodeNameToPair.Values.Select(x => x.Node).ToList();
+			BuildNodes = BuildNodeNameToPair.Values.Select(x => x.Node).ToList();
+		}
+
+		/// <summary>
+		/// Resolves the dependency names in an aggregate to NodeInfo instances, filling in the AggregateInfo.Dependenices array. Any referenced aggregates will also be linked, recursively.
+		/// </summary>
+		/// <param name="Aggregate">The aggregate to link</param>
+		/// <param name="AggregateNameToInfo">Map of other aggregate names to their corresponding instance.</param>
+		/// <param name="NodeNameToInfo">Map from node names to their corresponding instance.</param>
+		/// <param name="NumErrors">The number of errors output so far. Incremented if resolving this aggregate fails.</param>
+		private static void LinkAggregate(AggregateNodePair Pair, Dictionary<string, AggregateNodePair> AggregateNodeNameToPair, Dictionary<string, BuildNodePair> BuildNodeNameToPair, ref int NumErrors)
+		{
+			if (Pair.Node.Dependencies == null)
+			{
+				Pair.Node.Dependencies = new BuildNode[0];
+
+				HashSet<BuildNode> Dependencies = new HashSet<BuildNode>();
+				foreach (string DependencyName in ParseSemicolonDelimitedList(Pair.Template.DependencyNames))
+				{
+					AggregateNodePair AggregateNodeDependency;
+					if (AggregateNodeNameToPair.TryGetValue(DependencyName, out AggregateNodeDependency))
+					{
+						LinkAggregate(AggregateNodeDependency, AggregateNodeNameToPair, BuildNodeNameToPair, ref NumErrors);
+						Dependencies.UnionWith(AggregateNodeDependency.Node.Dependencies);
+						continue;
+					}
+
+					BuildNodePair BuildNodeDependency;
+					if (BuildNodeNameToPair.TryGetValue(DependencyName, out BuildNodeDependency))
+					{
+						Dependencies.Add(BuildNodeDependency.Node);
+						continue;
+					}
+
+					CommandUtils.LogError("Node {0} is not in the graph. It is a dependency of {1}.", DependencyName, Pair.Template.Name);
+					NumErrors++;
+				}
+
+				Pair.Node.Dependencies = Dependencies.ToArray();
+			}
+		}
+
+		/// <summary>
+		/// Resolves the dependencies in a build graph
+		/// </summary>
+		/// <param name="Pair">The build node template and instance to link.</param>
+		/// <param name="Stack">Stack of all the build nodes currently being processed. Used to detect cycles.</param>
+		/// <param name="AggregateNodeNameToPair">Mapping of aggregate node names to their template and instance.</param>
+		/// <param name="BuildNodeNameToPair">Mapping of build node names to their template and instance.</param>
+		/// <param name="NumErrors">Number of errors encountered. Will be incremented for each error encountered.</param>
+		private static void LinkNode(BuildNodePair Pair, List<BuildNodePair> Stack, Dictionary<string, AggregateNodePair> AggregateNodeNameToPair, Dictionary<string, BuildNodePair> BuildNodeNameToPair, ref int NumErrors)
+		{
+			if (Pair.Node.OrderDependencies == null)
+			{
+				Stack.Add(Pair);
+
+				int StackIdx = Stack.IndexOf(Pair);
+				if (StackIdx != Stack.Count - 1)
+				{
+					// There's a cycle in the build graph. Print out the chain.
+					CommandUtils.LogError("Build graph contains a cycle ({0})", String.Join(" -> ", Stack.Skip(StackIdx).Select(x => x.Template.Name)));
+					NumErrors++;
+					Pair.Node.OrderDependencies = new HashSet<BuildNode>();
+					Pair.Node.InputDependencies = new HashSet<BuildNode>();
+				}
+				else
+				{
+					// Find all the direct input dependencies
+					HashSet<BuildNode> DirectInputDependencies = new HashSet<BuildNode>();
+					foreach (string InputDependencyName in ParseSemicolonDelimitedList(Pair.Template.InputDependencyNames))
+					{
+						if (!ResolveDependencies(InputDependencyName, AggregateNodeNameToPair, BuildNodeNameToPair, DirectInputDependencies))
+						{
+							CommandUtils.LogError("Node {0} is not in the graph. It is an input dependency of {1}.", InputDependencyName, Pair.Template.Name);
+							NumErrors++;
+						}
+					}
+
+					// Find all the direct order dependencies. All the input dependencies are also order dependencies.
+					HashSet<BuildNode> DirectOrderDependencies = new HashSet<BuildNode>(DirectInputDependencies);
+					foreach (string OrderDependencyName in ParseSemicolonDelimitedList(Pair.Template.OrderDependencyNames))
+					{
+						if (!ResolveDependencies(OrderDependencyName, AggregateNodeNameToPair, BuildNodeNameToPair, DirectOrderDependencies))
+						{
+							CommandUtils.LogError("Node {0} is not in the graph. It is an order dependency of {1}.", OrderDependencyName, Pair.Template.Name);
+							NumErrors++;
+						}
+					}
+
+					// Link everything
+					foreach (BuildNode DirectOrderDependency in DirectOrderDependencies)
+					{
+						LinkNode(BuildNodeNameToPair[DirectOrderDependency.Name], Stack, AggregateNodeNameToPair, BuildNodeNameToPair, ref NumErrors);
+					}
+
+					// Recursively include all the input dependencies
+					Pair.Node.InputDependencies = new HashSet<BuildNode>(DirectInputDependencies);
+					foreach (BuildNode DirectInputDependency in DirectInputDependencies)
+					{
+						Pair.Node.InputDependencies.UnionWith(DirectInputDependency.InputDependencies);
+					}
+
+					// Same for the order dependencies
+					Pair.Node.OrderDependencies = new HashSet<BuildNode>(DirectOrderDependencies);
+					foreach (BuildNode DirectOrderDependency in DirectOrderDependencies)
+					{
+						Pair.Node.OrderDependencies.UnionWith(DirectOrderDependency.OrderDependencies);
+					}
+				}
+
+				Stack.RemoveAt(Stack.Count - 1);
+			}
+		}
+
+		/// <summary>
+		/// Adds all the nodes matching a given name to a hash set, expanding any aggregates to their dependencices.
+		/// </summary>
+		/// <param name="Name">The name to look for</param>
+		/// <param name="AggregateNameToInfo">Map of other aggregate names to their corresponding info instance.</param>
+		/// <param name="NodeNameToInfo">Map from node names to their corresponding info instance.</param>
+		/// <param name="Dependencies">The set of dependencies to add to.</param>
+		/// <returns>True if the name was found (and the dependencies list was updated), false otherwise.</returns>
+		private static bool ResolveDependencies(string Name, Dictionary<string, AggregateNodePair> AggregateNodeNameToPair, Dictionary<string, BuildNodePair> BuildNodeNameToPair, HashSet<BuildNode> Dependencies)
+		{
+			AggregateNodePair AggregateDependency;
+			if (AggregateNodeNameToPair.TryGetValue(Name, out AggregateDependency))
+			{
+				Dependencies.UnionWith(AggregateDependency.Node.Dependencies);
+				return true;
+			}
+
+			BuildNodePair NodeDependency;
+			if (BuildNodeNameToPair.TryGetValue(Name, out NodeDependency))
+			{
+				Dependencies.Add(NodeDependency.Node);
+				return true;
+			}
+
+			return false;
+		}
+
+		/// <summary>
+		/// Parses a semi-colon delimited list into a sequence of strings.
+		/// </summary>
+		/// <param name="StringList">Semicolon delimited list of strings</param>
+		/// <returns>Sequence of strings</returns>
+		private static IEnumerable<string> ParseSemicolonDelimitedList(string StringList)
+		{
+			return StringList.Split(new char[] { ';' }, StringSplitOptions.RemoveEmptyEntries);
+		}
+
+		/// <summary>
+		/// Recursively update the ControllingTriggers array for each of the nodes passed in. 
+		/// </summary>
+		/// <param name="NodesToDo">Nodes to update</param>
+		static void FindControllingTriggers(List<BuildNode> Nodes)
+		{
+			foreach (BuildNode Node in Nodes)
+			{
+				FindControllingTriggers(Node);
+			}
+		}
+
+		/// <summary>
+		/// Recursively find the controlling triggers for the given node.
+		/// </summary>
+		/// <param name="NodeToDo">Node to find the controlling triggers for</param>
+		static void FindControllingTriggers(BuildNode Node)
+		{
+			if (Node.ControllingTriggers == null)
+			{
+				Node.ControllingTriggers = new TriggerNode[0];
+
+				// Find the immediate trigger controlling this one
+				List<TriggerNode> PreviousTriggers = new List<TriggerNode>();
+				foreach (BuildNode Dependency in Node.OrderDependencies)
+				{
+					FindControllingTriggers(Dependency);
+
+					TriggerNode PreviousTrigger = Dependency as TriggerNode;
+					if (PreviousTrigger == null && Dependency.ControllingTriggers.Length > 0)
+					{
+						PreviousTrigger = Dependency.ControllingTriggers.Last();
+					}
+
+					if (PreviousTrigger != null && !PreviousTriggers.Contains(PreviousTrigger))
+					{
+						PreviousTriggers.Add(PreviousTrigger);
+					}
+				}
+
+				// Remove previous triggers from the list that aren't the last in the chain. If there's a trigger chain of X.Y.Z, and a node has dependencies behind all three, the only trigger we care about is Z.
+				PreviousTriggers.RemoveAll(x => PreviousTriggers.Any(y => y.ControllingTriggers.Contains(x)));
+
+				// We only support one direct controlling trigger at the moment (though it may be in a chain with other triggers)
+				if (PreviousTriggers.Count > 1)
+				{
+					throw new AutomationException("Node {0} has multiple controlling triggers: {1}", Node.Name, String.Join(", ", PreviousTriggers.Select(x => x.Name)));
+				}
+
+				// Update the list of controlling triggers
+				if (PreviousTriggers.Count == 1)
+				{
+					List<TriggerNode> ControllingTriggers = new List<TriggerNode>();
+					ControllingTriggers.AddRange(PreviousTriggers[0].ControllingTriggers);
+					ControllingTriggers.Add(PreviousTriggers[0]);
+					Node.ControllingTriggers = ControllingTriggers.ToArray();
+				}
+			}
+		}
+	}
+}
diff --git a/Engine/Source/Programs/AutomationTool/GUBP/GUBP.Automation.csproj b/Engine/Source/Programs/AutomationTool/GUBP/GUBP.Automation.csproj
--- a/Engine/Source/Programs/AutomationTool/GUBP/GUBP.Automation.csproj
+++ b/Engine/Source/Programs/AutomationTool/GUBP/GUBP.Automation.csproj
@@ -60,6 +60,7 @@
     <Compile Include="..\..\DotNETCommon\MetaData.cs">
       <Link>Properties\MetaData.cs</Link>
     </Compile>
+    <Compile Include="BuildGraph.cs" />
     <Compile Include="ElectricCommander.cs" />
     <Compile Include="GUBP.cs" />
     <Compile Include="LegacyBranchSetup.cs" />
diff --git a/Engine/Source/Programs/AutomationTool/GUBP/GUBP.cs b/Engine/Source/Programs/AutomationTool/GUBP/GUBP.cs
--- a/Engine/Source/Programs/AutomationTool/GUBP/GUBP.cs
+++ b/Engine/Source/Programs/AutomationTool/GUBP/GUBP.cs
@@ -223,18 +223,13 @@ public partial class GUBP : BuildCommand
 			ExplicitTriggerName = ParseParamValue("TriggerNode", "");
 		}
 
-		List<BuildNodeTemplate> AllBuildNodeTemplates;
-		List<AggregateNodeTemplate> AllAggregateNodeTemplates;
 		int TimeQuantum = 20;
-		AddNodesForBranch(HostPlatforms, JobInfo, out AllBuildNodeTemplates, out AllAggregateNodeTemplates, ref TimeQuantum);
+		BuildGraphTemplate GraphTemplate = new BuildGraphTemplate();
+		AddNodesForBranch(HostPlatforms, JobInfo, out GraphTemplate.BuildNodeTemplates, out GraphTemplate.AggregateNodeTemplates, ref TimeQuantum);
 
-		List<AggregateNode> AllAggregateNodes;
-		List<BuildNode> AllBuildNodes;
-		LinkGraph(AllAggregateNodeTemplates, AllBuildNodeTemplates, out AllAggregateNodes, out AllBuildNodes);
-
-		FindControllingTriggers(AllBuildNodes);
-        FindCompletionState(AllBuildNodes, JobInfo, LocalOnly);
-		ComputeDependentFrequencies(AllBuildNodes);
+		BuildGraph Graph = new BuildGraph(GraphTemplate);
+		FindCompletionState(Graph.BuildNodes, JobInfo, LocalOnly);
+		ComputeDependentFrequencies(Graph.BuildNodes);
 
         if (bCleanLocalTempStorage)  // shared temp storage can never be wiped
         {
@@ -252,7 +247,7 @@ public partial class GUBP : BuildCommand
 			Log("Setting TimeIndex to {0}", TimeIndex);
 		}
 
-		HashSet<BuildNode> NodesToDo = ParseNodesToDo(AllBuildNodes, AllAggregateNodes);
+		HashSet<BuildNode> NodesToDo = ParseNodesToDo(Graph.BuildNodes, Graph.AggregateNodes);
 		CullNodesForTimeIndex(NodesToDo, TimeIndex);
         if (JobInfo.IsPreflight)
         {
@@ -264,7 +259,7 @@ public partial class GUBP : BuildCommand
         {
             if (!String.IsNullOrEmpty(ExplicitTriggerName))
             {
-                foreach (TriggerNode Trigger in AllBuildNodes.OfType<TriggerNode>())
+                foreach (TriggerNode Trigger in Graph.BuildNodes.OfType<TriggerNode>())
                 {
 					if (Trigger.Name.Equals(ExplicitTriggerName, StringComparison.InvariantCultureIgnoreCase))
                     {
@@ -282,7 +277,7 @@ public partial class GUBP : BuildCommand
             {
                 if (bSkipTriggers)
                 {
-                    foreach (TriggerNode Trigger in AllBuildNodes.OfType<TriggerNode>())
+                    foreach (TriggerNode Trigger in Graph.BuildNodes.OfType<TriggerNode>())
                     {
 						Trigger.Activate();
                     }
@@ -294,7 +289,7 @@ public partial class GUBP : BuildCommand
 
 		List<TriggerNode> UnfinishedTriggers = FindUnfinishedTriggers(bSkipTriggers, ExplicitTrigger, OrderedToDo);
 
-		PrintNodes(this, OrderedToDo, AllAggregateNodes, UnfinishedTriggers, TimeQuantum);
+		PrintNodes(this, OrderedToDo, Graph.AggregateNodes, UnfinishedTriggers, TimeQuantum);
 
         //check sorting
 		CheckSortOrder(OrderedToDo);
@@ -304,7 +299,7 @@ public partial class GUBP : BuildCommand
 		string ShowHistoryParam = ParseParamValue("ShowHistory", null);
 		if(ShowHistoryParam != null)
 		{
-			BuildNode Node = AllBuildNodes.FirstOrDefault(x => x.Name.Equals(ShowHistoryParam, StringComparison.InvariantCultureIgnoreCase));
+			BuildNode Node = Graph.BuildNodes.FirstOrDefault(x => x.Name.Equals(ShowHistoryParam, StringComparison.InvariantCultureIgnoreCase));
 			if(Node == null)
 			{
 				throw new AutomationException("Couldn't find node {0}", ShowHistoryParam);
@@ -323,7 +318,7 @@ public partial class GUBP : BuildCommand
 			string FakeFail = ParseParamValue("FakeFail");
 			if(CommanderSetup)
 			{
-				DoCommanderSetup(EC, AllBuildNodes, AllAggregateNodes, OrderedToDo, TimeIndex, TimeQuantum, bSkipTriggers, bFake, bFakeEC, ExplicitTrigger, UnfinishedTriggers, FakeFail, JobInfo.IsPreflight);
+				DoCommanderSetup(EC, Graph.BuildNodes, Graph.AggregateNodes, OrderedToDo, TimeIndex, TimeQuantum, bSkipTriggers, bFake, bFakeEC, ExplicitTrigger, UnfinishedTriggers, FakeFail, JobInfo.IsPreflight);
 			}
 			else if(!ParseParam("ListOnly"))
 			{
@@ -333,66 +328,6 @@ public partial class GUBP : BuildCommand
 		PrintRunTime();
 	}
 
-	/// <summary>
-	/// Recursively update the ControllingTriggers array for each of the nodes passed in. 
-	/// </summary>
-	/// <param name="NodesToDo">Nodes to update</param>
-	static void FindControllingTriggers(IEnumerable<BuildNode> NodesToDo)
-	{
-		foreach(BuildNode NodeToDo in NodesToDo)
-		{
-			FindControllingTriggers(NodeToDo);
-		}
-	}
-
-	/// <summary>
-	/// Recursively find the controlling triggers for the given node.
-	/// </summary>
-	/// <param name="NodeToDo">Node to find the controlling triggers for</param>
-    static void FindControllingTriggers(BuildNode NodeToDo)
-    {
-		if(NodeToDo.ControllingTriggers == null)
-		{
-			NodeToDo.ControllingTriggers = new TriggerNode[0];
-
-			// Find the immediate trigger controlling this one
-			List<TriggerNode> PreviousTriggers = new List<TriggerNode>();
-			foreach (BuildNode Dependency in NodeToDo.OrderDependencies)
-			{
-				FindControllingTriggers(Dependency);
-
-				TriggerNode PreviousTrigger = Dependency as TriggerNode;
-				if(PreviousTrigger == null && Dependency.ControllingTriggers.Length > 0)
-				{
-					PreviousTrigger = Dependency.ControllingTriggers.Last();
-				}
-				
-				if(PreviousTrigger != null && !PreviousTriggers.Contains(PreviousTrigger))
-				{
-					PreviousTriggers.Add(PreviousTrigger);
-				}
-			}
-
-			// Remove previous triggers from the list that aren't the last in the chain. If there's a trigger chain of X.Y.Z, and a node has dependencies behind all three, the only trigger we care about is Z.
-			PreviousTriggers.RemoveAll(x => PreviousTriggers.Any(y => y.ControllingTriggers.Contains(x)));
-
-			// We only support one direct controlling trigger at the moment (though it may be in a chain with other triggers)
-			if(PreviousTriggers.Count > 1)
-			{
-				throw new AutomationException("Node {0} has multiple controlling triggers: {1}", NodeToDo.Name, String.Join(", ", PreviousTriggers.Select(x => x.Name)));
-			}
-
-			// Update the list of controlling triggers
-			if (PreviousTriggers.Count == 1)
-			{
-				List<TriggerNode> ControllingTriggers = new List<TriggerNode>();
-				ControllingTriggers.AddRange(PreviousTriggers[0].ControllingTriggers);
-				ControllingTriggers.Add(PreviousTriggers[0]);
-				NodeToDo.ControllingTriggers = ControllingTriggers.ToArray();
-			}
-		}
-	}
-
 	/// <summary>
 	/// Get a string describing a time interval, in the format "XhYm".
 	/// </summary>
@@ -1112,187 +1047,6 @@ public partial class GUBP : BuildCommand
         return Result;
     }
 
-	/// <summary>
-	/// Resolves the names of each node and aggregates' dependencies, and links them together into the build graph.
-	/// </summary>
-	/// <param name="AggregateNameToInfo">Map of aggregate names to their info objects</param>
-	/// <param name="NodeNameToInfo">Map of node names to their info objects</param>
-	private static void LinkGraph(List<AggregateNodeTemplate> AggregateNodeTemplates, List<BuildNodeTemplate> BuildNodeTemplates, out List<AggregateNode> AggregateNodes, out List<BuildNode> BuildNodes)
-	{
-		// Create all the build nodes, and a dictionary to find them by name
-		Dictionary<string, BuildNodePair> BuildNodeNameToPair = new Dictionary<string, BuildNodePair>(StringComparer.InvariantCultureIgnoreCase);
-		foreach (BuildNodeTemplate Template in BuildNodeTemplates)
-		{
-			BuildNodeNameToPair.Add(Template.Name, new BuildNodePair(Template));
-		}
-
-		// Create all the aggregate nodes, and add them to a dictionary too
-		Dictionary<string, AggregateNodePair> AggregateNodeNameToPair = new Dictionary<string, AggregateNodePair>(StringComparer.InvariantCultureIgnoreCase);
-		foreach (AggregateNodeTemplate Template in AggregateNodeTemplates)
-		{
-			AggregateNodeNameToPair.Add(Template.Name, new AggregateNodePair(Template));
-		}
-
-		// Link the graph
-		int NumErrors = 0;
-		foreach (AggregateNodePair Pair in AggregateNodeNameToPair.Values)
-		{
-			LinkAggregate(Pair, AggregateNodeNameToPair, BuildNodeNameToPair, ref NumErrors);
-		}
-		foreach (BuildNodePair Pair in BuildNodeNameToPair.Values)
-		{
-			LinkNode(Pair, new List<BuildNodePair>(), AggregateNodeNameToPair, BuildNodeNameToPair, ref NumErrors);
-		}
-		if(NumErrors > 0)
-		{
-			throw new AutomationException("Failed to link graph ({0} errors).", NumErrors);
-		}
-
-		// Set the output lists
-		AggregateNodes = AggregateNodeNameToPair.Values.Select(x => x.Node).ToList();
-		BuildNodes = BuildNodeNameToPair.Values.Select(x => x.Node).ToList();
-	}
-
-	/// <summary>
-	/// Resolves the dependency names in an aggregate to NodeInfo instances, filling in the AggregateInfo.Dependenices array. Any referenced aggregates will also be linked, recursively.
-	/// </summary>
-	/// <param name="Aggregate">The aggregate to link</param>
-	/// <param name="AggregateNameToInfo">Map of other aggregate names to their corresponding instance.</param>
-	/// <param name="NodeNameToInfo">Map from node names to their corresponding instance.</param>
-	/// <param name="NumErrors">The number of errors output so far. Incremented if resolving this aggregate fails.</param>
-	private static void LinkAggregate(AggregateNodePair Pair, Dictionary<string, AggregateNodePair> AggregateNodeNameToPair, Dictionary<string, BuildNodePair> BuildNodeNameToPair, ref int NumErrors)
-	{
-		if (Pair.Node.Dependencies == null)
-		{
-			Pair.Node.Dependencies = new BuildNode[0];
-
-			HashSet<BuildNode> Dependencies = new HashSet<BuildNode>();
-			foreach (string DependencyName in ParseSemicolonDelimitedList(Pair.Template.DependencyNames))
-			{
-				AggregateNodePair AggregateNodeDependency;
-				if(AggregateNodeNameToPair.TryGetValue(DependencyName, out AggregateNodeDependency))
-				{
-					LinkAggregate(AggregateNodeDependency, AggregateNodeNameToPair, BuildNodeNameToPair, ref NumErrors);
-					Dependencies.UnionWith(AggregateNodeDependency.Node.Dependencies);
-					continue;
-				}
-
-				BuildNodePair BuildNodeDependency;
-				if(BuildNodeNameToPair.TryGetValue(DependencyName, out BuildNodeDependency))
-				{
-					Dependencies.Add(BuildNodeDependency.Node);
-					continue;
-				}
-
-				CommandUtils.LogError("Node {0} is not in the graph. It is a dependency of {1}.", DependencyName, Pair.Template.Name);
-				NumErrors++;
-			}
-
-			Pair.Node.Dependencies = Dependencies.ToArray();
-		}
-	}
-
-	private static void LinkNode(BuildNodePair Pair, List<BuildNodePair> Stack, Dictionary<string, AggregateNodePair> AggregateNodeNameToPair, Dictionary<string, BuildNodePair> BuildNodeNameToPair, ref int NumErrors)
-	{
-		if(Pair.Node.OrderDependencies == null)
-		{
-			Stack.Add(Pair);
-
-			int StackIdx = Stack.IndexOf(Pair);
-			if (StackIdx != Stack.Count - 1)
-			{
-				// There's a cycle in the build graph. Print out the chain.
-				CommandUtils.LogError("Build graph contains a cycle ({0})", String.Join(" -> ", Stack.Skip(StackIdx).Select(x => x.Template.Name)));
-				NumErrors++;
-				Pair.Node.OrderDependencies = new HashSet<BuildNode>();
-				Pair.Node.InputDependencies = new HashSet<BuildNode>();
-			}
-			else
-			{
-				// Find all the direct input dependencies
-				HashSet<BuildNode> DirectInputDependencies = new HashSet<BuildNode>();
-				foreach (string InputDependencyName in ParseSemicolonDelimitedList(Pair.Template.InputDependencyNames))
-				{
-					if (!ResolveDependencies(InputDependencyName, AggregateNodeNameToPair, BuildNodeNameToPair, DirectInputDependencies))
-					{
-						CommandUtils.LogError("Node {0} is not in the graph. It is an input dependency of {1}.", InputDependencyName, Pair.Template.Name);
-						NumErrors++;
-					}
-				}
-
-				// Find all the direct order dependencies. All the input dependencies are also order dependencies.
-				HashSet<BuildNode> DirectOrderDependencies = new HashSet<BuildNode>(DirectInputDependencies);
-				foreach (string OrderDependencyName in ParseSemicolonDelimitedList(Pair.Template.OrderDependencyNames))
-				{
-					if (!ResolveDependencies(OrderDependencyName, AggregateNodeNameToPair, BuildNodeNameToPair, DirectOrderDependencies))
-					{
-						CommandUtils.LogError("Node {0} is not in the graph. It is an order dependency of {1}.", OrderDependencyName, Pair.Template.Name);
-						NumErrors++;
-					}
-				}
-
-				// Link everything
-				foreach (BuildNode DirectOrderDependency in DirectOrderDependencies)
-				{
-					LinkNode(BuildNodeNameToPair[DirectOrderDependency.Name], Stack, AggregateNodeNameToPair, BuildNodeNameToPair, ref NumErrors);
-				}
-
-				// Recursively include all the input dependencies
-				Pair.Node.InputDependencies = new HashSet<BuildNode>(DirectInputDependencies);
-				foreach (BuildNode DirectInputDependency in DirectInputDependencies)
-				{
-					Pair.Node.InputDependencies.UnionWith(DirectInputDependency.InputDependencies);
-				}
-
-				// Same for the order dependencies
-				Pair.Node.OrderDependencies = new HashSet<BuildNode>(DirectOrderDependencies);
-				foreach (BuildNode DirectOrderDependency in DirectOrderDependencies)
-				{
-					Pair.Node.OrderDependencies.UnionWith(DirectOrderDependency.OrderDependencies);
-				}
-			}
-
-			Stack.RemoveAt(Stack.Count - 1);
-		}
-	}
-
-	/// <summary>
-	/// Adds all the nodes matching a given name to a hash set, expanding any aggregates to their dependencices.
-	/// </summary>
-	/// <param name="Name">The name to look for</param>
-	/// <param name="AggregateNameToInfo">Map of other aggregate names to their corresponding info instance.</param>
-	/// <param name="NodeNameToInfo">Map from node names to their corresponding info instance.</param>
-	/// <param name="Dependencies">The set of dependencies to add to.</param>
-	/// <returns>True if the name was found (and the dependencies list was updated), false otherwise.</returns>
-	private static bool ResolveDependencies(string Name, Dictionary<string, AggregateNodePair> AggregateNodeNameToPair, Dictionary<string, BuildNodePair> BuildNodeNameToPair, HashSet<BuildNode> Dependencies)
-	{
-		AggregateNodePair AggregateDependency;
-		if (AggregateNodeNameToPair.TryGetValue(Name, out AggregateDependency))
-		{
-			Dependencies.UnionWith(AggregateDependency.Node.Dependencies);
-			return true;
-		}
-
-		BuildNodePair NodeDependency;
-		if (BuildNodeNameToPair.TryGetValue(Name, out NodeDependency))
-		{
-			Dependencies.Add(NodeDependency.Node);
-			return true;
-		}
-
-		return false;
-	}
-
-	/// <summary>
-	/// Parses a semi-colon delimited list into a sequence of strings.
-	/// </summary>
-	/// <param name="StringList">Semicolon delimited list of strings</param>
-	/// <returns>Sequence of strings</returns>
-	private static IEnumerable<string> ParseSemicolonDelimitedList(string StringList)
-	{
-		return StringList.Split(new char[] { ';' }, StringSplitOptions.RemoveEmptyEntries);
-	}
-
 	/// <summary>
 	/// Determine which nodes to build from the command line, or return everything if there's nothing specified explicitly.
 	/// </summary>
