Move ArchiveBuildProducts() and RetrieveBuildProducts() into the BuildNode base class.

diff --git a/Engine/Source/Programs/AutomationTool/GUBP/Nodes/BuildNode.cs b/Engine/Source/Programs/AutomationTool/GUBP/Nodes/BuildNode.cs
--- a/Engine/Source/Programs/AutomationTool/GUBP/Nodes/BuildNode.cs
+++ b/Engine/Source/Programs/AutomationTool/GUBP/Nodes/BuildNode.cs
@@ -82,9 +82,32 @@ namespace AutomationTool
 			RootIfAnyForTempStorage = Template.RootIfAnyForTempStorage;
 		}
 
-		public abstract void ArchiveBuildProducts(string GameNameIfAny, string StorageRootIfAny, TempStorageNodeInfo TempStorageNodeInfo, bool bLocalOnly);
+		public virtual void ArchiveBuildProducts(string GameNameIfAny, string StorageRootIfAny, TempStorageNodeInfo TempStorageNodeInfo, bool bLocalOnly)
+		{
+			TempStorage.StoreToTempStorage(TempStorageNodeInfo, BuildProducts, bLocalOnly, GameNameIfAny, StorageRootIfAny);
+		}
 
-		public abstract void RetrieveBuildProducts(string GameNameIfAny, string StorageRootIfAny, TempStorageNodeInfo TempStorageNodeInfo);
+		public virtual void RetrieveBuildProducts(string GameNameIfAny, string StorageRootIfAny, TempStorageNodeInfo TempStorageNodeInfo)
+		{
+			CommandUtils.LogConsole("***** Retrieving GUBP Node {0} from {1}", Name, TempStorageNodeInfo.GetRelativeDirectory());
+			try
+			{
+				bool WasLocal;
+				BuildProducts = TempStorage.RetrieveFromTempStorage(TempStorageNodeInfo, out WasLocal, GameNameIfAny, StorageRootIfAny);
+			}
+			catch (Exception Ex)
+			{
+				if (GameNameIfAny != "")
+				{
+					bool WasLocal;
+					BuildProducts = TempStorage.RetrieveFromTempStorage(TempStorageNodeInfo, out WasLocal, "", StorageRootIfAny);
+				}
+				else
+				{
+					throw new AutomationException(Ex, "Build Products cannot be found for node {0}", Name);
+				}
+			}
+		}
 
 		public abstract void DoBuild();
 
diff --git a/Engine/Source/Programs/AutomationTool/GUBP/Nodes/LegacyNode.cs b/Engine/Source/Programs/AutomationTool/GUBP/Nodes/LegacyNode.cs
--- a/Engine/Source/Programs/AutomationTool/GUBP/Nodes/LegacyNode.cs
+++ b/Engine/Source/Programs/AutomationTool/GUBP/Nodes/LegacyNode.cs
@@ -56,30 +56,10 @@ namespace AutomationTool
 			Node = Template.Node;
 		}
 
-		public override void ArchiveBuildProducts(string GameNameIfAny, string StorageRootIfAny, TempStorageNodeInfo TempStorageNodeInfo, bool bLocalOnly)
-		{
-			TempStorage.StoreToTempStorage(TempStorageNodeInfo, BuildProducts, bLocalOnly, GameNameIfAny, StorageRootIfAny);
-		}
-
 		public override void RetrieveBuildProducts(string GameNameIfAny, string StorageRootIfAny, TempStorageNodeInfo TempStorageNodeInfo)
 		{
-			CommandUtils.LogConsole("***** Retrieving GUBP Node {0} from {1}", Name, TempStorageNodeInfo.GetRelativeDirectory());
-			bool WasLocal;
-			try
-			{
-				BuildProducts = TempStorage.RetrieveFromTempStorage(TempStorageNodeInfo, out WasLocal, GameNameIfAny, StorageRootIfAny);
-			}
-			catch (Exception Ex)
-			{
-				if (GameNameIfAny != "")
-				{
-					BuildProducts = TempStorage.RetrieveFromTempStorage(TempStorageNodeInfo, out WasLocal, "", StorageRootIfAny);
-				}
-				else
-				{
-					throw new AutomationException(Ex, "Build Products cannot be found for node {0}", Name);
-				}
-			}
+			base.RetrieveBuildProducts(GameNameIfAny, StorageRootIfAny, TempStorageNodeInfo);
+
 			Node.BuildProducts = BuildProducts;
 		}
 
