Remove the FakeFail option; never used it, and it's easy enough to add back in if need be.

diff --git a/Engine/Source/Programs/AutomationTool/GUBP/ElectricCommander.cs b/Engine/Source/Programs/AutomationTool/GUBP/ElectricCommander.cs
--- a/Engine/Source/Programs/AutomationTool/GUBP/ElectricCommander.cs
+++ b/Engine/Source/Programs/AutomationTool/GUBP/ElectricCommander.cs
@@ -255,7 +255,7 @@ namespace AutomationTool
 			}
 		}
 
-        public void DoCommanderSetup(IEnumerable<BuildNode> AllNodes, IEnumerable<AggregateNode> AllAggregates, List<BuildNode> OrderedToDo, List<BuildNode> SortedNodes, int TimeIndex, int TimeQuantum, bool bSkipTriggers, bool bFake, bool bFakeEC, TriggerNode ExplicitTrigger, List<TriggerNode> UnfinishedTriggers, string FakeFail)
+        public void DoCommanderSetup(IEnumerable<BuildNode> AllNodes, IEnumerable<AggregateNode> AllAggregates, List<BuildNode> OrderedToDo, List<BuildNode> SortedNodes, int TimeIndex, int TimeQuantum, bool bSkipTriggers, bool bFake, bool bFakeEC, TriggerNode ExplicitTrigger, List<TriggerNode> UnfinishedTriggers)
 		{
 			// Check there's some nodes in the graph
 			if (OrderedToDo.Count == 0)
diff --git a/Engine/Source/Programs/AutomationTool/GUBP/GUBP.cs b/Engine/Source/Programs/AutomationTool/GUBP/GUBP.cs
--- a/Engine/Source/Programs/AutomationTool/GUBP/GUBP.cs
+++ b/Engine/Source/Programs/AutomationTool/GUBP/GUBP.cs
@@ -315,14 +315,13 @@ public partial class GUBP : BuildCommand
 		}
 		else
 		{
-			string FakeFail = ParseParamValue("FakeFail");
 			if(CommanderSetup)
 			{
-				DoCommanderSetup(EC, Graph.BuildNodes, Graph.AggregateNodes, OrderedToDo, TimeIndex, TimeQuantum, bSkipTriggers, bFake, bFakeEC, ExplicitTrigger, UnfinishedTriggers, FakeFail, JobInfo.IsPreflight);
+				DoCommanderSetup(EC, Graph.BuildNodes, Graph.AggregateNodes, OrderedToDo, TimeIndex, TimeQuantum, bSkipTriggers, bFake, bFakeEC, ExplicitTrigger, UnfinishedTriggers, JobInfo.IsPreflight);
 			}
 			else if(!ParseParam("ListOnly"))
 			{
-                ExecuteNodes(EC, OrderedToDo, bFake, bFakeEC, bSaveSharedTempStorage, JobInfo, FakeFail);
+                ExecuteNodes(EC, OrderedToDo, bFake, bFakeEC, bSaveSharedTempStorage, JobInfo);
 			}
 		}
 		PrintRunTime();
@@ -1307,7 +1306,7 @@ public partial class GUBP : BuildCommand
 		}
 	}
 
-    private void DoCommanderSetup(ElectricCommander EC, IEnumerable<BuildNode> AllNodes, IEnumerable<AggregateNode> AllAggregates, List<BuildNode> OrderedToDo, int TimeIndex, int TimeQuantum, bool bSkipTriggers, bool bFake, bool bFakeEC, TriggerNode ExplicitTrigger, List<TriggerNode> UnfinishedTriggers, string FakeFail, bool bPreflightBuild)
+    private void DoCommanderSetup(ElectricCommander EC, IEnumerable<BuildNode> AllNodes, IEnumerable<AggregateNode> AllAggregates, List<BuildNode> OrderedToDo, int TimeIndex, int TimeQuantum, bool bSkipTriggers, bool bFake, bool bFakeEC, TriggerNode ExplicitTrigger, List<TriggerNode> UnfinishedTriggers, bool bPreflightBuild)
 	{
 		List<BuildNode> SortedNodes = TopologicalSort(new HashSet<BuildNode>(AllNodes), null, SubSort: false, DoNotConsiderCompletion: true);
 		LogConsole("******* {0} GUBP Nodes", SortedNodes.Count);
@@ -1342,10 +1341,10 @@ public partial class GUBP : BuildCommand
 			PrintNodes(this, FilteredOrderedToDo, AllAggregates, UnfinishedTriggers, TimeQuantum);
 		}
 
-		EC.DoCommanderSetup(AllNodes, AllAggregates, FilteredOrderedToDo, SortedNodes, TimeIndex, TimeQuantum, bSkipTriggers, bFake, bFakeEC, ExplicitTrigger, UnfinishedTriggers, FakeFail);
+		EC.DoCommanderSetup(AllNodes, AllAggregates, FilteredOrderedToDo, SortedNodes, TimeIndex, TimeQuantum, bSkipTriggers, bFake, bFakeEC, ExplicitTrigger, UnfinishedTriggers);
 	}
 
-    void ExecuteNodes(ElectricCommander EC, List<BuildNode> OrderedToDo, bool bFake, bool bFakeEC, bool bSaveSharedTempStorage, JobInfo JobInfo, string FakeFail)
+    void ExecuteNodes(ElectricCommander EC, List<BuildNode> OrderedToDo, bool bFake, bool bFakeEC, bool bSaveSharedTempStorage, JobInfo JobInfo)
 	{
         Dictionary<string, BuildNode> BuildProductToNodeMap = new Dictionary<string, BuildNode>();
 		foreach (BuildNode NodeToDo in OrderedToDo)
@@ -1388,10 +1387,6 @@ public partial class GUBP : BuildCommand
 				double BuildDuration = 0.0;
                 try
                 {
-                    if (!String.IsNullOrEmpty(FakeFail) && FakeFail.Equals(NodeToDo.Name, StringComparison.InvariantCultureIgnoreCase))
-                    {
-                        throw new AutomationException("Failing node {0} by request.", NodeToDo.Name);
-                    }
                     if (bFake)
                     {
                         LogConsole("***** FAKE!! Building GUBP Node {0} for {1}", NodeToDo.Name, TempStorageNodeInfo.GetRelativeDirectory());
