Disable MSBuild banner spam in build logs.

diff --git a/Engine/Source/Programs/AutomationTool/AutomationUtils/ScriptCompiler.cs b/Engine/Source/Programs/AutomationTool/AutomationUtils/ScriptCompiler.cs
--- a/Engine/Source/Programs/AutomationTool/AutomationUtils/ScriptCompiler.cs
+++ b/Engine/Source/Programs/AutomationTool/AutomationUtils/ScriptCompiler.cs
@@ -169,7 +169,7 @@ namespace AutomationTool
 			{
 				string Modules = string.Join(";", ModulesToCompile.ToArray());
 				var UATProj = CommandUtils.CombinePaths(CommandUtils.CmdEnv.LocalRoot, @"Engine\Source\Programs\AutomationTool\Scripts\UAT.proj");
-				var CmdLine = String.Format("\"{0}\" /p:Modules=\"{1}\" /p:Configuration={2} /verbosity:minimal", UATProj, Modules, BuildConfig);
+				var CmdLine = String.Format("\"{0}\" /p:Modules=\"{1}\" /p:Configuration={2} /verbosity:minimal /nologo", UATProj, Modules, BuildConfig);
                 Log.TraceInformation("Building Automation projects in parallel...");
                 // supress the run command because it can be long and intimidating, making the logs around this code harder to read.
                 var Result = CommandUtils.Run(CommandUtils.CmdEnv.MsBuildExe, CmdLine, Options: CommandUtils.ERunOptions.Default | CommandUtils.ERunOptions.NoLoggingOfRunCommand | CommandUtils.ERunOptions.LoggingOfRunDuration);
