Treat ectool output as verbose. The XML response is rarely important.

diff --git a/Engine/Source/Programs/AutomationTool/GUBP/ElectricCommander.cs b/Engine/Source/Programs/AutomationTool/GUBP/ElectricCommander.cs
--- a/Engine/Source/Programs/AutomationTool/GUBP/ElectricCommander.cs
+++ b/Engine/Source/Programs/AutomationTool/GUBP/ElectricCommander.cs
@@ -27,7 +27,7 @@ namespace AutomationTool
 			}
 			else
 			{
-				CommandUtils.ERunOptions Opts = CommandUtils.ERunOptions.Default;
+				CommandUtils.ERunOptions Opts = CommandUtils.ERunOptions.Default | CommandUtils.ERunOptions.SpewIsVerbose;
 				if (bQuiet)
 				{
 					Opts = (Opts & ~CommandUtils.ERunOptions.AllowSpew) | CommandUtils.ERunOptions.NoLoggingOfRunCommand;
