Rename the ErrorCodes enum to ExitCode, and ErrorCodes.Error_Success to ExitCode.Success.

diff --git a/Engine/Source/Programs/AutomationTool/Android/AndroidPlatform.Automation.cs b/Engine/Source/Programs/AutomationTool/Android/AndroidPlatform.Automation.cs
--- a/Engine/Source/Programs/AutomationTool/Android/AndroidPlatform.Automation.cs
+++ b/Engine/Source/Programs/AutomationTool/Android/AndroidPlatform.Automation.cs
@@ -63,13 +63,13 @@ public class AndroidPlatform : Platform
 		string PackageName = GetPackageInfo(ApkName, false);
 		if (PackageName == null)
 		{
-			throw new AutomationException(ErrorCodes.Error_FailureGettingPackageInfo, "Failed to get package name from " + ApkName);
+			throw new AutomationException(ExitCode.Error_FailureGettingPackageInfo, "Failed to get package name from " + ApkName);
 		}
 
 		string PackageVersion = GetPackageInfo(ApkName, true);
 		if (PackageVersion == null || PackageVersion.Length == 0)
 		{
-			throw new AutomationException(ErrorCodes.Error_FailureGettingPackageInfo, "Failed to get package version from " + ApkName);
+			throw new AutomationException(ExitCode.Error_FailureGettingPackageInfo, "Failed to get package version from " + ApkName);
 		}
 
 		if (PackageVersion.Length > 0)
@@ -167,7 +167,7 @@ public class AndroidPlatform : Platform
                     if (FileExists_NoExceptions(UE4SOName) == false)
 					{
 						LogConsole("Failed to find game .so " + UE4SOName);
-                        throw new AutomationException(ErrorCodes.Error_MissingExecutable, "Stage Failed. Could not find .so {0}. You may need to build the UE4 project with your target configuration and platform.", UE4SOName);
+                        throw new AutomationException(ExitCode.Error_MissingExecutable, "Stage Failed. Could not find .so {0}. You may need to build the UE4 project with your target configuration and platform.", UE4SOName);
 					}
 				}
 				
@@ -314,7 +314,7 @@ public class AndroidPlatform : Platform
 	{
 		if (SC.StageTargetConfigurations.Count != 1)
 		{
-			throw new AutomationException(ErrorCodes.Error_OnlyOneTargetConfigurationSupported, "Android is currently only able to package one target configuration at a time, but StageTargetConfigurations contained {0} configurations", SC.StageTargetConfigurations.Count);
+			throw new AutomationException(ExitCode.Error_OnlyOneTargetConfigurationSupported, "Android is currently only able to package one target configuration at a time, but StageTargetConfigurations contained {0} configurations", SC.StageTargetConfigurations.Count);
 		}
 
 		var Architectures = UnrealBuildTool.AndroidToolChain.GetAllArchitectures();
@@ -335,11 +335,11 @@ public class AndroidPlatform : Platform
 				// verify the files exist
 				if (!FileExists(ApkName))
 				{
-					throw new AutomationException(ErrorCodes.Error_AppNotFound, "ARCHIVE FAILED - {0} was not found", ApkName);
+					throw new AutomationException(ExitCode.Error_AppNotFound, "ARCHIVE FAILED - {0} was not found", ApkName);
 				}
 				if (!bPackageDataInsideApk && !FileExists(ObbName))
 				{
-                    throw new AutomationException(ErrorCodes.Error_ObbNotFound, "ARCHIVE FAILED - {0} was not found", ObbName);
+                    throw new AutomationException(ExitCode.Error_ObbNotFound, "ARCHIVE FAILED - {0} was not found", ObbName);
 				}
 
 				SC.ArchiveFiles(Path.GetDirectoryName(ApkName), Path.GetFileName(ApkName));
@@ -586,7 +586,7 @@ public class AndroidPlatform : Platform
 					}
 				}
 
-				throw new AutomationException(ErrorCodes.Error_AppInstallFailed, ErrorMessage);
+				throw new AutomationException(ExitCode.Error_AppInstallFailed, ErrorMessage);
 			}
 		}
  
@@ -861,7 +861,7 @@ public class AndroidPlatform : Platform
 		string[] Subdirs = Directory.GetDirectories(path);
         if (Subdirs.Length == 0)
         {
-            throw new AutomationException(ErrorCodes.Error_AndroidBuildToolsPathNotFound, "Failed to find %ANDROID_HOME%/build-tools subdirectory");
+            throw new AutomationException(ExitCode.Error_AndroidBuildToolsPathNotFound, "Failed to find %ANDROID_HOME%/build-tools subdirectory");
         }
 		// we expect there to be one, so use the first one
 		return Path.Combine(Subdirs[0], Utils.IsRunningOnMono ? "aapt" : "aapt.exe");
@@ -926,7 +926,7 @@ public class AndroidPlatform : Platform
 		// if after the fallbacks, we still don't have it, we can't continue
 		if (!AppArchitectures.Contains(DeviceArch))
 		{
-            throw new AutomationException(ErrorCodes.Error_NoApkSuitableForArchitecture, "Unable to run because you don't have an apk that is usable on {0}. Looked for {1}", Params.Device, DeviceArch);
+            throw new AutomationException(ExitCode.Error_NoApkSuitableForArchitecture, "Unable to run because you don't have an apk that is usable on {0}. Looked for {1}", Params.Device, DeviceArch);
 		}
 
 		return DeviceArch;
@@ -976,7 +976,7 @@ public class AndroidPlatform : Platform
 		string PackageName = GetPackageInfo(ApkName, false);
 		if (PackageName == null)
 		{
-            throw new AutomationException(ErrorCodes.Error_FailureGettingPackageInfo, "Failed to get package name from " + ClientApp);
+            throw new AutomationException(ExitCode.Error_FailureGettingPackageInfo, "Failed to get package name from " + ClientApp);
 		}
 
 		if (Params.Prebuilt)
diff --git a/Engine/Source/Programs/AutomationTool/AutomationUtils/AutomationExpection.cs b/Engine/Source/Programs/AutomationTool/AutomationUtils/AutomationExpection.cs
--- a/Engine/Source/Programs/AutomationTool/AutomationUtils/AutomationExpection.cs
+++ b/Engine/Source/Programs/AutomationTool/AutomationUtils/AutomationExpection.cs
@@ -7,10 +7,10 @@ using System.Text;
 namespace AutomationTool
 {
     // NOTE: this needs to be kept in sync with EditorAnalytics.h and iPhonePackager.cs
-    public enum ErrorCodes
+    public enum ExitCode
     {
         Error_UATNotFound = -1,
-        Error_Success = 0,
+        Success = 0,
         Error_Unknown = 1,
         Error_Arguments = 2,
         Error_UnknownCommand = 3,
@@ -65,7 +65,7 @@ namespace AutomationTool
     /// </summary>
     public class AutomationException : System.Exception
 	{
-        public ErrorCodes ErrorCode = ErrorCodes.Error_Unknown;
+        public ExitCode ErrorCode = ExitCode.Error_Unknown;
 
 		public AutomationException()
 		{
@@ -76,7 +76,7 @@ namespace AutomationTool
 		{
 		}
 
-        public AutomationException(ErrorCodes ErrorCode, string Msg)
+        public AutomationException(ExitCode ErrorCode, string Msg)
             : base(Msg)
         {
             this.ErrorCode = ErrorCode;
@@ -93,7 +93,7 @@ namespace AutomationTool
 		{
 		}
 
-        public AutomationException(ErrorCodes ErrorCode, Exception InnerException, string Format, params object[] Args)
+        public AutomationException(ExitCode ErrorCode, Exception InnerException, string Format, params object[] Args)
             : base(string.Format(Format, Args), InnerException)
         {
             this.ErrorCode = ErrorCode;
@@ -102,7 +102,7 @@ namespace AutomationTool
         public AutomationException(string Format, params object[] Args)
 			: base(string.Format(Format, Args)) { }
 
-        public AutomationException(ErrorCodes ErrorCode, string Format, params object[] Args)
+        public AutomationException(ExitCode ErrorCode, string Format, params object[] Args)
             : base(string.Format(Format, Args)) 
         {
             this.ErrorCode = ErrorCode;
diff --git a/Engine/Source/Programs/AutomationTool/AutomationUtils/DeploymentContext.cs b/Engine/Source/Programs/AutomationTool/AutomationUtils/DeploymentContext.cs
--- a/Engine/Source/Programs/AutomationTool/AutomationUtils/DeploymentContext.cs
+++ b/Engine/Source/Programs/AutomationTool/AutomationUtils/DeploymentContext.cs
@@ -594,7 +594,7 @@ public class DeploymentContext //: ProjectParams
 
 		if (FilesAdded == 0 && !bAllowNone && !bIsCombiningMultiplePlatforms)
 		{
-			throw new AutomationException(ErrorCodes.Error_StageMissingFile, "No files found to deploy for {0} with wildcard {1} and exclusions {2}", InPath, Wildcard, ExcludeWildcard);
+			throw new AutomationException(ExitCode.Error_StageMissingFile, "No files found to deploy for {0} with wildcard {1} and exclusions {2}", InPath, Wildcard, ExcludeWildcard);
 		}
 
 		return FilesAdded;
diff --git a/Engine/Source/Programs/AutomationTool/AutomationUtils/ProcessUtils.cs b/Engine/Source/Programs/AutomationTool/AutomationUtils/ProcessUtils.cs
--- a/Engine/Source/Programs/AutomationTool/AutomationUtils/ProcessUtils.cs
+++ b/Engine/Source/Programs/AutomationTool/AutomationUtils/ProcessUtils.cs
@@ -797,7 +797,7 @@ namespace AutomationTool
 
             if (Result > MaxSuccessCode || Result < 0)
             {
-                throw new AutomationException((ErrorCodes)Result.ExitCode, String.Format("Command failed (Result:{3}): {0} {1}. See logfile for details: '{2}' ",
+                throw new AutomationException((ExitCode)Result.ExitCode, String.Format("Command failed (Result:{3}): {0} {1}. See logfile for details: '{2}' ",
                                                 App, CommandLine, Path.GetFileName(Logfile), Result.ExitCode));
             }
             if (Result.Output.Length > 0)
diff --git a/Engine/Source/Programs/AutomationTool/HTML5/HTML5Platform.Automation.cs b/Engine/Source/Programs/AutomationTool/HTML5/HTML5Platform.Automation.cs
--- a/Engine/Source/Programs/AutomationTool/HTML5/HTML5Platform.Automation.cs
+++ b/Engine/Source/Programs/AutomationTool/HTML5/HTML5Platform.Automation.cs
@@ -99,7 +99,7 @@ public class HTML5Platform : Platform
         if (!SC.IsCodeBasedProject && !FileExists_NoExceptions(FullGameExePath))
         {
 			LogConsole("Failed to find game application " + FullGameExePath);
-	        throw new AutomationException(ErrorCodes.Error_MissingExecutable, "Stage Failed. Could not find application {0}. You may need to build the UE4 project with your target configuration and platform.", FullGameExePath);
+	        throw new AutomationException(ExitCode.Error_MissingExecutable, "Stage Failed. Could not find application {0}. You may need to build the UE4 project with your target configuration and platform.", FullGameExePath);
         }
 
 		if (Path.Combine(Path.GetDirectoryName(Params.ProjectGameExeFilename), GameExe) != Path.Combine(PackagePath, GameExe))
diff --git a/Engine/Source/Programs/AutomationTool/IOS/IOSPlatform.Automation.cs b/Engine/Source/Programs/AutomationTool/IOS/IOSPlatform.Automation.cs
--- a/Engine/Source/Programs/AutomationTool/IOS/IOSPlatform.Automation.cs
+++ b/Engine/Source/Programs/AutomationTool/IOS/IOSPlatform.Automation.cs
@@ -134,7 +134,7 @@ public class IOSPlatform : Platform
 		if (!SC.IsCodeBasedProject && !FileExists_NoExceptions(FullExePath))
 		{
 			LogError("Failed to find game binary " + FullExePath);
-			throw new AutomationException(ErrorCodes.Error_MissingExecutable, "Stage Failed. Could not find binary {0}. You may need to build the UE4 project with your target configuration and platform.", FullExePath);
+			throw new AutomationException(ExitCode.Error_MissingExecutable, "Stage Failed. Could not find binary {0}. You may need to build the UE4 project with your target configuration and platform.", FullExePath);
 		}
 
 		//@TODO: We should be able to use this code on both platforms, when the following issues are sorted:
@@ -322,7 +322,7 @@ public class IOSPlatform : Platform
 			// verify the .ipa exists
 			if (!FileExists(ProjectIPA))
 			{
-				throw new AutomationException(ErrorCodes.Error_FailedToCreateIPA, "PACKAGE FAILED - {0} was not created", ProjectIPA);
+				throw new AutomationException(ExitCode.Error_FailedToCreateIPA, "PACKAGE FAILED - {0} was not created", ProjectIPA);
 			}
 
 			if (WorkingCL > 0)
@@ -413,7 +413,7 @@ public class IOSPlatform : Platform
 		}
 		if (Result.ExitCode != 0)
 		{
-			throw new AutomationException(ErrorCodes.Error_FailedToCodeSign, "CodeSign Failed");
+			throw new AutomationException(ExitCode.Error_FailedToCodeSign, "CodeSign Failed");
 		}
 	}
 
@@ -821,9 +821,9 @@ public class IOSPlatform : Platform
 			switch (Result.ExitCode)
 			{
 				case 253:
-                    throw new AutomationException(ErrorCodes.Error_DeviceNotSetupForDevelopment, "Launch Failure");
+                    throw new AutomationException(ExitCode.Error_DeviceNotSetupForDevelopment, "Launch Failure");
 				case 255:
-                    throw new AutomationException(ErrorCodes.Error_DeviceOSNewerThanSDK, "Launch Failure");
+                    throw new AutomationException(ExitCode.Error_DeviceOSNewerThanSDK, "Launch Failure");
 			}
 		}
 	}
diff --git a/Engine/Source/Programs/AutomationTool/Linux/LinuxPlatform.Automation.cs b/Engine/Source/Programs/AutomationTool/Linux/LinuxPlatform.Automation.cs
--- a/Engine/Source/Programs/AutomationTool/Linux/LinuxPlatform.Automation.cs
+++ b/Engine/Source/Programs/AutomationTool/Linux/LinuxPlatform.Automation.cs
@@ -123,7 +123,7 @@ public abstract class BaseLinuxPlatform : Platform
 						if (!SC.IsCodeBasedProject && !FileExists_NoExceptions(Params.ProjectGameExeFilename) && !SC.bIsCombiningMultiplePlatforms)
 						{
 							LogError("Failed to find game binary " + Params.ProjectGameExeFilename);
-                            throw new AutomationException(ErrorCodes.Error_MissingExecutable, "Stage Failed. Could not find game binary {0}. You may need to build the UE4 project with your target configuration and platform.", Params.ProjectGameExeFilename);
+                            throw new AutomationException(ExitCode.Error_MissingExecutable, "Stage Failed. Could not find game binary {0}. You may need to build the UE4 project with your target configuration and platform.", Params.ProjectGameExeFilename);
 						}
 
 					    SC.StageFiles(StagedFileType.NonUFS, CombinePaths(SC.LocalRoot, "Engine/Binaries", SC.PlatformDir), Path.GetFileNameWithoutExtension(Exe), true, null, CommandUtils.CombinePaths(SC.RelativeProjectRootForStage, "Binaries", SC.PlatformDir), false, true, SC.ShortProjectName);
diff --git a/Engine/Source/Programs/AutomationTool/Mac/MacPlatform.Automation.cs b/Engine/Source/Programs/AutomationTool/Mac/MacPlatform.Automation.cs
--- a/Engine/Source/Programs/AutomationTool/Mac/MacPlatform.Automation.cs
+++ b/Engine/Source/Programs/AutomationTool/Mac/MacPlatform.Automation.cs
@@ -133,7 +133,7 @@ public class MacPlatform : Platform
 					if (!SC.IsCodeBasedProject && !Directory.Exists(AbsoluteBundlePath) && !SC.bIsCombiningMultiplePlatforms)
 					{
 						LogError("Failed to find app bundle " + AbsoluteBundlePath);
-						throw new AutomationException(ErrorCodes.Error_MissingExecutable, "Could not find app bundle {0}. You may need to build the UE4 project with your target configuration and platform.", AbsoluteBundlePath);
+						throw new AutomationException(ExitCode.Error_MissingExecutable, "Could not find app bundle {0}. You may need to build the UE4 project with your target configuration and platform.", AbsoluteBundlePath);
 					}
 
 					StageAppBundle(SC, WorkingFileType, CombinePaths(SC.LocalRoot, "Engine/Binaries", SC.PlatformDir, Path.GetFileNameWithoutExtension(Exe) + ".app"), AppBundlePath);
diff --git a/Engine/Source/Programs/AutomationTool/Program.cs b/Engine/Source/Programs/AutomationTool/Program.cs
--- a/Engine/Source/Programs/AutomationTool/Program.cs
+++ b/Engine/Source/Programs/AutomationTool/Program.cs
@@ -21,53 +21,50 @@ namespace AutomationTool
             var CommandLine = SharedUtils.ParseCommandLine();
             LogUtils.InitLogging(CommandLine);
 
-            ErrorCodes ReturnCode = ErrorCodes.Error_Success;
+            ExitCode ReturnCode = ExitCode.Success;
 
-            try
-            {
-                // ensure we can resolve any external assemblies as necessary.
-                AssemblyUtils.InstallAssemblyResolver(Path.GetDirectoryName(Assembly.GetEntryAssembly().GetOriginalLocation()));
-                HostPlatform.Initialize();
+			try
+			{
+				// ensure we can resolve any external assemblies as necessary.
+				AssemblyUtils.InstallAssemblyResolver(Path.GetDirectoryName(Assembly.GetEntryAssembly().GetOriginalLocation()));
+				HostPlatform.Initialize();
 
-                Log.TraceInformation("{2}: Running on {0} as a {1}-bit process.", HostPlatform.Current.GetType().Name, Environment.Is64BitProcess ? 64 : 32, DateTime.UtcNow.ToString("o"));
+				Log.TraceInformation("{2}: Running on {0} as a {1}-bit process.", HostPlatform.Current.GetType().Name, Environment.Is64BitProcess ? 64 : 32, DateTime.UtcNow.ToString("o"));
 
-                XmlConfigLoader.Init();
+				XmlConfigLoader.Init();
 
-                // Log if we're running from the launcher
-                var ExecutingAssemblyLocation = Assembly.GetExecutingAssembly().Location;
-                if (string.Compare(ExecutingAssemblyLocation, Assembly.GetEntryAssembly().GetOriginalLocation(), StringComparison.OrdinalIgnoreCase) != 0)
-                {
-                    Log.TraceVerbose("Executed from AutomationToolLauncher ({0})", ExecutingAssemblyLocation);
-                }
-                Log.TraceVerbose("CWD={0}", Environment.CurrentDirectory);
+				// Log if we're running from the launcher
+				var ExecutingAssemblyLocation = Assembly.GetExecutingAssembly().Location;
+				if (string.Compare(ExecutingAssemblyLocation, Assembly.GetEntryAssembly().GetOriginalLocation(), StringComparison.OrdinalIgnoreCase) != 0)
+				{
+					Log.TraceVerbose("Executed from AutomationToolLauncher ({0})", ExecutingAssemblyLocation);
+				}
+				Log.TraceVerbose("CWD={0}", Environment.CurrentDirectory);
 
-                // Hook up exit callbacks
-                var Domain = AppDomain.CurrentDomain;
-                Domain.ProcessExit += Domain_ProcessExit;
-                Domain.DomainUnload += Domain_ProcessExit;
-                HostPlatform.Current.SetConsoleCtrlHandler(CtrlHandler);
+				// Hook up exit callbacks
+				var Domain = AppDomain.CurrentDomain;
+				Domain.ProcessExit += Domain_ProcessExit;
+				Domain.DomainUnload += Domain_ProcessExit;
+				HostPlatform.Current.SetConsoleCtrlHandler(CtrlHandler);
 
-                var Version = AssemblyUtils.ExecutableVersion;
-                Log.TraceVerbose("{0} ver. {1}", Version.ProductName, Version.ProductVersion);
+				var Version = AssemblyUtils.ExecutableVersion;
+				Log.TraceVerbose("{0} ver. {1}", Version.ProductName, Version.ProductVersion);
 
-                // Don't allow simultaneous execution of AT (in the same branch)
-                InternalUtils.RunSingleInstance(MainProc, CommandLine);
+				// Don't allow simultaneous execution of AT (in the same branch)
+				InternalUtils.RunSingleInstance(MainProc, CommandLine);
 
-            }
-            catch (Exception Ex)
-            {
-                // Catch all exceptions and propagate the ErrorCode if we are given one.
-                Log.TraceError("AutomationTool terminated with exception: {0}", Ex);
-                // set the exit code of the process
-                if (Ex is AutomationException)
-                {
-                    ReturnCode = (Ex as AutomationException).ErrorCode;
-                }
-                else
-                {
-                    ReturnCode = ErrorCodes.Error_Unknown;
-                }
-            }
+			}
+			catch (AutomationException Ex)
+			{
+				Log.TraceError("AutomationTool terminated with exception: {0}", Ex);
+				ReturnCode = Ex.ErrorCode;
+			}
+			catch (Exception Ex)
+			{
+				// Catch all exceptions and propagate the ErrorCode if we are given one.
+				Log.TraceError("AutomationTool terminated with exception: {0}", Ex);
+				ReturnCode = ExitCode.Error_Unknown;
+			}
             finally
             {
                 // In all cases, do necessary shut down stuff, but don't let any additional exceptions leak out while trying to shut down.
@@ -78,7 +75,7 @@ namespace AutomationTool
                 // Try to kill process before app domain exits to leave the other KillAll call to extreme edge cases
                 NoThrow(() => { if (ShouldKillProcesses && !Utils.IsRunningOnMono) ProcessManager.KillAll(); }, "Kill All Processes");
 
-                Log.TraceInformation("AutomationTool exiting with ExitCode={0}", ReturnCode);
+                Log.TraceInformation("AutomationTool exiting with ExitCode={0} ({1})", (int)ReturnCode, ReturnCode);
 
                 // Can't use NoThrow here because the code logs exceptions. We're shutting down logging!
                 LogUtils.ShutdownLogging();
diff --git a/Engine/Source/Programs/AutomationTool/Scripts/CookCommand.Automation.cs b/Engine/Source/Programs/AutomationTool/Scripts/CookCommand.Automation.cs
--- a/Engine/Source/Programs/AutomationTool/Scripts/CookCommand.Automation.cs
+++ b/Engine/Source/Programs/AutomationTool/Scripts/CookCommand.Automation.cs
@@ -243,7 +243,7 @@ public partial class Project : CommandUtils
 					// Delete cooked data (if any) as it may be incomplete / corrupted.
 					LogConsole("Cook failed. Deleting cooked data.");
 					CleanupCookedData(PlatformsToCook.ToList(), Params);
-					throw new AutomationException(ErrorCodes.Error_UnknownCookFailure, Ex, "Cook failed.");
+					throw new AutomationException(ExitCode.Error_UnknownCookFailure, Ex, "Cook failed.");
 				}
 			}
 
@@ -258,7 +258,7 @@ public partial class Project : CommandUtils
                     // Delete cooked data (if any) as it may be incomplete / corrupted.
                     LogConsole("Cook failed. Deleting cooked data.");
                     CleanupCookedData(PlatformsToCook.ToList(), Params);
-                    throw new AutomationException(ErrorCodes.Error_UnknownCookFailure, Ex, "Cook failed.");
+                    throw new AutomationException(ExitCode.Error_UnknownCookFailure, Ex, "Cook failed.");
                 }
             }
             
diff --git a/Engine/Source/Programs/AutomationTool/Scripts/CopyBuildToStagingDirectory.Automation.cs b/Engine/Source/Programs/AutomationTool/Scripts/CopyBuildToStagingDirectory.Automation.cs
--- a/Engine/Source/Programs/AutomationTool/Scripts/CopyBuildToStagingDirectory.Automation.cs
+++ b/Engine/Source/Programs/AutomationTool/Scripts/CopyBuildToStagingDirectory.Automation.cs
@@ -1069,7 +1069,7 @@ public partial class Project : CommandUtils
             catch (Exception Ex)
             {
                 // Delete cooked data (if any) as it may be incomplete / corrupted.
-                throw new AutomationException(ErrorCodes.Error_FailedToDeleteStagingDirectory, Ex, "Stage Failed. Failed to delete staging directory " + SC.StageDirectory);
+                throw new AutomationException(ExitCode.Error_FailedToDeleteStagingDirectory, Ex, "Stage Failed. Failed to delete staging directory " + SC.StageDirectory);
             }
 		}
 		else
@@ -1082,7 +1082,7 @@ public partial class Project : CommandUtils
             catch (Exception Ex)
             {
                 // Delete cooked data (if any) as it may be incomplete / corrupted.
-                throw new AutomationException(ErrorCodes.Error_FailedToDeleteStagingDirectory, Ex, "Stage Failed. Failed to delete pak files in " + SC.StageDirectory);
+                throw new AutomationException(ExitCode.Error_FailedToDeleteStagingDirectory, Ex, "Stage Failed. Failed to delete pak files in " + SC.StageDirectory);
             }
 		}
 		if (ShouldCreatePak(Params, SC))
@@ -1559,7 +1559,7 @@ public partial class Project : CommandUtils
 							if (bRequireStagedFilesToExist)
 							{
 								// if we aren't collecting multiple platforms, then it is expected to exist
-                                throw new AutomationException(ErrorCodes.Error_MissingExecutable, "Stage Failed. Missing receipt '{0}'. Check that this target has been built.", Path.GetFileName(ReceiptFileName));
+                                throw new AutomationException(ExitCode.Error_MissingExecutable, "Stage Failed. Missing receipt '{0}'. Check that this target has been built.", Path.GetFileName(ReceiptFileName));
 							}
 							else
 							{
