Reduce the amount of boilerplate log output when running UAT.

diff --git a/Engine/Build/BatchFiles/RunUAT.bat b/Engine/Build/BatchFiles/RunUAT.bat
--- a/Engine/Build/BatchFiles/RunUAT.bat
+++ b/Engine/Build/BatchFiles/RunUAT.bat
@@ -29,8 +29,6 @@ if not exist Source\Programs\AutomationToolLauncher\AutomationToolLauncher.cspro
 rem ## Check to see if we're already running under a Visual Studio environment shell
 if not "%INCLUDE%" == "" if not "%LIB%" == "" goto ReadyToCompile
 
-echo path="%path%"
-
 rem ## Check for Visual Studio 2013
 
 pushd %~dp0
diff --git a/Engine/Source/Programs/AutomationTool/AutomationUtils/Automation.cs b/Engine/Source/Programs/AutomationTool/AutomationUtils/Automation.cs
--- a/Engine/Source/Programs/AutomationTool/AutomationUtils/Automation.cs
+++ b/Engine/Source/Programs/AutomationTool/AutomationUtils/Automation.cs
@@ -423,8 +423,6 @@ AutomationTool.exe [-verbose] [-compileonly] [-p4] Command0 [-Arg0 -Arg1 -Arg2 
 					CommandUtils.ClearDirStack();
 				}
 			}
-
-			Log.TraceInformation("Script execution successful, exiting.");
 		}
 
 		#endregion
diff --git a/Engine/Source/Programs/AutomationTool/AutomationUtils/BuildCommand.cs b/Engine/Source/Programs/AutomationTool/AutomationUtils/BuildCommand.cs
--- a/Engine/Source/Programs/AutomationTool/AutomationUtils/BuildCommand.cs
+++ b/Engine/Source/Programs/AutomationTool/AutomationUtils/BuildCommand.cs
@@ -3,52 +3,54 @@ using System;
 using System.Collections.Generic;
 using System.Linq;
 using System.Text;
-using AutomationTool;
 
-/// <summary>
-/// Base class for buildcommands.
-/// </summary>
-public abstract class BuildCommand : CommandUtils
+namespace AutomationTool
 {
-
-	#region Interface
-
 	/// <summary>
-	/// Build command entry point.  Throws AutomationExceptions on failure.
+	/// Base class for buildcommands.
 	/// </summary>
-	public abstract void ExecuteBuild();
+	public abstract class BuildCommand : CommandUtils
+	{
 
-	#endregion
+		#region Interface
 
-	/// <summary>
-	/// Command entry point.
-	/// </summary>
-	public void Execute()
-	{
-		try
-		{
-			ExecuteBuild();
-		}
-		catch
+		/// <summary>
+		/// Build command entry point.  Throws AutomationExceptions on failure.
+		/// </summary>
+		public abstract void ExecuteBuild();
+
+		#endregion
+
+		/// <summary>
+		/// Command entry point.
+		/// </summary>
+		public void Execute()
 		{
-			LogError("BUILD FAILED");
-			throw;
+			try
+			{
+				ExecuteBuild();
+			}
+			catch
+			{
+				LogError("BUILD FAILED");
+				throw;
+			}
+			LogConsole("BUILD SUCCESSFUL");
 		}
-		LogConsole("BUILD SUCCESSFUL");
-	}
 
-	/// <summary>
-	/// Executes a new command as a child of another command.
-	/// </summary>
-	/// <typeparam name="T"></typeparam>
-	/// <param name="ParentCommand"></param>
-	public static void Execute<T>(BuildCommand ParentCommand) where T : BuildCommand, new()
-	{
-		T Command = new T();
-		if (ParentCommand != null)
+		/// <summary>
+		/// Executes a new command as a child of another command.
+		/// </summary>
+		/// <typeparam name="T"></typeparam>
+		/// <param name="ParentCommand"></param>
+		public static void Execute<T>(BuildCommand ParentCommand) where T : BuildCommand, new()
 		{
-			Command.Params = ParentCommand.Params;
+			T Command = new T();
+			if (ParentCommand != null)
+			{
+				Command.Params = ParentCommand.Params;
+			}
+			Command.Execute();
 		}
-		Command.Execute();
 	}
-}
+}
\ No newline at end of file
diff --git a/Engine/Source/Programs/AutomationTool/AutomationUtils/ScriptCompiler.cs b/Engine/Source/Programs/AutomationTool/AutomationUtils/ScriptCompiler.cs
--- a/Engine/Source/Programs/AutomationTool/AutomationUtils/ScriptCompiler.cs
+++ b/Engine/Source/Programs/AutomationTool/AutomationUtils/ScriptCompiler.cs
@@ -170,7 +170,7 @@ namespace AutomationTool
 				string Modules = string.Join(";", ModulesToCompile.ToArray());
 				var UATProj = CommandUtils.CombinePaths(CommandUtils.CmdEnv.LocalRoot, @"Engine\Source\Programs\AutomationTool\Scripts\UAT.proj");
 				var CmdLine = String.Format("\"{0}\" /p:Modules=\"{1}\" /p:Configuration={2} /verbosity:minimal /nologo", UATProj, Modules, BuildConfig);
-                Log.TraceInformation("Building Automation projects in parallel...");
+                Log.TraceVerbose("Building Automation projects in parallel...");
                 // supress the run command because it can be long and intimidating, making the logs around this code harder to read.
                 var Result = CommandUtils.Run(CommandUtils.CmdEnv.MsBuildExe, CmdLine, Options: CommandUtils.ERunOptions.Default | CommandUtils.ERunOptions.NoLoggingOfRunCommand | CommandUtils.ERunOptions.LoggingOfRunDuration);
 				if (Result.ExitCode != 0)
@@ -273,7 +273,7 @@ namespace AutomationTool
 		/// <param name="OutScriptAssemblies">List to store all loaded assemblies.</param>
 		private static void LoadPreCompiledScriptAssemblies(List<Assembly> OutScriptAssemblies)
 		{
-			CommandUtils.LogConsole("Loading precompiled script DLLs");
+			CommandUtils.LogVerbose("Loading precompiled script DLLs");
 
 			bool DefaultScriptsDLLFound = false;
 			var ScriptsLocation = GetScriptAssemblyFolder();
@@ -281,7 +281,7 @@ namespace AutomationTool
 			{
 				var ScriptDLLFiles = Directory.GetFiles(ScriptsLocation, "*.Automation.dll", SearchOption.AllDirectories);
 
-				CommandUtils.LogConsole("Found {0} script DLL(s).", ScriptDLLFiles.Length);
+				CommandUtils.LogVerbose("Found {0} script DLL(s).", ScriptDLLFiles.Length);
 				foreach (var ScriptsDLLFilename in ScriptDLLFiles)
 				{
 
@@ -322,7 +322,7 @@ namespace AutomationTool
 
 		private void CleanupScriptsAssemblies()
 		{
-			Log.TraceInformation("Cleaning up script DLL folder");
+			CommandUtils.LogVerbose("Cleaning up script DLL folder");
 			CommandUtils.DeleteDirectory(GetScriptAssemblyFolder());
 		}
 
diff --git a/Engine/Source/Programs/AutomationTool/GUBP/ElectricCommander.cs b/Engine/Source/Programs/AutomationTool/GUBP/ElectricCommander.cs
--- a/Engine/Source/Programs/AutomationTool/GUBP/ElectricCommander.cs
+++ b/Engine/Source/Programs/AutomationTool/GUBP/ElectricCommander.cs
@@ -5,6 +5,7 @@ using System.IO;
 using System.Linq;
 using System.Text;
 using System.Xml;
+using AutomationTool;
 using UnrealBuildTool;
 
 namespace AutomationTool
diff --git a/Engine/Source/Programs/AutomationTool/GUBP/GUBP.cs b/Engine/Source/Programs/AutomationTool/GUBP/GUBP.cs
--- a/Engine/Source/Programs/AutomationTool/GUBP/GUBP.cs
+++ b/Engine/Source/Programs/AutomationTool/GUBP/GUBP.cs
@@ -602,8 +602,6 @@ public partial class GUBP : BuildCommand
 
     static List<BuildNode> TopologicalSort(HashSet<BuildNode> NodesToDo, BuildNode ExplicitTrigger, bool SubSort, bool DoNotConsiderCompletion)
     {
-        DateTime StartTime = DateTime.UtcNow;
-
         List<BuildNode> OrderedToDo = new List<BuildNode>();
 
         Dictionary<string, List<BuildNode>> SortedAgentGroupChains = new Dictionary<string, List<BuildNode>>();
@@ -638,7 +636,6 @@ public partial class GUBP : BuildCommand
 					throw new AutomationException("Agent sharing group '{0}' has multiple controlling triggers: {1}", Chain.Key, Triggers);
 				}
 			}
-			LogConsole("***************Done with recursion");
         }
 
         // here we do a topological sort of the nodes, subject to a lexographical and priority sort
@@ -826,11 +823,6 @@ public partial class GUBP : BuildCommand
                 throw new AutomationException("Cycle in GUBP");
             }
         }
-        if (!SubSort)
-        {
-            double BuildDuration = (DateTime.UtcNow - StartTime).TotalMilliseconds;
-			LogConsole("Took {0}s to sort {1} nodes", BuildDuration / 1000, OrderedToDo.Count);
-        }
 
         return OrderedToDo;
     }
diff --git a/Engine/Source/Programs/UnrealBuildTool/System/Utils.cs b/Engine/Source/Programs/UnrealBuildTool/System/Utils.cs
--- a/Engine/Source/Programs/UnrealBuildTool/System/Utils.cs
+++ b/Engine/Source/Programs/UnrealBuildTool/System/Utils.cs
@@ -1087,7 +1087,7 @@ namespace UnrealBuildTool
 				case LogEventType.Warning:
 					return "WARNING: ";
 				case LogEventType.Console:
-					return "DISPLAY: ";
+					return "";
 				case LogEventType.Verbose:
 					return "VERBOSE: ";
 				case LogEventType.VeryVerbose:
