Move FileFilter inside the AutomationTool namespace.

diff --git a/Engine/Source/Programs/AutomationTool/AutomationUtils/FileFilter.cs b/Engine/Source/Programs/AutomationTool/AutomationUtils/FileFilter.cs
--- a/Engine/Source/Programs/AutomationTool/AutomationUtils/FileFilter.cs
+++ b/Engine/Source/Programs/AutomationTool/AutomationUtils/FileFilter.cs
@@ -9,615 +9,618 @@ using System.Linq;
 using System.Text;
 using System.Text.RegularExpressions;
 
-/// <summary>
-/// Indicates whether files which match a pattern should be included or excluded
-/// </summary>
-public enum FileFilterType
+namespace AutomationTool
 {
-	Include,
-	Exclude,
-}
-
-/// <summary>
-/// Stores a set of rules, similar to p4 path specifications, which can be used to efficiently filter files in or out of a set.
-/// The rules are merged into a tree with nodes for each path fragments and rule numbers in each leaf, allowing tree traversal in an arbitrary order 
-/// to how they are applied.
-/// </summary>
-public class FileFilter
-{
-	/// <summary>
-	/// Root node for the tree.
-	/// </summary>
-	FileFilterNode RootNode;
-
-	/// <summary>
-	/// The default node, which will match any path
-	/// </summary>
-	FileFilterNode DefaultNode;
-
 	/// <summary>
-	/// Terminating nodes for each rule added to the filter
+	/// Indicates whether files which match a pattern should be included or excluded
 	/// </summary>
-	List<FileFilterNode> Rules = new List<FileFilterNode>();
-
-	/// <summary>
-	/// Default constructor
-	/// </summary>
-	public FileFilter(FileFilterType DefaultType = FileFilterType.Exclude)
+	public enum FileFilterType
 	{
-		RootNode = new FileFilterNode(null, "");
-
-		DefaultNode = new FileFilterNode(RootNode, "...");
-		DefaultNode.Type = DefaultType;
+		Include,
+		Exclude,
 	}
 
 	/// <summary>
-	/// Copy constructor
+	/// Stores a set of rules, similar to p4 path specifications, which can be used to efficiently filter files in or out of a set.
+	/// The rules are merged into a tree with nodes for each path fragments and rule numbers in each leaf, allowing tree traversal in an arbitrary order 
+	/// to how they are applied.
 	/// </summary>
-	/// <param name="Other">Filter to copy from</param>
-	public FileFilter(FileFilter Other)
-		: this(Other.DefaultNode.Type)
+	public class FileFilter
 	{
-		foreach (FileFilterNode OtherRule in Other.Rules)
+		/// <summary>
+		/// Root node for the tree.
+		/// </summary>
+		FileFilterNode RootNode;
+
+		/// <summary>
+		/// The default node, which will match any path
+		/// </summary>
+		FileFilterNode DefaultNode;
+
+		/// <summary>
+		/// Terminating nodes for each rule added to the filter
+		/// </summary>
+		List<FileFilterNode> Rules = new List<FileFilterNode>();
+
+		/// <summary>
+		/// Default constructor
+		/// </summary>
+		public FileFilter(FileFilterType DefaultType = FileFilterType.Exclude)
 		{
-			AddRule(OtherRule.ToString(), OtherRule.Type);
+			RootNode = new FileFilterNode(null, "");
+
+			DefaultNode = new FileFilterNode(RootNode, "...");
+			DefaultNode.Type = DefaultType;
 		}
-	}
 
-	/// <summary>
-	/// Constructs a file filter from a p4-style filespec. Exclude lines are prefixed with a - character.
-	/// </summary>
-	/// <param name="Lines">Lines to parse rules from</param>
-	public FileFilter(IEnumerable<string> Lines)
-		: this()
-	{
-		foreach (string Line in Lines)
+		/// <summary>
+		/// Copy constructor
+		/// </summary>
+		/// <param name="Other">Filter to copy from</param>
+		public FileFilter(FileFilter Other)
+			: this(Other.DefaultNode.Type)
 		{
-			AddRule(Line);
+			foreach (FileFilterNode OtherRule in Other.Rules)
+			{
+				AddRule(OtherRule.ToString(), OtherRule.Type);
+			}
 		}
-	}
 
-	/// <summary>
-	/// Adds an include or exclude rule to the filter
-	/// </summary>
-	/// <param name="Pattern">Pattern to match. See CreateRegex() for details.</param>
-	public void AddRule(string Rule)
-	{
-		if (Rule.StartsWith("-"))
+		/// <summary>
+		/// Constructs a file filter from a p4-style filespec. Exclude lines are prefixed with a - character.
+		/// </summary>
+		/// <param name="Lines">Lines to parse rules from</param>
+		public FileFilter(IEnumerable<string> Lines)
+			: this()
 		{
-			Exclude(Rule.Substring(1).TrimStart());
+			foreach (string Line in Lines)
+			{
+				AddRule(Line);
+			}
 		}
-		else
+
+		/// <summary>
+		/// Adds an include or exclude rule to the filter
+		/// </summary>
+		/// <param name="Pattern">Pattern to match. See CreateRegex() for details.</param>
+		public void AddRule(string Rule)
 		{
-			Include(Rule);
+			if (Rule.StartsWith("-"))
+			{
+				Exclude(Rule.Substring(1).TrimStart());
+			}
+			else
+			{
+				Include(Rule);
+			}
 		}
-	}
 
-	/// <summary>
-	/// Adds an include or exclude rule to the filter. The rule may be 
-	/// </summary>
-	/// <param name="Pattern">Pattern to match. See CreateRegex() for details.</param>
-	public void AddRule(string Rule, params string[] AllowTags)
-	{
-		string CleanRule = Rule.Trim();
-		if(CleanRule.StartsWith("{"))
+		/// <summary>
+		/// Adds an include or exclude rule to the filter. The rule may be 
+		/// </summary>
+		/// <param name="Pattern">Pattern to match. See CreateRegex() for details.</param>
+		public void AddRule(string Rule, params string[] AllowTags)
 		{
-			// Find the end of the condition
-			int ConditionEnd = CleanRule.IndexOf('}');
-			if(ConditionEnd == -1)
+			string CleanRule = Rule.Trim();
+			if(CleanRule.StartsWith("{"))
 			{
-				throw new Exception(String.Format("Missing closing parenthesis in rule: {0}", CleanRule));
+				// Find the end of the condition
+				int ConditionEnd = CleanRule.IndexOf('}');
+				if(ConditionEnd == -1)
+				{
+					throw new Exception(String.Format("Missing closing parenthesis in rule: {0}", CleanRule));
+				}
+
+				// Check there's a matching tag
+				string[] RuleTags = CleanRule.Substring(1, ConditionEnd - 1).Split(',').Select(x => x.Trim()).ToArray();
+				if(!RuleTags.Any(x => AllowTags.Contains(x)))
+				{
+					return;
+				}
+
+				// Strip the condition from the rule
+				CleanRule = CleanRule.Substring(ConditionEnd + 1).TrimStart(); 
 			}
+			AddRule(CleanRule);
+		}
 
-			// Check there's a matching tag
-			string[] RuleTags = CleanRule.Substring(1, ConditionEnd - 1).Split(',').Select(x => x.Trim()).ToArray();
-			if(!RuleTags.Any(x => AllowTags.Contains(x)))
+		/// <summary>
+		/// Adds several rules to the filter
+		/// </summary>
+		/// <param name="Patterns">List of patterns to match.</param>
+		public void AddRules(IEnumerable<string> Rules)
+		{
+			foreach(string Rule in Rules)
 			{
-				return;
+				AddRule(Rule);
 			}
-
-			// Strip the condition from the rule
-			CleanRule = CleanRule.Substring(ConditionEnd + 1).TrimStart(); 
 		}
-		AddRule(CleanRule);
-	}
 
-	/// <summary>
-	/// Adds several rules to the filter
-	/// </summary>
-	/// <param name="Patterns">List of patterns to match.</param>
-	public void AddRules(IEnumerable<string> Rules)
-	{
-		foreach(string Rule in Rules)
+		/// <summary>
+		/// Adds several rules in the given lines. Rules may be prefixed with conditions of the syntax {Key=Value, Key2=Value2}, which
+		/// will be evaluated using variables in the given dictionary before being added.
+		/// </summary>
+		/// <param name="Lines"></param>
+		/// <param name="Variables">Lookup for variables to test against</param>
+		public void AddRules(IEnumerable<string> Rules, params string[] Tags)
 		{
-			AddRule(Rule);
+			foreach(string Rule in Rules)
+			{
+				AddRule(Rule, Tags);
+			}
 		}
-	}
 
-	/// <summary>
-	/// Adds several rules in the given lines. Rules may be prefixed with conditions of the syntax {Key=Value, Key2=Value2}, which
-	/// will be evaluated using variables in the given dictionary before being added.
-	/// </summary>
-	/// <param name="Lines"></param>
-	/// <param name="Variables">Lookup for variables to test against</param>
-	public void AddRules(IEnumerable<string> Rules, params string[] Tags)
-	{
-		foreach(string Rule in Rules)
+		/// <summary>
+		/// Adds a rule which matches a filename relative to a given base directory.
+		/// </summary>
+		/// <param name="FileName">The filename to add a rule for</param>
+		/// <param name="BaseDirectoryName">Base directory for the rule</param>
+		/// <param name="Type">Whether to add an include or exclude rule</param>
+		public void AddRuleForFile(string FileName, string BaseDirectoryName, FileFilterType Type)
 		{
-			AddRule(Rule, Tags);
+			AddRuleForFiles(new string[]{ FileName }, BaseDirectoryName, Type);
 		}
-	}
-
-	/// <summary>
-	/// Adds a rule which matches a filename relative to a given base directory.
-	/// </summary>
-	/// <param name="FileName">The filename to add a rule for</param>
-	/// <param name="BaseDirectoryName">Base directory for the rule</param>
-	/// <param name="Type">Whether to add an include or exclude rule</param>
-	public void AddRuleForFile(string FileName, string BaseDirectoryName, FileFilterType Type)
-	{
-		AddRuleForFiles(new string[]{ FileName }, BaseDirectoryName, Type);
-	}
 
-	/// <summary>
-	/// Adds rules for files which match the given names
-	/// </summary>
-	/// <param name="FileName">The filenames to rules for</param>
-	/// <param name="BaseDirectoryName">Base directory for the rules</param>
-	/// <param name="Type">Whether to add an include or exclude rule</param>
-	public void AddRuleForFiles(IEnumerable<string> FileNames, string BaseDirectoryName, FileFilterType Type)
-	{
-		string FullBaseDirectoryName = Path.GetFullPath(BaseDirectoryName);
-		foreach (string FileName in FileNames)
+		/// <summary>
+		/// Adds rules for files which match the given names
+		/// </summary>
+		/// <param name="FileName">The filenames to rules for</param>
+		/// <param name="BaseDirectoryName">Base directory for the rules</param>
+		/// <param name="Type">Whether to add an include or exclude rule</param>
+		public void AddRuleForFiles(IEnumerable<string> FileNames, string BaseDirectoryName, FileFilterType Type)
 		{
-			AddRule("/" + CommandUtils.StripBaseDirectory(Path.GetFullPath(FileName), FullBaseDirectoryName), Type);
+			string FullBaseDirectoryName = Path.GetFullPath(BaseDirectoryName);
+			foreach (string FileName in FileNames)
+			{
+				AddRule("/" + CommandUtils.StripBaseDirectory(Path.GetFullPath(FileName), FullBaseDirectoryName), Type);
+			}
 		}
-	}
 
-	/// <summary>
-	/// Reads a configuration file split into sections
-	/// </summary>
-	/// <param name="Filter"></param>
-	/// <param name="RulesFileName"></param>
-	/// <param name="Conditions"></param>
-	public void ReadRulesFromFile(string FileName, string SectionName, params string[] AllowTags)
-	{
-		bool bInSection = false;
-		foreach(string Line in File.ReadAllLines(FileName))
+		/// <summary>
+		/// Reads a configuration file split into sections
+		/// </summary>
+		/// <param name="Filter"></param>
+		/// <param name="RulesFileName"></param>
+		/// <param name="Conditions"></param>
+		public void ReadRulesFromFile(string FileName, string SectionName, params string[] AllowTags)
 		{
-			string TrimLine = Line.Trim();
-			if(!TrimLine.StartsWith(";") && TrimLine.Length > 0)
+			bool bInSection = false;
+			foreach(string Line in File.ReadAllLines(FileName))
 			{
-				if(TrimLine.StartsWith("["))
+				string TrimLine = Line.Trim();
+				if(!TrimLine.StartsWith(";") && TrimLine.Length > 0)
 				{
-					bInSection = (TrimLine == "[" + SectionName + "]");
-				}
-				else if(bInSection)
-				{
-					AddRule(Line, AllowTags);
+					if(TrimLine.StartsWith("["))
+					{
+						bInSection = (TrimLine == "[" + SectionName + "]");
+					}
+					else if(bInSection)
+					{
+						AddRule(Line, AllowTags);
+					}
 				}
 			}
 		}
-	}
-
-	/// <summary>
-	/// Adds an include rule to the filter
-	/// </summary>
-	/// <param name="Pattern">Pattern to match. See CreateRegex() for details.</param>
-	public void Include(string Pattern)
-	{
-		AddRule(Pattern, FileFilterType.Include);
-	}
 
-	/// <summary>
-	/// Adds several exclude rules to the filter
-	/// </summary>
-	/// <param name="Patterns">Patterns to match. See CreateRegex() for details.</param>
-	public void Include(IEnumerable<string> Patterns)
-	{
-		foreach (string Pattern in Patterns)
+		/// <summary>
+		/// Adds an include rule to the filter
+		/// </summary>
+		/// <param name="Pattern">Pattern to match. See CreateRegex() for details.</param>
+		public void Include(string Pattern)
 		{
-			Include(Pattern);
+			AddRule(Pattern, FileFilterType.Include);
 		}
-	}
-
-	/// <summary>
-	/// Adds an exclude rule to the filter
-	/// </summary>
-	/// <param name="Pattern">Mask to match. See CreateRegex() for details.</param>
-	public void Exclude(string Pattern)
-	{
-		AddRule(Pattern, FileFilterType.Exclude);
-	}
 
-	/// <summary>
-	/// Adds several exclude rules to the filter
-	/// </summary>
-	/// <param name="Patterns">Patterns to match. See CreateRegex() for details.</param>
-	public void Exclude(IEnumerable<string> Patterns)
-	{
-		foreach (string Pattern in Patterns)
+		/// <summary>
+		/// Adds several exclude rules to the filter
+		/// </summary>
+		/// <param name="Patterns">Patterns to match. See CreateRegex() for details.</param>
+		public void Include(IEnumerable<string> Patterns)
 		{
-			Exclude(Pattern);
+			foreach (string Pattern in Patterns)
+			{
+				Include(Pattern);
+			}
 		}
-	}
-
-	/// <summary>
-	/// Adds an include or exclude rule to the filter
-	/// </summary>
-	/// <param name="Pattern">The pattern which the rule should match</param>
-	/// <param name="bInclude">Whether to include or exclude files matching this rule</param>
-	public void AddRule(string Pattern, FileFilterType Type)
-	{
-		string NormalizedPattern = Pattern.Replace('\\', '/');
 
-		// We don't want a slash at the start, but if there was not one specified, it's not anchored to the root of the tree.
-		if (NormalizedPattern.StartsWith("/"))
+		/// <summary>
+		/// Adds an exclude rule to the filter
+		/// </summary>
+		/// <param name="Pattern">Mask to match. See CreateRegex() for details.</param>
+		public void Exclude(string Pattern)
 		{
-			NormalizedPattern = NormalizedPattern.Substring(1);
-		}
-		else if(!NormalizedPattern.StartsWith("..."))
-		{
-			NormalizedPattern = ".../" + NormalizedPattern;
+			AddRule(Pattern, FileFilterType.Exclude);
 		}
 
-		// All directories indicate a wildcard match
-		if (NormalizedPattern.EndsWith("/"))
+		/// <summary>
+		/// Adds several exclude rules to the filter
+		/// </summary>
+		/// <param name="Patterns">Patterns to match. See CreateRegex() for details.</param>
+		public void Exclude(IEnumerable<string> Patterns)
 		{
-			NormalizedPattern += "...";
+			foreach (string Pattern in Patterns)
+			{
+				Exclude(Pattern);
+			}
 		}
 
-		// Replace any directory wildcards mid-string
-		for (int Idx = NormalizedPattern.IndexOf("..."); Idx != -1; Idx = NormalizedPattern.IndexOf("...", Idx))
+		/// <summary>
+		/// Adds an include or exclude rule to the filter
+		/// </summary>
+		/// <param name="Pattern">The pattern which the rule should match</param>
+		/// <param name="bInclude">Whether to include or exclude files matching this rule</param>
+		public void AddRule(string Pattern, FileFilterType Type)
 		{
-			if (Idx > 0 && NormalizedPattern[Idx - 1] != '/')
+			string NormalizedPattern = Pattern.Replace('\\', '/');
+
+			// We don't want a slash at the start, but if there was not one specified, it's not anchored to the root of the tree.
+			if (NormalizedPattern.StartsWith("/"))
 			{
-				NormalizedPattern = NormalizedPattern.Insert(Idx, "*/");
-				Idx++;
+				NormalizedPattern = NormalizedPattern.Substring(1);
 			}
-
-			Idx += 3;
-
-			if (Idx < NormalizedPattern.Length && NormalizedPattern[Idx] != '/')
+			else if(!NormalizedPattern.StartsWith("..."))
 			{
-				NormalizedPattern = NormalizedPattern.Insert(Idx, "/*");
-				Idx += 2;
+				NormalizedPattern = ".../" + NormalizedPattern;
 			}
-		}
 
-		// Split the pattern into fragments
-		string[] BranchPatterns = NormalizedPattern.Split('/');
+			// All directories indicate a wildcard match
+			if (NormalizedPattern.EndsWith("/"))
+			{
+				NormalizedPattern += "...";
+			}
 
-		// Add it into the tree
-		FileFilterNode LastNode = RootNode;
-		foreach (string BranchPattern in BranchPatterns)
-		{
-			FileFilterNode NextNode = LastNode.Branches.FirstOrDefault(x => x.Pattern == BranchPattern);
-			if (NextNode == null)
+			// Replace any directory wildcards mid-string
+			for (int Idx = NormalizedPattern.IndexOf("..."); Idx != -1; Idx = NormalizedPattern.IndexOf("...", Idx))
 			{
-				NextNode = new FileFilterNode(LastNode, BranchPattern);
-				LastNode.Branches.Add(NextNode);
+				if (Idx > 0 && NormalizedPattern[Idx - 1] != '/')
+				{
+					NormalizedPattern = NormalizedPattern.Insert(Idx, "*/");
+					Idx++;
+				}
+
+				Idx += 3;
+
+				if (Idx < NormalizedPattern.Length && NormalizedPattern[Idx] != '/')
+				{
+					NormalizedPattern = NormalizedPattern.Insert(Idx, "/*");
+					Idx += 2;
+				}
 			}
-			LastNode = NextNode;
-		}
 
-		// We've reached the end of the pattern, so mark it as a leaf node
-		Rules.Add(LastNode);
-		LastNode.RuleNumber = Rules.Count - 1;
-		LastNode.Type = Type;
+			// Split the pattern into fragments
+			string[] BranchPatterns = NormalizedPattern.Split('/');
 
-		// Update the maximums along that path
-		for (FileFilterNode UpdateNode = LastNode; UpdateNode != null; UpdateNode = UpdateNode.Parent)
-		{
-			if (Type == FileFilterType.Include)
+			// Add it into the tree
+			FileFilterNode LastNode = RootNode;
+			foreach (string BranchPattern in BranchPatterns)
 			{
-				UpdateNode.MaxIncludeRuleNumber = LastNode.RuleNumber;
+				FileFilterNode NextNode = LastNode.Branches.FirstOrDefault(x => x.Pattern == BranchPattern);
+				if (NextNode == null)
+				{
+					NextNode = new FileFilterNode(LastNode, BranchPattern);
+					LastNode.Branches.Add(NextNode);
+				}
+				LastNode = NextNode;
 			}
-			else
+
+			// We've reached the end of the pattern, so mark it as a leaf node
+			Rules.Add(LastNode);
+			LastNode.RuleNumber = Rules.Count - 1;
+			LastNode.Type = Type;
+
+			// Update the maximums along that path
+			for (FileFilterNode UpdateNode = LastNode; UpdateNode != null; UpdateNode = UpdateNode.Parent)
 			{
-				UpdateNode.MaxExcludeRuleNumber = LastNode.RuleNumber;
+				if (Type == FileFilterType.Include)
+				{
+					UpdateNode.MaxIncludeRuleNumber = LastNode.RuleNumber;
+				}
+				else
+				{
+					UpdateNode.MaxExcludeRuleNumber = LastNode.RuleNumber;
+				}
 			}
 		}
-	}
 
-	/// <summary>
-	/// Excludes all restricted platform folders from the filter
-	/// </summary>
-	public void ExcludeConfidentialPlatforms()
-	{
-		AddRule(".../PS4/...", FileFilterType.Exclude);
-		AddRule(".../XboxOne/...", FileFilterType.Exclude);
-	}
-
-	/// <summary>
-	/// Excludes all confidential folders from the filter
-	/// </summary>
-	public void ExcludeConfidentialFolders()
-	{
-		AddRule(".../EpicInternal/...", FileFilterType.Exclude);
-		AddRule(".../CarefullyRedist/...", FileFilterType.Exclude);
-		AddRule(".../NotForLicensees/...", FileFilterType.Exclude);
-		AddRule(".../NoRedist/...", FileFilterType.Exclude);
-	}
-
-	/// <summary>
-	/// Determines whether the given file matches the filter
-	/// </summary>
-	/// <param name="FileName">File to match</param>
-	/// <returns>True if the file passes the filter</returns>
-	public bool Matches(string FileName)
-	{
-		string[] Tokens = FileName.TrimStart('/', '\\').Split('/', '\\');
-
-		FileFilterNode MatchingNode = FindMatchingNode(RootNode, Tokens, 0, DefaultNode);
-
-		return MatchingNode.Type == FileFilterType.Include;
-	}
+		/// <summary>
+		/// Excludes all restricted platform folders from the filter
+		/// </summary>
+		public void ExcludeConfidentialPlatforms()
+		{
+			AddRule(".../PS4/...", FileFilterType.Exclude);
+			AddRule(".../XboxOne/...", FileFilterType.Exclude);
+		}
 
-	/// <summary>
-	/// Determines whether it's possible for anything within the given folder name to match the filter. Useful to early out of recursive file searches.
-	/// </summary>
-	/// <param name="FolderName">File to match</param>
-	/// <returns>True if the file passes the filter</returns>
-	public bool PossiblyMatches(string FolderName)
-	{
-		string[] Tokens = FolderName.Trim('/', '\\').Split('/', '\\');
+		/// <summary>
+		/// Excludes all confidential folders from the filter
+		/// </summary>
+		public void ExcludeConfidentialFolders()
+		{
+			AddRule(".../EpicInternal/...", FileFilterType.Exclude);
+			AddRule(".../CarefullyRedist/...", FileFilterType.Exclude);
+			AddRule(".../NotForLicensees/...", FileFilterType.Exclude);
+			AddRule(".../NoRedist/...", FileFilterType.Exclude);
+		}
 
-		FileFilterNode MatchingNode = FindMatchingNode(RootNode, Tokens.Union(new string[]{ "" }).ToArray(), 0, DefaultNode);
+		/// <summary>
+		/// Determines whether the given file matches the filter
+		/// </summary>
+		/// <param name="FileName">File to match</param>
+		/// <returns>True if the file passes the filter</returns>
+		public bool Matches(string FileName)
+		{
+			string[] Tokens = FileName.TrimStart('/', '\\').Split('/', '\\');
 
-		return MatchingNode.Type == FileFilterType.Include || HighestPossibleIncludeMatch(RootNode, Tokens, 0, MatchingNode.RuleNumber) > MatchingNode.RuleNumber;
-	}
+			FileFilterNode MatchingNode = FindMatchingNode(RootNode, Tokens, 0, DefaultNode);
 
-	/// <summary>
-	/// Applies the filter to each element in a sequence, and returns the list of files that match
-	/// </summary>
-	/// <param name="FileNames">List of filenames</param>
-	/// <returns>List of filenames which match the filter</returns>
-	public IEnumerable<string> ApplyTo(IEnumerable<string> FileNames)
-	{
-		return FileNames.Where(x => Matches(x));
-	}
+			return MatchingNode.Type == FileFilterType.Include;
+		}
 
-	/// <summary>
-	/// Applies the filter to each element in a sequence, and returns the list of files that match
-	/// </summary>
-	/// <param name="FileNames">List of filenames</param>
-	/// <returns>List of filenames which match the filter</returns>
-	public IEnumerable<string> ApplyTo(string BaseDirectory, IEnumerable<string> FileNames)
-	{
-		string FullBaseDirectory = Path.GetFullPath(BaseDirectory);
-		return FileNames.Where(x => Matches(AutomationTool.CommandUtils.StripBaseDirectory(Path.GetFullPath(x), FullBaseDirectory)));
-	}
+		/// <summary>
+		/// Determines whether it's possible for anything within the given folder name to match the filter. Useful to early out of recursive file searches.
+		/// </summary>
+		/// <param name="FolderName">File to match</param>
+		/// <returns>True if the file passes the filter</returns>
+		public bool PossiblyMatches(string FolderName)
+		{
+			string[] Tokens = FolderName.Trim('/', '\\').Split('/', '\\');
 
-	/// <summary>
-	/// Finds a list of files within a given directory which match the filter.
-	/// </summary>
-    /// <param name="DirectoryName">File to match</param>
-	/// <returns>List of files that pass the filter</returns>
-	public List<string> ApplyToDirectory(string DirectoryName, bool bIgnoreSymlinks)
-	{
-		List<string> MatchingFileNames = new List<string>();
-		FindMatchesFromDirectory(new DirectoryInfo(DirectoryName), "", bIgnoreSymlinks, MatchingFileNames);
-		return MatchingFileNames;
-	}
+			FileFilterNode MatchingNode = FindMatchingNode(RootNode, Tokens.Union(new string[]{ "" }).ToArray(), 0, DefaultNode);
 
-	/// <summary>
-	/// Finds a list of files within a given directory which match the filter.
-	/// </summary>
-	/// <param name="FolderName">File to match</param>
-	/// <returns>True if the file passes the filter</returns>
-	void FindMatchesFromDirectory(DirectoryInfo CurrentDirectory, string NamePrefix, bool bIgnoreSymlinks, List<string> MatchingFileNames)
-	{
-		foreach (FileInfo NextFile in CurrentDirectory.EnumerateFiles())
-		{
-			string FileName = NamePrefix + NextFile.Name;
-			if (Matches(FileName) && (!bIgnoreSymlinks || !NextFile.Attributes.HasFlag(FileAttributes.ReparsePoint)))
-			{
-				MatchingFileNames.Add(FileName);
-			}
+			return MatchingNode.Type == FileFilterType.Include || HighestPossibleIncludeMatch(RootNode, Tokens, 0, MatchingNode.RuleNumber) > MatchingNode.RuleNumber;
 		}
-		foreach (DirectoryInfo NextDirectory in CurrentDirectory.EnumerateDirectories())
+
+		/// <summary>
+		/// Applies the filter to each element in a sequence, and returns the list of files that match
+		/// </summary>
+		/// <param name="FileNames">List of filenames</param>
+		/// <returns>List of filenames which match the filter</returns>
+		public IEnumerable<string> ApplyTo(IEnumerable<string> FileNames)
 		{
-			string NextNamePrefix = NamePrefix + NextDirectory.Name;
-			if (PossiblyMatches(NextNamePrefix))
-			{
-				FindMatchesFromDirectory(NextDirectory, NextNamePrefix + "/", bIgnoreSymlinks, MatchingFileNames);
-			}
+			return FileNames.Where(x => Matches(x));
 		}
-	}
 
-	/// <summary>
-	/// Finds the node which matches a given list of tokens.
-	/// </summary>
-	/// <param name="CurrentNode"></param>
-	/// <param name="Tokens"></param>
-	/// <param name="TokenIdx"></param>
-	/// <param name="CurrentBestNode"></param>
-	/// <returns></returns>
-	FileFilterNode FindMatchingNode(FileFilterNode CurrentNode, string[] Tokens, int TokenIdx, FileFilterNode CurrentBestNode)
-	{
-		// If we've matched all the input tokens, check if this rule is better than any other we've seen
-		if (TokenIdx == Tokens.Length)
+		/// <summary>
+		/// Applies the filter to each element in a sequence, and returns the list of files that match
+		/// </summary>
+		/// <param name="FileNames">List of filenames</param>
+		/// <returns>List of filenames which match the filter</returns>
+		public IEnumerable<string> ApplyTo(string BaseDirectory, IEnumerable<string> FileNames)
 		{
-			return (CurrentNode.RuleNumber > CurrentBestNode.RuleNumber) ? CurrentNode : CurrentBestNode;
+			string FullBaseDirectory = Path.GetFullPath(BaseDirectory);
+			return FileNames.Where(x => Matches(AutomationTool.CommandUtils.StripBaseDirectory(Path.GetFullPath(x), FullBaseDirectory)));
 		}
 
-		// If there is no rule under the current node which is better than the current best node, early out
-		if (CurrentNode.MaxIncludeRuleNumber <= CurrentBestNode.RuleNumber && CurrentNode.MaxExcludeRuleNumber <= CurrentBestNode.RuleNumber)
+		/// <summary>
+		/// Finds a list of files within a given directory which match the filter.
+		/// </summary>
+		/// <param name="DirectoryName">File to match</param>
+		/// <returns>List of files that pass the filter</returns>
+		public List<string> ApplyToDirectory(string DirectoryName, bool bIgnoreSymlinks)
 		{
-			return CurrentBestNode;
+			List<string> MatchingFileNames = new List<string>();
+			FindMatchesFromDirectory(new DirectoryInfo(DirectoryName), "", bIgnoreSymlinks, MatchingFileNames);
+			return MatchingFileNames;
 		}
 
-		// Test all the branches for one that matches
-		FileFilterNode BestNode = CurrentBestNode;
-		foreach (FileFilterNode Branch in CurrentNode.Branches)
+		/// <summary>
+		/// Finds a list of files within a given directory which match the filter.
+		/// </summary>
+		/// <param name="FolderName">File to match</param>
+		/// <returns>True if the file passes the filter</returns>
+		void FindMatchesFromDirectory(DirectoryInfo CurrentDirectory, string NamePrefix, bool bIgnoreSymlinks, List<string> MatchingFileNames)
 		{
-			if (Branch.Pattern == "...")
+			foreach (FileInfo NextFile in CurrentDirectory.EnumerateFiles())
 			{
-				for (int NextTokenIdx = Tokens.Length; NextTokenIdx >= TokenIdx; NextTokenIdx--)
+				string FileName = NamePrefix + NextFile.Name;
+				if (Matches(FileName) && (!bIgnoreSymlinks || !NextFile.Attributes.HasFlag(FileAttributes.ReparsePoint)))
 				{
-					BestNode = FindMatchingNode(Branch, Tokens, NextTokenIdx, BestNode);
+					MatchingFileNames.Add(FileName);
 				}
 			}
-			else
+			foreach (DirectoryInfo NextDirectory in CurrentDirectory.EnumerateDirectories())
 			{
-				if (Branch.IsMatch(Tokens[TokenIdx]))
+				string NextNamePrefix = NamePrefix + NextDirectory.Name;
+				if (PossiblyMatches(NextNamePrefix))
 				{
-					BestNode = FindMatchingNode(Branch, Tokens, TokenIdx + 1, BestNode);
+					FindMatchesFromDirectory(NextDirectory, NextNamePrefix + "/", bIgnoreSymlinks, MatchingFileNames);
 				}
 			}
 		}
-		return BestNode;
-	}
 
-	/// <summary>
-	/// Returns the highest possible rule number which can match the given list of input tokens, assuming that the list of input tokens is incomplete.
-	/// </summary>
-	/// <param name="CurrentNode">The current node being checked</param>
-	/// <param name="Tokens">The tokens to match</param>
-	/// <param name="TokenIdx">Current token index</param>
-	/// <param name="CurrentBestRuleNumber">The highest rule number seen so far. Used to optimize tree traversals.</param>
-	/// <returns>New highest rule number</returns>
-	int HighestPossibleIncludeMatch(FileFilterNode CurrentNode, string[] Tokens, int TokenIdx, int CurrentBestRuleNumber)
-	{
-		// If we've matched all the input tokens, check if this rule is better than any other we've seen
-		if (TokenIdx == Tokens.Length)
+		/// <summary>
+		/// Finds the node which matches a given list of tokens.
+		/// </summary>
+		/// <param name="CurrentNode"></param>
+		/// <param name="Tokens"></param>
+		/// <param name="TokenIdx"></param>
+		/// <param name="CurrentBestNode"></param>
+		/// <returns></returns>
+		FileFilterNode FindMatchingNode(FileFilterNode CurrentNode, string[] Tokens, int TokenIdx, FileFilterNode CurrentBestNode)
 		{
-			return Math.Max(CurrentBestRuleNumber, CurrentNode.MaxIncludeRuleNumber);
-		}
+			// If we've matched all the input tokens, check if this rule is better than any other we've seen
+			if (TokenIdx == Tokens.Length)
+			{
+				return (CurrentNode.RuleNumber > CurrentBestNode.RuleNumber) ? CurrentNode : CurrentBestNode;
+			}
 
-		// Test all the branches for one that matches
-		int BestRuleNumber = CurrentBestRuleNumber;
-		if (CurrentNode.MaxIncludeRuleNumber > BestRuleNumber)
-		{
+			// If there is no rule under the current node which is better than the current best node, early out
+			if (CurrentNode.MaxIncludeRuleNumber <= CurrentBestNode.RuleNumber && CurrentNode.MaxExcludeRuleNumber <= CurrentBestNode.RuleNumber)
+			{
+				return CurrentBestNode;
+			}
+
+			// Test all the branches for one that matches
+			FileFilterNode BestNode = CurrentBestNode;
 			foreach (FileFilterNode Branch in CurrentNode.Branches)
 			{
 				if (Branch.Pattern == "...")
 				{
-					if (Branch.MaxIncludeRuleNumber > BestRuleNumber)
+					for (int NextTokenIdx = Tokens.Length; NextTokenIdx >= TokenIdx; NextTokenIdx--)
 					{
-						BestRuleNumber = Branch.MaxIncludeRuleNumber;
+						BestNode = FindMatchingNode(Branch, Tokens, NextTokenIdx, BestNode);
 					}
 				}
 				else
 				{
 					if (Branch.IsMatch(Tokens[TokenIdx]))
 					{
-						BestRuleNumber = HighestPossibleIncludeMatch(Branch, Tokens, TokenIdx + 1, BestRuleNumber);
+						BestNode = FindMatchingNode(Branch, Tokens, TokenIdx + 1, BestNode);
 					}
 				}
 			}
+			return BestNode;
 		}
-		return BestRuleNumber;
-	}
-}
 
-/// <summary>
-/// Node within a filter tree. Each node matches a single path fragment - a folder or file name, with an include or exclude rule consisting of a sequence of nodes. 
-/// </summary>
-class FileFilterNode
-{
-	/// <summary>
-	/// Node which this is parented to. Null for the root node.
-	/// </summary>
-	public readonly FileFilterNode Parent;
-
-	/// <summary>
-	/// Pattern to match for this node. May contain * or ? wildcards.
-	/// </summary>
-	public readonly string Pattern;
-
-	/// <summary>
-	/// Highest include rule number matched by this node or any child nodes.
-	/// </summary>
-	public int MaxIncludeRuleNumber = -1;
-
-	/// <summary>
-	/// Highest exclude rule number matched by this node or any child nodes.
-	/// </summary>
-	public int MaxExcludeRuleNumber = -1;
-
-	/// <summary>
-	/// Child branches of this node, distinct by the the pattern for each.
-	/// </summary>
-	public List<FileFilterNode> Branches = new List<FileFilterNode>();
-
-	/// <summary>
-	/// If a path matches the sequence of nodes ending in this node, the number of the rule which added it. -1 if this was not the last node in a rule.
-	/// </summary>
-	public int RuleNumber;
-
-	/// <summary>
-	/// Whether a rule terminating in this node should include files or exclude files.
-	/// </summary>
-	public FileFilterType Type;
+		/// <summary>
+		/// Returns the highest possible rule number which can match the given list of input tokens, assuming that the list of input tokens is incomplete.
+		/// </summary>
+		/// <param name="CurrentNode">The current node being checked</param>
+		/// <param name="Tokens">The tokens to match</param>
+		/// <param name="TokenIdx">Current token index</param>
+		/// <param name="CurrentBestRuleNumber">The highest rule number seen so far. Used to optimize tree traversals.</param>
+		/// <returns>New highest rule number</returns>
+		int HighestPossibleIncludeMatch(FileFilterNode CurrentNode, string[] Tokens, int TokenIdx, int CurrentBestRuleNumber)
+		{
+			// If we've matched all the input tokens, check if this rule is better than any other we've seen
+			if (TokenIdx == Tokens.Length)
+			{
+				return Math.Max(CurrentBestRuleNumber, CurrentNode.MaxIncludeRuleNumber);
+			}
 
-	/// <summary>
-	/// Default constructor.
-	/// </summary>
-	public FileFilterNode(FileFilterNode InParent, string InPattern)
-	{
-		Parent = InParent;
-		Pattern = InPattern;
-		RuleNumber = -1;
-		Type = FileFilterType.Exclude;
+			// Test all the branches for one that matches
+			int BestRuleNumber = CurrentBestRuleNumber;
+			if (CurrentNode.MaxIncludeRuleNumber > BestRuleNumber)
+			{
+				foreach (FileFilterNode Branch in CurrentNode.Branches)
+				{
+					if (Branch.Pattern == "...")
+					{
+						if (Branch.MaxIncludeRuleNumber > BestRuleNumber)
+						{
+							BestRuleNumber = Branch.MaxIncludeRuleNumber;
+						}
+					}
+					else
+					{
+						if (Branch.IsMatch(Tokens[TokenIdx]))
+						{
+							BestRuleNumber = HighestPossibleIncludeMatch(Branch, Tokens, TokenIdx + 1, BestRuleNumber);
+						}
+					}
+				}
+			}
+			return BestRuleNumber;
+		}
 	}
 
 	/// <summary>
-	/// Determine if the given token matches the pattern in this node
+	/// Node within a filter tree. Each node matches a single path fragment - a folder or file name, with an include or exclude rule consisting of a sequence of nodes. 
 	/// </summary>
-	public bool IsMatch(string Token)
+	class FileFilterNode
 	{
-		if(Pattern.EndsWith("."))
-		{
-			return !Token.Contains('.') && IsMatch(Token, 0, Pattern.Substring(0, Pattern.Length - 1), 0);
-		}
-		else
+		/// <summary>
+		/// Node which this is parented to. Null for the root node.
+		/// </summary>
+		public readonly FileFilterNode Parent;
+
+		/// <summary>
+		/// Pattern to match for this node. May contain * or ? wildcards.
+		/// </summary>
+		public readonly string Pattern;
+
+		/// <summary>
+		/// Highest include rule number matched by this node or any child nodes.
+		/// </summary>
+		public int MaxIncludeRuleNumber = -1;
+
+		/// <summary>
+		/// Highest exclude rule number matched by this node or any child nodes.
+		/// </summary>
+		public int MaxExcludeRuleNumber = -1;
+
+		/// <summary>
+		/// Child branches of this node, distinct by the the pattern for each.
+		/// </summary>
+		public List<FileFilterNode> Branches = new List<FileFilterNode>();
+
+		/// <summary>
+		/// If a path matches the sequence of nodes ending in this node, the number of the rule which added it. -1 if this was not the last node in a rule.
+		/// </summary>
+		public int RuleNumber;
+
+		/// <summary>
+		/// Whether a rule terminating in this node should include files or exclude files.
+		/// </summary>
+		public FileFilterType Type;
+
+		/// <summary>
+		/// Default constructor.
+		/// </summary>
+		public FileFilterNode(FileFilterNode InParent, string InPattern)
 		{
-			return IsMatch(Token, 0, Pattern, 0);
+			Parent = InParent;
+			Pattern = InPattern;
+			RuleNumber = -1;
+			Type = FileFilterType.Exclude;
 		}
-	}
 
-	/// <summary>
-	/// Determine if a given token matches a pattern, starting from the given positions within each string.
-	/// </summary>
-	static bool IsMatch(string Token, int TokenIdx, string Pattern, int PatternIdx)
-	{
-		for (; ; )
+		/// <summary>
+		/// Determine if the given token matches the pattern in this node
+		/// </summary>
+		public bool IsMatch(string Token)
 		{
-			if (PatternIdx == Pattern.Length)
+			if(Pattern.EndsWith("."))
 			{
-				return (TokenIdx == Token.Length);
+				return !Token.Contains('.') && IsMatch(Token, 0, Pattern.Substring(0, Pattern.Length - 1), 0);
 			}
-			else if (Pattern[PatternIdx] == '*')
+			else
+			{
+				return IsMatch(Token, 0, Pattern, 0);
+			}
+		}
+
+		/// <summary>
+		/// Determine if a given token matches a pattern, starting from the given positions within each string.
+		/// </summary>
+		static bool IsMatch(string Token, int TokenIdx, string Pattern, int PatternIdx)
+		{
+			for (; ; )
 			{
-				for (int NextTokenIdx = Token.Length; NextTokenIdx >= TokenIdx; NextTokenIdx--)
+				if (PatternIdx == Pattern.Length)
+				{
+					return (TokenIdx == Token.Length);
+				}
+				else if (Pattern[PatternIdx] == '*')
 				{
-					if (IsMatch(Token, NextTokenIdx, Pattern, PatternIdx + 1))
+					for (int NextTokenIdx = Token.Length; NextTokenIdx >= TokenIdx; NextTokenIdx--)
 					{
-						return true;
+						if (IsMatch(Token, NextTokenIdx, Pattern, PatternIdx + 1))
+						{
+							return true;
+						}
 					}
+					return false;
+				}
+				else if (TokenIdx < Token.Length && (Char.ToLower(Token[TokenIdx]) == Char.ToLower(Pattern[PatternIdx]) || Pattern[PatternIdx] == '?'))
+				{
+					TokenIdx++;
+					PatternIdx++;
+				}
+				else
+				{
+					return false;
 				}
-				return false;
-			}
-			else if (TokenIdx < Token.Length && (Char.ToLower(Token[TokenIdx]) == Char.ToLower(Pattern[PatternIdx]) || Pattern[PatternIdx] == '?'))
-			{
-				TokenIdx++;
-				PatternIdx++;
-			}
-			else
-			{
-				return false;
 			}
 		}
-	}
 
-	/// <summary>
-	/// Debugger visualization
-	/// </summary>
-	/// <returns>Path to this node</returns>
-	public override string ToString()
-	{
-		return (Parent == null) ? "/" : Parent.ToString().TrimEnd('/') + "/" + Pattern;
+		/// <summary>
+		/// Debugger visualization
+		/// </summary>
+		/// <returns>Path to this node</returns>
+		public override string ToString()
+		{
+			return (Parent == null) ? "/" : Parent.ToString().TrimEnd('/') + "/" + Pattern;
+		}
 	}
 }
