refactored Log to LogLog and LogConsole to Log
#uat

diff --git a/Engine/Source/Programs/AutomationTool/Android/AndroidPlatform.Automation.cs b/Engine/Source/Programs/AutomationTool/Android/AndroidPlatform.Automation.cs
--- a/Engine/Source/Programs/AutomationTool/Android/AndroidPlatform.Automation.cs
+++ b/Engine/Source/Programs/AutomationTool/Android/AndroidPlatform.Automation.cs
@@ -120,7 +120,7 @@ public class AndroidPlatform : Platform
 		var Deploy = UEBuildDeploy.GetBuildDeploy(UnrealTargetPlatform.Android);
 
 		string BaseApkName = GetFinalApkName(Params, SC.StageExecutables[0], true, "", "");
-		LogConsole("BaseApkName = {0}", BaseApkName);
+		Log("BaseApkName = {0}", BaseApkName);
 
 		// Create main OBB with entire contents of staging dir. This
 		// includes any PAK files, movie files, etc.
@@ -134,7 +134,7 @@ public class AndroidPlatform : Platform
 		}
 
 		// Now create the OBB as a ZIP archive.
-		LogConsole("Creating {0} from {1}", LocalObbName, SC.StageDirectory);
+		Log("Creating {0} from {1}", LocalObbName, SC.StageDirectory);
 		using (ZipFile ObbFile = new ZipFile(LocalObbName))
 		{
 			ObbFile.CompressionMethod = CompressionMethod.None;
@@ -146,7 +146,7 @@ public class AndroidPlatform : Platform
 					if (e.EventType == ZipProgressEventType.Adding_AfterAddEntry)
 					{
 						ObbFileCount += 1;
-						LogConsole("[{0}/{1}] Adding {2} to OBB",
+						Log("[{0}/{1}] Adding {2} to OBB",
 							ObbFileCount, e.EntriesTotal,
 							e.CurrentEntry.FileName);
 					}
@@ -166,7 +166,7 @@ public class AndroidPlatform : Platform
                     UE4SOName = UE4SOName.Replace(".apk", ".so");
                     if (FileExists_NoExceptions(UE4SOName) == false)
 					{
-						LogConsole("Failed to find game .so " + UE4SOName);
+						Log("Failed to find game .so " + UE4SOName);
                         throw new AutomationException(ExitCode.Error_MissingExecutable, "Stage Failed. Could not find .so {0}. You may need to build the UE4 project with your target configuration and platform.", UE4SOName);
 					}
 				}
@@ -227,7 +227,7 @@ public class AndroidPlatform : Platform
         {
             string OBBInstallCommand = bNoObbInstall ? "shell 'rm -r $EXTERNAL_STORAGE/" + DeviceObbName + "'" : "push " + Path.GetFileName(ObbName) + " $STORAGE/" + DeviceObbName;
 
-			LogConsole("Writing shell script for install with {0}", bPackageDataInsideApk ? "data in APK" : "separate obb");
+			Log("Writing shell script for install with {0}", bPackageDataInsideApk ? "data in APK" : "separate obb");
             BatchLines = new string[] {
 						"#!/bin/sh",
 						"cd \"`dirname \"$0\"`\"",
@@ -271,7 +271,7 @@ public class AndroidPlatform : Platform
         {
             string OBBInstallCommand = bNoObbInstall ? "shell rm -r %STORAGE%/" + DeviceObbName : "push " + Path.GetFileName(ObbName) + " %STORAGE%/" + DeviceObbName;
 
-			LogConsole("Writing bat for install with {0}", bPackageDataInsideApk ? "data in APK" : "separate OBB");
+			Log("Writing bat for install with {0}", bPackageDataInsideApk ? "data in APK" : "separate OBB");
             BatchLines = new string[] {
 						"setlocal",
                         "set ANDROIDHOME=%ANDROID_HOME%",
@@ -1018,7 +1018,7 @@ public class AndroidPlatform : Platform
 				TimeSpan DeltaRunTime = DateTime.Now - StartTime;
 				if ((DeltaRunTime.TotalSeconds > TimeOutSeconds) && (TimeOutSeconds != 0))
 				{
-					LogConsole("Device: " + Params.Device + " timed out while waiting for run to finish");
+					Log("Device: " + Params.Device + " timed out while waiting for run to finish");
 					break;
 				}
 			}
diff --git a/Engine/Source/Programs/AutomationTool/AutomationUtils/Automation.cs b/Engine/Source/Programs/AutomationTool/AutomationUtils/Automation.cs
--- a/Engine/Source/Programs/AutomationTool/AutomationUtils/Automation.cs
+++ b/Engine/Source/Programs/AutomationTool/AutomationUtils/Automation.cs
@@ -471,7 +471,7 @@ AutomationTool.exe [-verbose] [-compileonly] [-p4] Command0 [-Arg0 -Arg1 -Arg2 
 			{
 				Message += String.Format("  {0}{1}", AvailableCommand.Key, Environment.NewLine);
 			}
-			CommandUtils.LogConsole(Message);
+			CommandUtils.Log(Message);
 		}
 
 		#endregion
diff --git a/Engine/Source/Programs/AutomationTool/AutomationUtils/BuildCommand.cs b/Engine/Source/Programs/AutomationTool/AutomationUtils/BuildCommand.cs
--- a/Engine/Source/Programs/AutomationTool/AutomationUtils/BuildCommand.cs
+++ b/Engine/Source/Programs/AutomationTool/AutomationUtils/BuildCommand.cs
@@ -35,7 +35,7 @@ namespace AutomationTool
 				LogError("BUILD FAILED");
 				throw;
 			}
-			LogConsole("BUILD SUCCESSFUL");
+			Log("BUILD SUCCESSFUL");
 		}
 
 		/// <summary>
diff --git a/Engine/Source/Programs/AutomationTool/AutomationUtils/CommandUtils.cs b/Engine/Source/Programs/AutomationTool/AutomationUtils/CommandUtils.cs
--- a/Engine/Source/Programs/AutomationTool/AutomationUtils/CommandUtils.cs
+++ b/Engine/Source/Programs/AutomationTool/AutomationUtils/CommandUtils.cs
@@ -115,7 +115,7 @@ namespace AutomationTool
 		/// <param name="Format">Format string</param>
 		/// <param name="Args">Parameters</param>
 		[MethodImplAttribute(MethodImplOptions.NoInlining)]
-		public static void LogConsole(string Format, params object[] Args)
+		public static void Log(string Format, params object[] Args)
 		{
 			UnrealBuildTool.Log.WriteLine(1, null, UnrealBuildTool.LogEventType.Console, Format, Args);
 		}
@@ -125,7 +125,7 @@ namespace AutomationTool
 		/// </summary>
 		/// <param name="Message">Text</param>
 		[MethodImplAttribute(MethodImplOptions.NoInlining)]
-		public static void LogConsole(string Message)
+		public static void Log(string Message)
 		{
 			UnrealBuildTool.Log.WriteLine(1, null, UnrealBuildTool.LogEventType.Console, Message);
 		}
@@ -220,7 +220,7 @@ namespace AutomationTool
 		/// <param name="Foramt">Format string</param>
 		/// <param name="Args">Arguments</param>
 		[MethodImplAttribute(MethodImplOptions.NoInlining)]
-		public static void Log(string Format, params object[] Args)
+		public static void LogLog(string Format, params object[] Args)
 		{
 			UnrealBuildTool.Log.WriteLine(1, null, UnrealBuildTool.LogEventType.Log, Format, Args);
 		}
@@ -230,7 +230,7 @@ namespace AutomationTool
 		/// </summary>
 		/// <param name="Message">Text</param>
 		[MethodImplAttribute(MethodImplOptions.NoInlining)]
-		public static void Log(string Message)
+		public static void LogLog(string Message)
 		{
 			UnrealBuildTool.Log.WriteLine(1, null, UnrealBuildTool.LogEventType.Log, Message);
 		}
@@ -277,7 +277,7 @@ namespace AutomationTool
 		{
 			if(bShowProgress)
 			{
-				LogConsole("[@progress push {0}/{1} skipline]", Numerator, Denominator);
+				Log("[@progress push {0}/{1} skipline]", Numerator, Denominator);
 			}
 		}
 
@@ -285,7 +285,7 @@ namespace AutomationTool
 		{
 			if(bShowProgress)
 			{
-				LogConsole("[@progress pop skipline]");
+				Log("[@progress pop skipline]");
 			}
 		}
 
@@ -293,7 +293,7 @@ namespace AutomationTool
 		{
 			if(bShowProgress)
 			{
-				LogConsole("[@progress increment {0}/{1} skipline]", Numerator, Denominator);
+				Log("[@progress increment {0}/{1} skipline]", Numerator, Denominator);
 			}
 		}
 
@@ -301,7 +301,7 @@ namespace AutomationTool
 		{
 			if(bShowProgress)
 			{
-				LogConsole("[@progress '{0}' skipline]", String.Format(Format, Args));
+				Log("[@progress '{0}' skipline]", String.Format(Format, Args));
 			}
 		}
 
@@ -309,7 +309,7 @@ namespace AutomationTool
 		{
 			if(bShowProgress)
 			{
-				LogConsole("[@progress {0}/{1} '{2}' skipline]", Numerator, Denominator, String.Format(Format, Args));
+				Log("[@progress {0}/{1} '{2}' skipline]", Numerator, Denominator, String.Format(Format, Args));
 			}
 		}
 
@@ -1535,7 +1535,7 @@ namespace AutomationTool
 		/// <param name="MaxThreads"></param>
 		public static void ThreadedCopyFiles(List<string> Source, List<string> Dest, int MaxThreads = 64)
 		{
-			LogConsole("Copying {0} file(s) using max {1} thread(s)", Source.Count, MaxThreads);
+			Log("Copying {0} file(s) using max {1} thread(s)", Source.Count, MaxThreads);
 
             if (Source.Count != Dest.Count)
 			{
@@ -1582,7 +1582,7 @@ namespace AutomationTool
 		public static List<string> ThreadedCopyFiles(string SourceDir, string TargetDir, FileFilter Filter, bool bIgnoreSymlinks, int MaxThreads = 64)
 		{
 			// Filter all the relative paths
-			LogConsole("Applying filter to {0}...", SourceDir);
+			Log("Applying filter to {0}...", SourceDir);
 			var RelativePaths = Filter.ApplyToDirectory(SourceDir, bIgnoreSymlinks);
 			return ThreadedCopyFiles(SourceDir, TargetDir, RelativePaths);
 		}
@@ -1622,7 +1622,7 @@ namespace AutomationTool
 		{
 			try
 			{
-				Log("SetEnvVar {0}={1}", Name, Value);
+				LogLog("SetEnvVar {0}={1}", Name, Value);
 				Environment.SetEnvironmentVariable(Name, Value.ToString());
 			}
 			catch (Exception Ex)
@@ -1885,7 +1885,7 @@ namespace AutomationTool
         {
             if (!DirectoryExists_NoExceptions(Dir))
             {
-				Log("Directory {0} does not exist", Dir);
+				LogLog("Directory {0} does not exist", Dir);
 				return false;
 			}
 
@@ -1900,7 +1900,7 @@ namespace AutomationTool
 				if(File.Exists(NativeFilename))
 				{
 		            DeleteFile_NoExceptions(Filename, true);
-		            Log("Directory {0} is writable", Dir);
+		            LogLog("Directory {0} is writable", Dir);
 					return true;
 				}
 			}
@@ -1908,7 +1908,7 @@ namespace AutomationTool
 			{
 			}
 
-			Log("Directory {0} is not writable", Dir);
+			LogLog("Directory {0} is not writable", Dir);
 			return false;
 		}
 
@@ -1941,7 +1941,7 @@ namespace AutomationTool
 
                 DirectoryInfo DirInfo = new DirectoryInfo(ParentDir);
                 var TopLevelDirs = DirInfo.GetDirectories();
-				LogConsole("Looking for directories to delete in {0}   {1} dirs", ParentDir, TopLevelDirs.Length);
+				Log("Looking for directories to delete in {0}   {1} dirs", ParentDir, TopLevelDirs.Length);
                 foreach (var TopLevelDir in TopLevelDirs)
                 {
                     if (DirectoryExists_NoExceptions(TopLevelDir.FullName))
@@ -2231,7 +2231,7 @@ namespace AutomationTool
 		{
             if (UnrealBuildTool.Utils.IsRunningOnMono)
             {
-                CommandUtils.Log(String.Format("Can't sign '{0}', we are running under mono.", Filename));
+                CommandUtils.LogLog(String.Format("Can't sign '{0}', we are running under mono.", Filename));
                 return;
             }
             if (!CommandUtils.FileExists(Filename))
@@ -2258,7 +2258,7 @@ namespace AutomationTool
 			}
 			if (!IsExecutable)
 			{
-				CommandUtils.Log(String.Format("Won't sign '{0}', not an executable.", TargetFileInfo.FullName));
+				CommandUtils.LogLog(String.Format("Won't sign '{0}', not an executable.", TargetFileInfo.FullName));
 				return;
 			}
 
@@ -2367,7 +2367,7 @@ namespace AutomationTool
 			}
 			if (!IsExecutable)
 			{
-				CommandUtils.Log(String.Format("Won't sign '{0}', not an executable.", InPath));
+				CommandUtils.LogLog(String.Format("Won't sign '{0}', not an executable.", InPath));
 				return;
 			}
 
@@ -2410,7 +2410,7 @@ namespace AutomationTool
 		{
 			if (!Command.ParseParam("NoSign"))
 			{
-				CommandUtils.LogConsole("Signing up to {0} files...", Files.Count);
+				CommandUtils.Log("Signing up to {0} files...", Files.Count);
 				UnrealBuildTool.UnrealTargetPlatform TargetPlatform = UnrealBuildTool.BuildHostPlatform.Current.Platform;
 				if (TargetPlatform == UnrealBuildTool.UnrealTargetPlatform.Mac)
 				{
@@ -2434,7 +2434,7 @@ namespace AutomationTool
 			}
 			else
 			{
-				CommandUtils.Log("Skipping signing {0} files due to -nosign.", Files.Count);
+				CommandUtils.LogLog("Skipping signing {0} files due to -nosign.", Files.Count);
 			}
 		}
 
@@ -2513,7 +2513,7 @@ namespace AutomationTool
 		{
 			if (UnrealBuildTool.Utils.IsRunningOnMono)
 			{
-				CommandUtils.Log(String.Format("Can't sign we are running under mono."));
+				CommandUtils.LogLog(String.Format("Can't sign we are running under mono."));
 				return;
 			}
 			List<string> FinalFiles = new List<string>();
diff --git a/Engine/Source/Programs/AutomationTool/AutomationUtils/CommandletUtils.cs b/Engine/Source/Programs/AutomationTool/AutomationUtils/CommandletUtils.cs
--- a/Engine/Source/Programs/AutomationTool/AutomationUtils/CommandletUtils.cs
+++ b/Engine/Source/Programs/AutomationTool/AutomationUtils/CommandletUtils.cs
@@ -150,7 +150,7 @@ namespace AutomationTool
 		/// <param name="Parameters">Command line parameters (without -run=)</param>
 		public static void RunCommandlet(string ProjectName, string UE4Exe, string Commandlet, string Parameters = null)
 		{
-			LogConsole("Running UE4Editor {0} for project {1}", Commandlet, ProjectName);
+			Log("Running UE4Editor {0} for project {1}", Commandlet, ProjectName);
 
             var CWD = Path.GetDirectoryName(UE4Exe);
 
@@ -165,7 +165,7 @@ namespace AutomationTool
 			PushDir(CWD);
 
 			string LocalLogFile = LogUtils.GetUniqueLogName(CombinePaths(CmdEnv.EngineSavedFolder, Commandlet));
-			LogConsole("Commandlet log file is {0}", LocalLogFile);
+			Log("Commandlet log file is {0}", LocalLogFile);
 			string Args = String.Format(
 				"{0} -run={1} {2} -abslog={3} -stdout -FORCELOGFLUSH -CrashForUAT -unattended {5}{4}",
 				(ProjectName == null) ? "" : CommandUtils.MakePathSafeToUseWithCommandLine(ProjectName),
diff --git a/Engine/Source/Programs/AutomationTool/AutomationUtils/HelpUtils.cs b/Engine/Source/Programs/AutomationTool/AutomationUtils/HelpUtils.cs
--- a/Engine/Source/Programs/AutomationTool/AutomationUtils/HelpUtils.cs
+++ b/Engine/Source/Programs/AutomationTool/AutomationUtils/HelpUtils.cs
@@ -147,7 +147,7 @@ namespace AutomationTool
 				HelpMessage += FormatParams(Params, 4, 24);
 			}
 
-			LogConsole(HelpMessage);
+			Log(HelpMessage);
 		}
 
 		/// <summary>
diff --git a/Engine/Source/Programs/AutomationTool/AutomationUtils/P4Blame.cs b/Engine/Source/Programs/AutomationTool/AutomationUtils/P4Blame.cs
--- a/Engine/Source/Programs/AutomationTool/AutomationUtils/P4Blame.cs
+++ b/Engine/Source/Programs/AutomationTool/AutomationUtils/P4Blame.cs
@@ -302,7 +302,7 @@ namespace AutomationTool
 		{
 			if (!bQuiet)
 			{
-				CommandUtils.LogConsole("Filelog({0})", DepotFilename);
+				CommandUtils.Log("Filelog({0})", DepotFilename);
 			}
 			var Result = new List<RevisionInfo>();
 			var FilelogResult = P4("filelog " + DepotFilename, AllowSpew: false);
@@ -335,7 +335,7 @@ namespace AutomationTool
 		{
 			if (!bQuiet)
 			{
-				CommandUtils.LogConsole("P4Print({0})", CommandLine);
+				CommandUtils.Log("P4Print({0})", CommandLine);
 			}
 			var PrintResult = P4("print " + CommandLine, AllowSpew: false);
 			var OutputLines = new List<string>(PrintResult.Output.Split(new string[] { Environment.NewLine }, StringSplitOptions.None));
@@ -354,7 +354,7 @@ namespace AutomationTool
 		{
 			if (!bQuiet)
 			{
-				CommandUtils.LogConsole("Annotate({0})", CommandLine);
+				CommandUtils.Log("Annotate({0})", CommandLine);
 			}
 			var PrintResult = P4("annotate " + CommandLine, AllowSpew: false);
 			var OutputLines = new List<string>(PrintResult.Output.Split(new string[] { Environment.NewLine }, StringSplitOptions.None));
@@ -374,7 +374,7 @@ namespace AutomationTool
 		{
 			if (!bQuiet)
 			{
-				CommandUtils.LogConsole("Diff2({0}, {1}, {2})", DepotFilename, RevA, RevB);
+				CommandUtils.Log("Diff2({0}, {1}, {2})", DepotFilename, RevA, RevB);
 			}
 			var Changes = new List<DiffChange>();
 			var Diff2Result = P4(String.Format("diff2 {0}#{1} {0}#{2}", DepotFilename, RevA, RevB), AllowSpew: false);
@@ -400,7 +400,7 @@ namespace AutomationTool
 		/// <returns>An array of lines with blame information.</returns>
 		public BlameLineInfo[] Timelapse(string DepotFilename, int MaxRevision = Int32.MaxValue)
 		{
-			CommandUtils.LogConsole("Timelapse({0}#{1})", DepotFilename, (MaxRevision == Int32.MaxValue) ? "head" : MaxRevision.ToString());
+			CommandUtils.Log("Timelapse({0}#{1})", DepotFilename, (MaxRevision == Int32.MaxValue) ? "head" : MaxRevision.ToString());
 			// Get the file history.
 			var FileHistory = Filelog(DepotFilename, bQuiet: true);
 			if (!CommandUtils.IsNullOrEmpty(FileHistory))
@@ -453,7 +453,7 @@ namespace AutomationTool
 		/// <returns>An array of lines with blame information.</returns>
 		public BlameLineInfo[] Blame(string DepotFilename)
 		{
-			CommandUtils.LogConsole("Blame({0})", DepotFilename);
+			CommandUtils.Log("Blame({0})", DepotFilename);
 			// Get the file history.
 			var FileHistory = Filelog(DepotFilename, bQuiet: true);
 			
diff --git a/Engine/Source/Programs/AutomationTool/AutomationUtils/P4Utils.cs b/Engine/Source/Programs/AutomationTool/AutomationUtils/P4Utils.cs
--- a/Engine/Source/Programs/AutomationTool/AutomationUtils/P4Utils.cs
+++ b/Engine/Source/Programs/AutomationTool/AutomationUtils/P4Utils.cs
@@ -1343,7 +1343,7 @@ namespace AutomationTool
                 bool isClPending = false;
                 if (ChangeFiles(CL, out isClPending, false).Count == 0)
                 {
-					CommandUtils.LogConsole("No edits left to commit after brutal submit resolve. Assuming another build committed same changes already and exiting as success.");
+					CommandUtils.Log("No edits left to commit after brutal submit resolve. Assuming another build committed same changes already and exiting as success.");
                     DeleteChange(CL);
                     // No changes to submit, no need to retry.
                     return;
@@ -1355,7 +1355,7 @@ namespace AutomationTool
 					{
 						throw new P4Exception("Change {0} failed to submit.\n{1}", CL, CmdOutput);
 					}
-					CommandUtils.LogConsole("**** P4 Returned\n{0}\n*******", CmdOutput);
+					CommandUtils.Log("**** P4 Returned\n{0}\n*******", CmdOutput);
 
 					LastCmdOutput = CmdOutput;
 					bool DidSomething = false;
@@ -1418,7 +1418,7 @@ namespace AutomationTool
                                 {
                                     continue;
                                 }
-								CommandUtils.LogConsole("Brutal 'resolve' on {0} to force submit.\n", File);
+								CommandUtils.Log("Brutal 'resolve' on {0} to force submit.\n", File);
 								Revert(CL, "-k " + CommandUtils.MakePathSafeToUseWithCommandLine(File));  // revert the file without overwriting the local one
 								Sync("-f -k " + CommandUtils.MakePathSafeToUseWithCommandLine(File + "#head"), false); // sync the file without overwriting local one
 								ReconcileNoDeletes(CL, CommandUtils.MakePathSafeToUseWithCommandLine(File));  // re-check out, if it changed, or add
@@ -1429,7 +1429,7 @@ namespace AutomationTool
                     }
 					if (!DidSomething)
 					{
-						CommandUtils.LogConsole("Change {0} failed to submit for reasons we do not recognize.\n{1}\nWaiting and retrying.", CL, CmdOutput);
+						CommandUtils.Log("Change {0} failed to submit for reasons we do not recognize.\n{1}\nWaiting and retrying.", CL, CmdOutput);
 					}
 					System.Threading.Thread.Sleep(30000);
 				}
@@ -1461,7 +1461,7 @@ namespace AutomationTool
 							}
 						}
 
-						CommandUtils.LogConsole("Submitted CL {0} which became CL {1}\n", CL, SubmittedCL);
+						CommandUtils.Log("Submitted CL {0} which became CL {1}\n", CL, SubmittedCL);
 					}
 
 					if (SubmittedCL < CL)
@@ -1498,7 +1498,7 @@ namespace AutomationTool
 			int CL = 0;
 			if(AllowSpew)
 			{
-				CommandUtils.LogConsole("Creating Change\n {0}\n", ChangeSpec);
+				CommandUtils.Log("Creating Change\n {0}\n", ChangeSpec);
 			}
 			if (LogP4Output(out CmdOutput, "change -i", Input: ChangeSpec, AllowSpew: AllowSpew))
 			{
@@ -1517,7 +1517,7 @@ namespace AutomationTool
 			}
 			else if(AllowSpew)
 			{
-				CommandUtils.LogConsole("Returned CL {0}\n", CL);
+				CommandUtils.Log("Returned CL {0}\n", CL);
 			}
 			return CL;
 		}
@@ -1648,7 +1648,7 @@ namespace AutomationTool
 				int EndOffset = CmdOutput.LastIndexOf(EndStr);
 				if (Offset == 0 && Offset < EndOffset)
 				{
-					CommandUtils.LogConsole("Change {0} does not exist", CL);
+					CommandUtils.Log("Change {0} does not exist", CL);
 					return false;
 				}
 
@@ -1664,7 +1664,7 @@ namespace AutomationTool
 				}
 
 				string Status = CmdOutput.Substring(StatusOffset + StatusStr.Length, DescOffset - StatusOffset - StatusStr.Length).Trim();
-				CommandUtils.LogConsole("Change {0} exists ({1})", CL, Status);
+				CommandUtils.Log("Change {0} exists ({1})", CL, Status);
 				Pending = (Status == "pending");
 				return true;
 			}
@@ -1756,7 +1756,7 @@ namespace AutomationTool
 			string Output;
 			if (!LogP4Output(out Output, CommandLine, null, AllowSpew))
 			{
-				CommandUtils.LogConsole("Couldn't delete label '{0}'.  It may not have existed in the first place.", LabelName);
+				CommandUtils.Log("Couldn't delete label '{0}'.  It may not have existed in the first place.", LabelName);
 			}
 		}
 
@@ -1788,7 +1788,7 @@ namespace AutomationTool
 			LabelSpec += "View: \n";
 			LabelSpec += " " + View;
 
-			CommandUtils.LogConsole("Creating Label\n {0}\n", LabelSpec);
+			CommandUtils.Log("Creating Label\n {0}\n", LabelSpec);
 			LogP4("label -i", Input: LabelSpec);
 		}
 
@@ -2039,8 +2039,8 @@ namespace AutomationTool
         {
             CheckP4Enabled();
 			var Label = Env.LabelPrefix + BuildNamePrefix + "-CL-" + Env.ChangelistString;
-			CommandUtils.Log("Label prefix {0}", BuildNamePrefix);
-			CommandUtils.Log("Full Label name {0}", Label); 
+			CommandUtils.LogLog("Label prefix {0}", BuildNamePrefix);
+			CommandUtils.LogLog("Full Label name {0}", Label); 
             return Label;
         }
 
@@ -2062,13 +2062,13 @@ namespace AutomationTool
 			}
 
 			{
-				CommandUtils.Log("Making downstream label");
+				CommandUtils.LogLog("Making downstream label");
                 var Label = FullLabelName(Env, BuildNamePrefix);
 
-				CommandUtils.Log("Deleting old label {0} (if any)...", Label);
+				CommandUtils.LogLog("Deleting old label {0} (if any)...", Label);
                 DeleteLabel(Label, false);
 
-				CommandUtils.Log("Creating new label...");
+				CommandUtils.LogLog("Creating new label...");
 				CreateLabel(
 					Name: Label,
 					Description: "BVT Time " + CommandUtils.CmdEnv.TimestampAsString + "  CL " + Env.ChangelistString,
@@ -2077,12 +2077,12 @@ namespace AutomationTool
 					);
 				if (Files == null)
 				{
-					CommandUtils.Log("Adding all files to new label {0}...", Label);
+					CommandUtils.LogLog("Adding all files to new label {0}...", Label);
 					LabelSync(Label, false);
 				}
 				else
 				{
-					CommandUtils.Log("Adding build products to new label {0}...", Label);
+					CommandUtils.LogLog("Adding build products to new label {0}...", Label);
 					foreach (string LabelFile in Files)
 					{
 						LabelSync(Label, false, LabelFile);
@@ -2109,13 +2109,13 @@ namespace AutomationTool
             }
 
             {
-				CommandUtils.Log("Making downstream label");
+				CommandUtils.LogLog("Making downstream label");
                 var Label = FullLabelName(Env, BuildNamePrefix);
 
-				CommandUtils.Log("Deleting old label {0} (if any)...", Label);
+				CommandUtils.LogLog("Deleting old label {0} (if any)...", Label);
                 DeleteLabel(Label, false);
 
-				CommandUtils.Log("Creating new label...");
+				CommandUtils.LogLog("Creating new label...");
                 CreateLabel(
                     Name: Label,
 					Description: "BVT Time " + CommandUtils.CmdEnv.TimestampAsString + "  CL " + Env.ChangelistString,
@@ -2302,7 +2302,7 @@ namespace AutomationTool
 		/// <param name="Attributes">Attributes to set.</param>
 		public void ChangeFileType(string Filename, P4FileAttributes Attributes, string Changelist = null)
 		{
-			CommandUtils.Log("ChangeFileType({0}, {1}, {2})", Filename, Attributes, String.IsNullOrEmpty(Changelist) ? "null" : Changelist);
+			CommandUtils.LogLog("ChangeFileType({0}, {1}, {2})", Filename, Attributes, String.IsNullOrEmpty(Changelist) ? "null" : Changelist);
 
 			var Stat = FStat(Filename);
 			if (String.IsNullOrEmpty(Changelist))
@@ -2382,7 +2382,7 @@ namespace AutomationTool
 			CheckP4Enabled();
 			if(!Quiet)
 			{
-				CommandUtils.Log("Checking if client {0} exists", ClientName);
+				CommandUtils.LogLog("Checking if client {0} exists", ClientName);
 			}
 
             var P4Result = P4(String.Format("-c {0} where //...", ClientName), AllowSpew: false, WithClient: false);
@@ -2400,7 +2400,7 @@ namespace AutomationTool
 
 			if(!Quiet)
 			{
-				CommandUtils.Log("Getting info for client {0}", ClientName);
+				CommandUtils.LogLog("Getting info for client {0}", ClientName);
 			}
 			if (!DoesClientExist(ClientName, Quiet))
 			{
@@ -2583,7 +2583,7 @@ namespace AutomationTool
             {
                 SpecInput += "\t" + Mapping.Key + " //" + ClientSpec.Name + Mapping.Value + Environment.NewLine;
             }
-			if (AllowSpew) CommandUtils.Log(SpecInput);
+			if (AllowSpew) CommandUtils.LogLog(SpecInput);
             LogP4("client -i", SpecInput, AllowSpew: AllowSpew, WithClient: false);
             return ClientSpec;
         }
diff --git a/Engine/Source/Programs/AutomationTool/AutomationUtils/ProcessUtils.cs b/Engine/Source/Programs/AutomationTool/AutomationUtils/ProcessUtils.cs
--- a/Engine/Source/Programs/AutomationTool/AutomationUtils/ProcessUtils.cs
+++ b/Engine/Source/Programs/AutomationTool/AutomationUtils/ProcessUtils.cs
@@ -119,14 +119,14 @@ namespace AutomationTool
 				var ProcessName = ProcessesToKill[ProcessIndex].GetProcessName();
 				if (!String.IsNullOrEmpty(ProcessName) && !CanBeKilled(ProcessName))
 				{
-					CommandUtils.Log("Ignoring process \"{0}\" because it can't be killed.", ProcessName);
+					CommandUtils.LogLog("Ignoring process \"{0}\" because it can't be killed.", ProcessName);
 					ProcessesToKill.RemoveAt(ProcessIndex);
 				}
 			}
-			CommandUtils.Log("Trying to kill {0} spawned processes.", ProcessesToKill.Count);
+			CommandUtils.LogLog("Trying to kill {0} spawned processes.", ProcessesToKill.Count);
 			foreach (var Proc in ProcessesToKill)
 			{
-				CommandUtils.Log("  {0}", Proc.GetProcessName());
+				CommandUtils.LogLog("  {0}", Proc.GetProcessName());
 			}
             if (CommandUtils.IsBuildMachine)
             {
@@ -140,7 +140,7 @@ namespace AutomationTool
                             if (!Proc.HasExited)
                             {
                                 AllDone = false;
-								CommandUtils.Log("Waiting for process: {0}", Proc.GetProcessName());
+								CommandUtils.LogLog("Waiting for process: {0}", Proc.GetProcessName());
                             }
                         }
                         catch (Exception)
@@ -154,7 +154,7 @@ namespace AutomationTool
                         if (ProcessResult.HasAnyDescendants(Process.GetCurrentProcess()))
                         {
                             AllDone = false;
-							CommandUtils.LogConsole("Waiting for descendants of main process...");
+							CommandUtils.Log("Waiting for descendants of main process...");
                         }
                     }
                     catch (Exception)
@@ -177,7 +177,7 @@ namespace AutomationTool
 				{
 					if (!Proc.HasExited)
 					{
-						CommandUtils.Log("Killing process: {0}", ProcName);
+						CommandUtils.LogLog("Killing process: {0}", ProcName);
 						Proc.StopProcess(false);
 					}
 				}
@@ -191,7 +191,7 @@ namespace AutomationTool
             {
                 if (CommandUtils.IsBuildMachine && ProcessResult.HasAnyDescendants(Process.GetCurrentProcess()))
                 {
-					CommandUtils.Log("current process still has descendants, trying to kill them...");
+					CommandUtils.LogLog("current process still has descendants, trying to kill them...");
                     ProcessResult.KillAllDescendants(Process.GetCurrentProcess());
                 }
             }
@@ -448,7 +448,7 @@ namespace AutomationTool
 				}
                 if (!(bStdOutSignalReceived && bStdErrSignalReceived))
                 {
-					CommandUtils.Log("Waited for a long time for output of {0}, some output may be missing; we gave up.", AppName);
+					CommandUtils.LogLog("Waited for a long time for output of {0}, some output may be missing; we gave up.", AppName);
                 }
 
 				// Double-check if the process terminated
@@ -526,7 +526,7 @@ namespace AutomationTool
 						IsOurDescendant(ProcessToKill, KillCandidate.Id, VisitedPids))
 					{
 						KilledPids.Add(KillCandidate.Id);
-						CommandUtils.Log("Trying to kill descendant pid={0}, name={1}", KillCandidate.Id, KillCandidate.ProcessName);
+						CommandUtils.LogLog("Trying to kill descendant pid={0}, name={1}", KillCandidate.Id, KillCandidate.ProcessName);
 						try
 						{
 							KillCandidate.Kill();
@@ -555,7 +555,7 @@ namespace AutomationTool
                 HashSet<int> VisitedPids = new HashSet<int>();
                 if (ProcessManager.CanBeKilled(KillCandidate.ProcessName) && IsOurDescendant(ProcessToCheck, KillCandidate.Id, VisitedPids))
                 {
-					CommandUtils.Log("Descendant pid={0}, name={1}, filename={2}", KillCandidate.Id, KillCandidate.ProcessName,
+					CommandUtils.LogLog("Descendant pid={0}, name={1}, filename={2}", KillCandidate.Id, KillCandidate.ProcessName,
 						KillCandidate.MainModule != null ? KillCandidate.MainModule.FileName : "unknown");
                     return true;
                 }
@@ -587,11 +587,11 @@ namespace AutomationTool
 					ProcToKill.WaitForExit(60000);
 					if (!ProcToKill.HasExited)
 					{
-						CommandUtils.Log("Process {0} failed to exit.", ProcToKillName);
+						CommandUtils.LogLog("Process {0} failed to exit.", ProcToKillName);
 					}
 					else
 					{
-						CommandUtils.Log("Process {0} successfully exited.", ProcToKillName);
+						CommandUtils.LogLog("Process {0} successfully exited.", ProcToKillName);
 						OnProcessExited();
 					}
 					ProcToKill.Close();					
@@ -890,7 +890,7 @@ namespace AutomationTool
 
 			string App = CmdEnv.UATExe;
 
-			LogConsole("Running {0} {1}", App, CommandLine);
+			Log("Running {0} {1}", App, CommandLine);
 			var OSEnv = new Dictionary<string, string>();
 
 			OSEnv.Add(AutomationTool.EnvVarNames.LogFolder, LogSubdir);
@@ -954,7 +954,7 @@ namespace AutomationTool
 		{
 			while (!FileExists(LogFilename) && !LogProcess.HasExited)
 			{
-				LogConsole("Waiting for logging process to start...");
+				Log("Waiting for logging process to start...");
 				Thread.Sleep(2000);
 			}
 			Thread.Sleep(1000);
diff --git a/Engine/Source/Programs/AutomationTool/AutomationUtils/ProjectParams.cs b/Engine/Source/Programs/AutomationTool/AutomationUtils/ProjectParams.cs
--- a/Engine/Source/Programs/AutomationTool/AutomationUtils/ProjectParams.cs
+++ b/Engine/Source/Programs/AutomationTool/AutomationUtils/ProjectParams.cs
@@ -2052,95 +2052,95 @@ namespace AutomationTool
 			if (!bLogged)
 			{
 				// In alphabetical order.
-				CommandUtils.Log("Project Params **************");
-
-				CommandUtils.Log("AdditionalServerMapParams={0}", AdditionalServerMapParams);
-				CommandUtils.Log("Archive={0}", Archive);
-				CommandUtils.Log("ArchiveMetaData={0}", ArchiveMetaData);
-				CommandUtils.Log("CreateAppBundle={0}", CreateAppBundle);
-				CommandUtils.Log("BaseArchiveDirectory={0}", BaseArchiveDirectory);
-				CommandUtils.Log("BaseStageDirectory={0}", BaseStageDirectory);
-				CommandUtils.Log("Build={0}", Build);
-				CommandUtils.Log("Cook={0}", Cook);
-				CommandUtils.Log("Clean={0}", Clean);
-				CommandUtils.Log("Client={0}", Client);
-				CommandUtils.Log("ClientConfigsToBuild={0}", string.Join(",", ClientConfigsToBuild));
-				CommandUtils.Log("ClientCookedTargets={0}", ClientCookedTargets.ToString());
-				CommandUtils.Log("ClientTargetPlatform={0}", string.Join(",", ClientTargetPlatforms));
-				CommandUtils.Log("Compressed={0}", Compressed);
-				CommandUtils.Log("UseDebugParamForEditorExe={0}", UseDebugParamForEditorExe);
-				CommandUtils.Log("CookFlavor={0}", CookFlavor);
-				CommandUtils.Log("CookOnTheFly={0}", CookOnTheFly);
-				CommandUtils.Log("CookOnTheFlyStreaming={0}", CookOnTheFlyStreaming);
-				CommandUtils.Log("UnversionedCookedContent={0}", UnversionedCookedContent);
-				CommandUtils.Log("SkipCookingEditorContent={0}", SkipCookingEditorContent);
-                CommandUtils.Log("NumCookersToSpawn={0}", NumCookersToSpawn);
-                CommandUtils.Log("GeneratePatch={0}", GeneratePatch);
-                CommandUtils.Log("CreateReleaseVersion={0}", CreateReleaseVersion);
-                CommandUtils.Log("BasedOnReleaseVersion={0}", BasedOnReleaseVersion);
-                CommandUtils.Log("DLCName={0}", DLCName);
-                CommandUtils.Log("DLCIncludeEngineContent={0}", DLCIncludeEngineContent);
-                CommandUtils.Log("DiffCookedContentPath={0}", DiffCookedContentPath);
-                CommandUtils.Log("AdditionalCookerOptions={0}", AdditionalCookerOptions);
-				CommandUtils.Log("DedicatedServer={0}", DedicatedServer);
-				CommandUtils.Log("DirectoriesToCook={0}", DirectoriesToCook.ToString());
-                CommandUtils.Log("CulturesToCook={0}", CulturesToCook.ToString());
-				CommandUtils.Log("EditorTargets={0}", EditorTargets.ToString());
-				CommandUtils.Log("Foreign={0}", Foreign);
-				CommandUtils.Log("IsCodeBasedProject={0}", IsCodeBasedProject.ToString());
-				CommandUtils.Log("IsProgramTarget={0}", IsProgramTarget.ToString());
-				CommandUtils.Log("IterativeCooking={0}", IterativeCooking);
-                CommandUtils.Log("CookAll={0}", CookAll);
-                CommandUtils.Log("CookMapsOnly={0}", CookMapsOnly);
-                CommandUtils.Log("Deploy={0}", Deploy);
-				CommandUtils.Log("IterativeDeploy={0}", IterativeDeploy);
-				CommandUtils.Log("FastCook={0}", FastCook);
-				CommandUtils.Log("LogWindow={0}", LogWindow);
-				CommandUtils.Log("Manifests={0}", Manifests);
-				CommandUtils.Log("MapToRun={0}", MapToRun);
-				CommandUtils.Log("NoClient={0}", NoClient);
-				CommandUtils.Log("NumClients={0}", NumClients);                
-				CommandUtils.Log("NoDebugInfo={0}", NoDebugInfo);
-				CommandUtils.Log("NoCleanStage={0}", NoCleanStage);
-				CommandUtils.Log("NoXGE={0}", NoXGE);
-				CommandUtils.Log("MapsToCook={0}", MapsToCook.ToString());
-				CommandUtils.Log("Pak={0}", Pak);
-				CommandUtils.Log("Package={0}", Package);
-				CommandUtils.Log("NullRHI={0}", NullRHI);
-				CommandUtils.Log("FakeClient={0}", FakeClient);
-                CommandUtils.Log("EditorTest={0}", EditorTest);
-                CommandUtils.Log("RunAutomationTests={0}", RunAutomationTests); 
-                CommandUtils.Log("RunAutomationTest={0}", RunAutomationTest);
-                CommandUtils.Log("RunTimeoutSeconds={0}", RunTimeoutSeconds);
-                CommandUtils.Log("CrashIndex={0}", CrashIndex);
-				CommandUtils.Log("ProgramTargets={0}", ProgramTargets.ToString());
-                CommandUtils.Log("ProjectBinariesFolder={0}", ProjectBinariesFolder);
-				CommandUtils.Log("ProjectBinariesPath={0}", ProjectBinariesPath);
-				CommandUtils.Log("ProjectGameExeFilename={0}", ProjectGameExeFilename);
-				CommandUtils.Log("ProjectGameExePath={0}", ProjectGameExePath);
-				CommandUtils.Log("Distribution={0}", Distribution);
-                CommandUtils.Log("Prebuilt={0}", Prebuilt);
-				CommandUtils.Log("Prereqs={0}", Prereqs);
-				CommandUtils.Log("NoBootstrapExe={0}", NoBootstrapExe);
-				CommandUtils.Log("RawProjectPath={0}", RawProjectPath);
-				CommandUtils.Log("Rocket={0}", Rocket);
-				CommandUtils.Log("Run={0}", Run);
-				CommandUtils.Log("ServerConfigsToBuild={0}", string.Join(",", ServerConfigsToBuild));
-				CommandUtils.Log("ServerCookedTargets={0}", ServerCookedTargets.ToString());
-				CommandUtils.Log("ServerTargetPlatform={0}", string.Join(",", ServerTargetPlatforms));
-				CommandUtils.Log("ShortProjectName={0}", ShortProjectName.ToString());
-				CommandUtils.Log("SignedPak={0}", SignedPak);
-				CommandUtils.Log("SignPak={0}", SignPak);				
-				CommandUtils.Log("SkipCook={0}", SkipCook);
-				CommandUtils.Log("SkipCookOnTheFly={0}", SkipCookOnTheFly);
-				CommandUtils.Log("SkipPak={0}", SkipPak);
-				CommandUtils.Log("SkipStage={0}", SkipStage);
-				CommandUtils.Log("Stage={0}", Stage);
-				CommandUtils.Log("bUsesSteam={0}", bUsesSteam);
-				CommandUtils.Log("bUsesCEF3={0}", bUsesCEF3);
-				CommandUtils.Log("bUsesSlate={0}", bUsesSlate);
-                CommandUtils.Log("bDebugBuildsActuallyUseDebugCRT={0}", bDebugBuildsActuallyUseDebugCRT);
-				CommandUtils.Log("Project Params **************");
+				CommandUtils.LogLog("Project Params **************");
+
+				CommandUtils.LogLog("AdditionalServerMapParams={0}", AdditionalServerMapParams);
+				CommandUtils.LogLog("Archive={0}", Archive);
+				CommandUtils.LogLog("ArchiveMetaData={0}", ArchiveMetaData);
+				CommandUtils.LogLog("CreateAppBundle={0}", CreateAppBundle);
+				CommandUtils.LogLog("BaseArchiveDirectory={0}", BaseArchiveDirectory);
+				CommandUtils.LogLog("BaseStageDirectory={0}", BaseStageDirectory);
+				CommandUtils.LogLog("Build={0}", Build);
+				CommandUtils.LogLog("Cook={0}", Cook);
+				CommandUtils.LogLog("Clean={0}", Clean);
+				CommandUtils.LogLog("Client={0}", Client);
+				CommandUtils.LogLog("ClientConfigsToBuild={0}", string.Join(",", ClientConfigsToBuild));
+				CommandUtils.LogLog("ClientCookedTargets={0}", ClientCookedTargets.ToString());
+				CommandUtils.LogLog("ClientTargetPlatform={0}", string.Join(",", ClientTargetPlatforms));
+				CommandUtils.LogLog("Compressed={0}", Compressed);
+				CommandUtils.LogLog("UseDebugParamForEditorExe={0}", UseDebugParamForEditorExe);
+				CommandUtils.LogLog("CookFlavor={0}", CookFlavor);
+				CommandUtils.LogLog("CookOnTheFly={0}", CookOnTheFly);
+				CommandUtils.LogLog("CookOnTheFlyStreaming={0}", CookOnTheFlyStreaming);
+				CommandUtils.LogLog("UnversionedCookedContent={0}", UnversionedCookedContent);
+				CommandUtils.LogLog("SkipCookingEditorContent={0}", SkipCookingEditorContent);
+                CommandUtils.LogLog("NumCookersToSpawn={0}", NumCookersToSpawn);
+                CommandUtils.LogLog("GeneratePatch={0}", GeneratePatch);
+                CommandUtils.LogLog("CreateReleaseVersion={0}", CreateReleaseVersion);
+                CommandUtils.LogLog("BasedOnReleaseVersion={0}", BasedOnReleaseVersion);
+                CommandUtils.LogLog("DLCName={0}", DLCName);
+                CommandUtils.LogLog("DLCIncludeEngineContent={0}", DLCIncludeEngineContent);
+                CommandUtils.LogLog("DiffCookedContentPath={0}", DiffCookedContentPath);
+                CommandUtils.LogLog("AdditionalCookerOptions={0}", AdditionalCookerOptions);
+				CommandUtils.LogLog("DedicatedServer={0}", DedicatedServer);
+				CommandUtils.LogLog("DirectoriesToCook={0}", DirectoriesToCook.ToString());
+                CommandUtils.LogLog("CulturesToCook={0}", CulturesToCook.ToString());
+				CommandUtils.LogLog("EditorTargets={0}", EditorTargets.ToString());
+				CommandUtils.LogLog("Foreign={0}", Foreign);
+				CommandUtils.LogLog("IsCodeBasedProject={0}", IsCodeBasedProject.ToString());
+				CommandUtils.LogLog("IsProgramTarget={0}", IsProgramTarget.ToString());
+				CommandUtils.LogLog("IterativeCooking={0}", IterativeCooking);
+                CommandUtils.LogLog("CookAll={0}", CookAll);
+                CommandUtils.LogLog("CookMapsOnly={0}", CookMapsOnly);
+                CommandUtils.LogLog("Deploy={0}", Deploy);
+				CommandUtils.LogLog("IterativeDeploy={0}", IterativeDeploy);
+				CommandUtils.LogLog("FastCook={0}", FastCook);
+				CommandUtils.LogLog("LogWindow={0}", LogWindow);
+				CommandUtils.LogLog("Manifests={0}", Manifests);
+				CommandUtils.LogLog("MapToRun={0}", MapToRun);
+				CommandUtils.LogLog("NoClient={0}", NoClient);
+				CommandUtils.LogLog("NumClients={0}", NumClients);                
+				CommandUtils.LogLog("NoDebugInfo={0}", NoDebugInfo);
+				CommandUtils.LogLog("NoCleanStage={0}", NoCleanStage);
+				CommandUtils.LogLog("NoXGE={0}", NoXGE);
+				CommandUtils.LogLog("MapsToCook={0}", MapsToCook.ToString());
+				CommandUtils.LogLog("Pak={0}", Pak);
+				CommandUtils.LogLog("Package={0}", Package);
+				CommandUtils.LogLog("NullRHI={0}", NullRHI);
+				CommandUtils.LogLog("FakeClient={0}", FakeClient);
+                CommandUtils.LogLog("EditorTest={0}", EditorTest);
+                CommandUtils.LogLog("RunAutomationTests={0}", RunAutomationTests); 
+                CommandUtils.LogLog("RunAutomationTest={0}", RunAutomationTest);
+                CommandUtils.LogLog("RunTimeoutSeconds={0}", RunTimeoutSeconds);
+                CommandUtils.LogLog("CrashIndex={0}", CrashIndex);
+				CommandUtils.LogLog("ProgramTargets={0}", ProgramTargets.ToString());
+                CommandUtils.LogLog("ProjectBinariesFolder={0}", ProjectBinariesFolder);
+				CommandUtils.LogLog("ProjectBinariesPath={0}", ProjectBinariesPath);
+				CommandUtils.LogLog("ProjectGameExeFilename={0}", ProjectGameExeFilename);
+				CommandUtils.LogLog("ProjectGameExePath={0}", ProjectGameExePath);
+				CommandUtils.LogLog("Distribution={0}", Distribution);
+                CommandUtils.LogLog("Prebuilt={0}", Prebuilt);
+				CommandUtils.LogLog("Prereqs={0}", Prereqs);
+				CommandUtils.LogLog("NoBootstrapExe={0}", NoBootstrapExe);
+				CommandUtils.LogLog("RawProjectPath={0}", RawProjectPath);
+				CommandUtils.LogLog("Rocket={0}", Rocket);
+				CommandUtils.LogLog("Run={0}", Run);
+				CommandUtils.LogLog("ServerConfigsToBuild={0}", string.Join(",", ServerConfigsToBuild));
+				CommandUtils.LogLog("ServerCookedTargets={0}", ServerCookedTargets.ToString());
+				CommandUtils.LogLog("ServerTargetPlatform={0}", string.Join(",", ServerTargetPlatforms));
+				CommandUtils.LogLog("ShortProjectName={0}", ShortProjectName.ToString());
+				CommandUtils.LogLog("SignedPak={0}", SignedPak);
+				CommandUtils.LogLog("SignPak={0}", SignPak);				
+				CommandUtils.LogLog("SkipCook={0}", SkipCook);
+				CommandUtils.LogLog("SkipCookOnTheFly={0}", SkipCookOnTheFly);
+				CommandUtils.LogLog("SkipPak={0}", SkipPak);
+				CommandUtils.LogLog("SkipStage={0}", SkipStage);
+				CommandUtils.LogLog("Stage={0}", Stage);
+				CommandUtils.LogLog("bUsesSteam={0}", bUsesSteam);
+				CommandUtils.LogLog("bUsesCEF3={0}", bUsesCEF3);
+				CommandUtils.LogLog("bUsesSlate={0}", bUsesSlate);
+                CommandUtils.LogLog("bDebugBuildsActuallyUseDebugCRT={0}", bDebugBuildsActuallyUseDebugCRT);
+				CommandUtils.LogLog("Project Params **************");
 			}
 			bLogged = true;
 
diff --git a/Engine/Source/Programs/AutomationTool/AutomationUtils/UE4Build.cs b/Engine/Source/Programs/AutomationTool/AutomationUtils/UE4Build.cs
--- a/Engine/Source/Programs/AutomationTool/AutomationUtils/UE4Build.cs
+++ b/Engine/Source/Programs/AutomationTool/AutomationUtils/UE4Build.cs
@@ -78,7 +78,7 @@ namespace AutomationTool
 				FileNum++;
 			}
 			CopyFile(ManifestName, OutFile);
-			Log("Copied UBT manifest to {0}", OutFile);
+			LogLog("Copied UBT manifest to {0}", OutFile);
 
 
 			UnrealBuildTool.BuildManifest Manifest = ReadManifest(ManifestName);
@@ -461,10 +461,10 @@ namespace AutomationTool
 				string VerFile = CombinePaths(CmdEnv.LocalRoot, "Engine", "Build", "Build.version");
 				if (bDoUpdateVersionFiles)
 				{
-					Log("Updating {0} with:", VerFile);
-					Log("  Changelist={0}", ChangelistNumber);
-					Log("  IsLicenseeVersion={0}", bIsLicenseeVersion? 1 : 0);
-					Log("  BranchName={0}", Branch);
+					LogLog("Updating {0} with:", VerFile);
+					LogLog("  Changelist={0}", ChangelistNumber);
+					LogLog("  IsLicenseeVersion={0}", bIsLicenseeVersion? 1 : 0);
+					LogLog("  BranchName={0}", Branch);
 
 					string VerText = CommandUtils.ReadAllText(VerFile);
 
@@ -490,10 +490,10 @@ namespace AutomationTool
                 string VerFile = EngineVersionFile;
 				if (bDoUpdateVersionFiles)
 				{
-					Log("Updating {0} with:", VerFile);
-					Log(" #define	BRANCH_NAME  {0}", Branch);
-					Log(" #define	BUILT_FROM_CHANGELIST  {0}", ChangelistString);
-					Log(" #define   ENGINE_IS_LICENSEE_VERSION  {0}", bIsLicenseeVersion ? "1" : "0");
+					LogLog("Updating {0} with:", VerFile);
+					LogLog(" #define	BRANCH_NAME  {0}", Branch);
+					LogLog(" #define	BUILT_FROM_CHANGELIST  {0}", ChangelistString);
+					LogLog(" #define   ENGINE_IS_LICENSEE_VERSION  {0}", bIsLicenseeVersion ? "1" : "0");
 
 					VersionFileUpdater VersionH = new VersionFileUpdater(VerFile);
 					VersionH.ReplaceLine("#define BRANCH_NAME ", "\"" + Branch + "\"");
@@ -518,7 +518,7 @@ namespace AutomationTool
                     // Get the MAJOR/MINOR/PATCH from the Engine Version file, as it is authoritative. The rest we get from the P4Env.
                     string NewInformationalVersion = FEngineVersionSupport.FromVersionFile(EngineVersionFile, ChangelistNumber).ToString();
 
-                    Log("Updating {0} with AssemblyInformationalVersion: {1}", VerFile, NewInformationalVersion);
+                    LogLog("Updating {0} with AssemblyInformationalVersion: {1}", VerFile, NewInformationalVersion);
 
                     VersionFileUpdater VersionH = new VersionFileUpdater(VerFile);
                     VersionH.SetAssemblyInformationalVersion(NewInformationalVersion);
@@ -1018,7 +1018,7 @@ namespace AutomationTool
 											if (PCHFileName.Contains(@"\SharedPCHs\") && PCHFileName.Contains(@"\UE4Editor\"))
 											{
 												Key = "SharedEditorPCH$ " + PCHFileName;
-												Log("Hack: detected Shared PCH, which will use a different key {0}", Key);
+												LogLog("Hack: detected Shared PCH, which will use a different key {0}", Key);
 											}
 										}
 									}
@@ -1335,18 +1335,18 @@ namespace AutomationTool
 			string XGEConsole = XGEConsoleExe();
 			if (bCanUseXGE && string.IsNullOrEmpty(XGEConsole))
 			{
-				Log("XGE was requested, but is unavailable, so we won't use it.");
+				LogLog("XGE was requested, but is unavailable, so we won't use it.");
 				bCanUseXGE = false;
 			}
 
 			// only run ParallelExecutor if not running XGE (and we've requested ParallelExecutor and it exists)
 			bool bCanUseParallelExecutor = !bCanUseXGE && InUseParallelExecutor && !string.IsNullOrEmpty(ParallelExecutorExec());
-			Log("************************* UE4Build:");
-			Log("************************* ForceMonolithic: {0}", bForceMonolithic);
-			Log("************************* ForceNonUnity:{0} ", bForceNonUnity);
-			Log("************************* ForceDebugInfo: {0}", bForceDebugInfo);
-			Log("************************* UseXGE: {0}", bCanUseXGE);
-			Log("************************* UseParallelExecutor: {0}", bCanUseParallelExecutor);
+			LogLog("************************* UE4Build:");
+			LogLog("************************* ForceMonolithic: {0}", bForceMonolithic);
+			LogLog("************************* ForceNonUnity:{0} ", bForceNonUnity);
+			LogLog("************************* ForceDebugInfo: {0}", bForceDebugInfo);
+			LogLog("************************* UseXGE: {0}", bCanUseXGE);
+			LogLog("************************* UseParallelExecutor: {0}", bCanUseParallelExecutor);
 
 			// process XGE files
 			if (bCanUseXGE || bCanUseParallelExecutor)
@@ -1445,10 +1445,10 @@ namespace AutomationTool
 		{
 			// Check build products
 			{
-				Log("Build products *******");
+				LogLog("Build products *******");
 				if (BuildProductFiles.Count < 1)
 				{
-					LogConsole("No build products were made");
+					Log("No build products were made");
 				}
 				else
 				{
@@ -1458,10 +1458,10 @@ namespace AutomationTool
 						{
 							throw new AutomationException("BUILD FAILED {0} was a build product but no longer exists", Product);
 						}
-						Log(Product);
+						LogLog(Product);
 					}
 				}
-				Log("End Build products *******");
+				LogLog("End Build products *******");
 			}
 		}
 
@@ -1472,7 +1472,7 @@ namespace AutomationTool
 		/// <param name="Files">List of files to check out</param>
 		public static void AddBuildProductsToChangelist(int WorkingCL, List<string> Files)
 		{
-			LogConsole("Adding {0} build products to changelist {1}...", Files.Count, WorkingCL);
+			Log("Adding {0} build products to changelist {1}...", Files.Count, WorkingCL);
 			foreach (var File in Files)
 			{
 				P4.Sync("-f -k " + CommandUtils.MakePathSafeToUseWithCommandLine(File) + "#head"); // sync the file without overwriting local one
diff --git a/Engine/Source/Programs/AutomationTool/AutomationUtils/Utils.cs b/Engine/Source/Programs/AutomationTool/AutomationUtils/Utils.cs
--- a/Engine/Source/Programs/AutomationTool/AutomationUtils/Utils.cs
+++ b/Engine/Source/Programs/AutomationTool/AutomationUtils/Utils.cs
@@ -724,7 +724,7 @@ namespace AutomationTool
 	                if (Retry > 0)
 	                {
                         //@todo: These retries should be reported so we can track how often they are occurring.
-                        CommandUtils.LogConsole("*** Mac temp storage retry {0}", OutputFileName);
+                        CommandUtils.Log("*** Mac temp storage retry {0}", OutputFileName);
 	                    System.Threading.Thread.Sleep(1000);
 	                }
 	                bCopied = CommandUtils.CopyFile_NoExceptions(InputFileName, OutputFileName, true);
@@ -754,7 +754,7 @@ namespace AutomationTool
 	                while (!bFound && Retry < 60)
 	                {
                         //@todo: These retries should be reported so we can track how often they are occurring.
-                        CommandUtils.LogConsole("*** Mac temp storage retry {0}", Filename);
+                        CommandUtils.Log("*** Mac temp storage retry {0}", Filename);
 	                    System.Threading.Thread.Sleep(10000);
 	                    bFound = CommandUtils.FileExists_NoExceptions(bQuiet, Filename);
 	                    Retry++;
@@ -779,7 +779,7 @@ namespace AutomationTool
 	                while (!bFound && Retry < 60)
 	                {
                         //@todo: These retries should be reported so we can track how often they are occurring.
-                        CommandUtils.LogConsole("*** Mac temp storage retry {0}", Directoryname);
+                        CommandUtils.Log("*** Mac temp storage retry {0}", Directoryname);
 	                    System.Threading.Thread.Sleep(10000);
 	                    bFound = CommandUtils.DirectoryExists_NoExceptions(Directoryname);
 	                    Retry++;
@@ -810,7 +810,7 @@ namespace AutomationTool
 	                while (!bFound && Retry < NumRetries)
 	                {
                         //@todo: These retries should be reported so we can track how often they are occurring.
-                        CommandUtils.LogConsole("*** Mac temp storage retry {0}", Directoryname);
+                        CommandUtils.Log("*** Mac temp storage retry {0}", Directoryname);
 	                    System.Threading.Thread.Sleep(1000);
 	                    bFound = CommandUtils.DirectoryExistsAndIsWritable_NoExceptions(Directoryname);
 	                    Retry++;
@@ -888,7 +888,7 @@ namespace AutomationTool
             {
                 throw new AutomationException("Failed to find MAJOR, MINOR, and PATCH fields from version file {0}", Filename);
             }
-			CommandUtils.LogConsole("Read {0}.{1}.{2} from {3}", foundElements["MAJOR"], foundElements["MINOR"], foundElements["PATCH"], Filename);
+			CommandUtils.Log("Read {0}.{1}.{2} from {3}", foundElements["MAJOR"], foundElements["MINOR"], foundElements["PATCH"], Filename);
             return new Version(foundElements["MAJOR"], foundElements["MINOR"], foundElements["PATCH"]);
         }
 
diff --git a/Engine/Source/Programs/AutomationTool/GUBP/ElectricCommander.cs b/Engine/Source/Programs/AutomationTool/GUBP/ElectricCommander.cs
--- a/Engine/Source/Programs/AutomationTool/GUBP/ElectricCommander.cs
+++ b/Engine/Source/Programs/AutomationTool/GUBP/ElectricCommander.cs
@@ -187,7 +187,7 @@ namespace AutomationTool
 		{
 			try
 			{
-				CommandUtils.LogConsole("Updating node props for node {0}", NodeToDo.Name);
+				CommandUtils.Log("Updating node props for node {0}", NodeToDo.Name);
 				RunECTool(String.Format("setProperty \"/myWorkflow/FailEmails/{0}\" \"{1}\"", NodeToDo.Name, String.Join(" ", NodeToDo.RecipientsForFailureEmails)), true);
 			}
 			catch (Exception Ex)
@@ -200,7 +200,7 @@ namespace AutomationTool
 		{
 			try
 			{
-				CommandUtils.LogConsole("Updating duration prop for node {0}", NodeToDo.Name);
+				CommandUtils.Log("Updating duration prop for node {0}", NodeToDo.Name);
 				RunECTool(String.Format("setProperty \"/myWorkflow/NodeDuration/{0}\" \"{1}\"", NodeToDo.Name, BuildDuration.ToString()));
 				RunECTool(String.Format("setProperty \"/myJobStep/NodeDuration\" \"{0}\"", BuildDuration.ToString()));
 			}
@@ -893,7 +893,7 @@ public class TestECJobErrorParse : BuildCommand
 {
     public override void ExecuteBuild()
     {
-		LogConsole("*********************** TestECJobErrorParse");
+		Log("*********************** TestECJobErrorParse");
 
         string Filename = CombinePaths(@"P:\Builds\UE4\GUBP\++depot+UE4-2104401-RootEditor_Failed\Engine\Saved\Logs", "RootEditor_Failed.log");
         var Errors = ECJobPropsUtils.ErrorsFromProps(Filename);
diff --git a/Engine/Source/Programs/AutomationTool/GUBP/GUBP.cs b/Engine/Source/Programs/AutomationTool/GUBP/GUBP.cs
--- a/Engine/Source/Programs/AutomationTool/GUBP/GUBP.cs
+++ b/Engine/Source/Programs/AutomationTool/GUBP/GUBP.cs
@@ -150,7 +150,7 @@ public partial class GUBP : BuildCommand
 	/// </summary>
     public override void ExecuteBuild()
     {
-		LogConsole("************************* GUBP");
+		Log("************************* GUBP");
 
 		List<UnrealTargetPlatform> HostPlatforms = new List<UnrealTargetPlatform>();
         if (!ParseParam("NoPC"))
@@ -183,7 +183,7 @@ public partial class GUBP : BuildCommand
             // We must be passed a valid shelve and be a build machine (or be running a preflight test) to be a valid preflight.
             if (!string.IsNullOrEmpty(PreflightShelveCLString) && (IsBuildMachine || ParseParam("PreflightTest")))
             {
-                LogConsole("**** Preflight shelve {0}", PreflightShelveCLString);
+                Log("**** Preflight shelve {0}", PreflightShelveCLString);
                 PreflightShelveCL = int.Parse(PreflightShelveCLString);
                 if (PreflightShelveCL < 2000000)
                 {
@@ -481,7 +481,7 @@ public partial class GUBP : BuildCommand
 		AggregateNode[] MatchingAggregates = Aggregates.Where(x => x.Dependencies.All(y => Nodes.Contains(y))).ToArray();
 		if (MatchingAggregates.Length > 0)
 		{
-			LogConsole("*********** Aggregates");
+			Log("*********** Aggregates");
 			foreach (AggregateNode Aggregate in MatchingAggregates.OrderBy(x => x.Name))
 			{
 				StringBuilder Note = new StringBuilder("    " + Aggregate.Name);
@@ -489,7 +489,7 @@ public partial class GUBP : BuildCommand
 				{
 					Note.Append(" (promotable)");
 				}
-				LogConsole(Note.ToString());
+				Log(Note.ToString());
 			}
 		}
 
@@ -498,7 +498,7 @@ public partial class GUBP : BuildCommand
         string LastControllingTrigger = "";
         string LastAgentGroup = "";
 
-		LogConsole("*********** Desired And Dependent Nodes, in order.");
+		Log("*********** Desired And Dependent Nodes, in order.");
         foreach (BuildNode NodeToDo in Nodes)
         {
             string MyControllingTrigger = NodeToDo.ControllingTriggerDotName;
@@ -519,13 +519,13 @@ public partial class GUBP : BuildCommand
                             Finished = "(already triggered)";
                         }
                     }
-					LogConsole("  Controlling Trigger: {0}    {1}", MyControllingTrigger, Finished);
+					Log("  Controlling Trigger: {0}    {1}", MyControllingTrigger, Finished);
                 }
             }
 
 			if (NodeToDo.AgentSharingGroup != LastAgentGroup && NodeToDo.AgentSharingGroup != "")
             {
-				LogConsole("    Agent Group: {0}", NodeToDo.AgentSharingGroup);
+				Log("    Agent Group: {0}", NodeToDo.AgentSharingGroup);
             }
             LastAgentGroup = NodeToDo.AgentSharingGroup;
 
@@ -565,17 +565,17 @@ public partial class GUBP : BuildCommand
             {
 				Builder.AppendFormat(" {0}", String.Join(" ", NodeToDo.RecipientsForFailureEmails));
             }
-			LogConsole(Builder.ToString());
+			Log(Builder.ToString());
 
             if (bShowDependencies)
             {
 				foreach (BuildNode Dep in NodeToDo.InputDependencies.OrderBy(x => x.Name))
                 {
-					LogConsole("            dep> {0}", Dep.Name);
+					Log("            dep> {0}", Dep.Name);
                 }
 				foreach (BuildNode Dep in NodeToDo.OrderDependencies.OrderBy(x => x.Name))
 				{
-					LogConsole("           pdep> {0}", Dep.Name);
+					Log("           pdep> {0}", Dep.Name);
 				}
             }
         }
@@ -787,7 +787,7 @@ public partial class GUBP : BuildCommand
 
             if (!bProgressMade && NodesToDo.Count > 0)
             {
-				LogConsole("Cycle in GUBP, could not resolve:");
+				Log("Cycle in GUBP, could not resolve:");
                 foreach (BuildNode NodeToDo in NodesToDo)
                 {
                     string Deps = "";
@@ -818,7 +818,7 @@ public partial class GUBP : BuildCommand
                             Deps = Deps + Dep.Name + " ";
                         }
                     }
-					LogConsole("  {0}    deps: {1}", NodeToDo.Name, Deps);
+					Log("  {0}    deps: {1}", NodeToDo.Name, Deps);
                 }
                 throw new AutomationException("Cycle in GUBP");
             }
@@ -890,7 +890,7 @@ public partial class GUBP : BuildCommand
                         LastNonDuplicateFail = FindLastNonDuplicateFail(NodeToDo, History, JobInfo);
                         if (LastNonDuplicateFail < P4Env.Changelist)
                         {
-							LogConsole("*** Red-after-red spam reduction, changed CL {0} to CL {1} because the errors didn't change.", P4Env.Changelist, LastNonDuplicateFail);
+							Log("*** Red-after-red spam reduction, changed CL {0} to CL {1} because the errors didn't change.", P4Env.Changelist, LastNonDuplicateFail);
                         }
                     }
                 }
@@ -1301,7 +1301,7 @@ public partial class GUBP : BuildCommand
     private void DoCommanderSetup(ElectricCommander EC, IEnumerable<BuildNode> AllNodes, IEnumerable<AggregateNode> AllAggregates, List<BuildNode> OrderedToDo, int TimeIndex, int TimeQuantum, bool bSkipTriggers, bool bFake, bool bFakeEC, TriggerNode ExplicitTrigger, List<TriggerNode> UnfinishedTriggers, bool bPreflightBuild)
 	{
 		List<BuildNode> SortedNodes = TopologicalSort(new HashSet<BuildNode>(AllNodes), null, SubSort: false, DoNotConsiderCompletion: true);
-		LogConsole("******* {0} GUBP Nodes", SortedNodes.Count);
+		Log("******* {0} GUBP Nodes", SortedNodes.Count);
 
 		List<BuildNode> FilteredOrderedToDo = new List<BuildNode>();
 		using(TelemetryStopwatch StartFilterTimer = new TelemetryStopwatch("FilterNodes"))
@@ -1329,7 +1329,7 @@ public partial class GUBP : BuildCommand
 		}
 		using(TelemetryStopwatch PrintNodesTimer = new TelemetryStopwatch("SetupCommanderPrint"))
 		{
-			LogConsole("*********** EC Nodes, in order.");
+			Log("*********** EC Nodes, in order.");
 			PrintNodes(this, FilteredOrderedToDo, AllAggregates, UnfinishedTriggers, TimeQuantum);
 		}
 
@@ -1364,7 +1364,7 @@ public partial class GUBP : BuildCommand
             bool SaveSuccessRecords = (IsBuildMachine || bFakeEC) && // no real reason to make these locally except for fakeEC tests
                 (!(NodeToDo is TriggerNode) || NodeToDo.IsSticky); // trigger nodes are run twice, one to start the new workflow and once when it is actually triggered, we will save reconds for the latter
 
-            LogConsole("***** Running GUBP Node {0} -> {1} : {2}", NodeToDo.Name, GameNameIfAny, TempStorageNodeInfo.GetRelativeDirectory());
+            Log("***** Running GUBP Node {0} -> {1} : {2}", NodeToDo.Name, GameNameIfAny, TempStorageNodeInfo.GetRelativeDirectory());
             if (NodeToDo.IsComplete)
             {
 				NodeToDo.RetrieveBuildProducts(GameNameIfAny, StorageRootIfAny, TempStorageNodeInfo);
@@ -1381,12 +1381,12 @@ public partial class GUBP : BuildCommand
                 {
                     if (bFake)
                     {
-                        LogConsole("***** FAKE!! Building GUBP Node {0} for {1}", NodeToDo.Name, TempStorageNodeInfo.GetRelativeDirectory());
+                        Log("***** FAKE!! Building GUBP Node {0} for {1}", NodeToDo.Name, TempStorageNodeInfo.GetRelativeDirectory());
                         NodeToDo.DoFakeBuild();
                     }
                     else
                     {
-                        LogConsole("***** Building GUBP Node {0} for {1}", NodeToDo.Name, TempStorageNodeInfo.GetRelativeDirectory());
+                        Log("***** Building GUBP Node {0} for {1}", NodeToDo.Name, TempStorageNodeInfo.GetRelativeDirectory());
 						DateTime StartTime = DateTime.UtcNow;
 						using(TelemetryStopwatch DoBuildStopwatch = new TelemetryStopwatch("DoBuild.{0}", NodeToDo.Name))
 						{
@@ -1428,17 +1428,17 @@ public partial class GUBP : BuildCommand
 						EC.UpdateECBuildTime(NodeToDo, BuildDuration);
                     }
 
-					LogConsole("{0}", Ex);
+					Log("{0}", Ex);
 
 
                     if (History != null)
                     {
-						LogConsole("Changes since last green *********************************");
-						LogConsole("");
-						LogConsole("");
-						LogConsole("");
+						Log("Changes since last green *********************************");
+						Log("");
+						Log("");
+						Log("");
                         PrintDetailedChanges(History, P4Env.Changelist);
-						LogConsole("End changes since last green");
+						Log("End changes since last green");
                     }
 
                     string FailInfo = "";
diff --git a/Engine/Source/Programs/AutomationTool/GUBP/LegacyNodes.cs b/Engine/Source/Programs/AutomationTool/GUBP/LegacyNodes.cs
--- a/Engine/Source/Programs/AutomationTool/GUBP/LegacyNodes.cs
+++ b/Engine/Source/Programs/AutomationTool/GUBP/LegacyNodes.cs
@@ -2701,7 +2701,7 @@ partial class GUBP
                 var StartTime = DateTime.UtcNow;
                 TempStorage.CleanSharedTempStorageDirectory(GameNames);
                 var BuildDuration = (DateTime.UtcNow - StartTime).TotalMilliseconds;
-                LogConsole("Took {0}s to clear temp storage of old files.", BuildDuration / 1000);
+                Log("Took {0}s to clear temp storage of old files.", BuildDuration / 1000);
             }
 
             BuildProducts = new List<string>();
diff --git a/Engine/Source/Programs/AutomationTool/GUBP/Nodes/BuildNode.cs b/Engine/Source/Programs/AutomationTool/GUBP/Nodes/BuildNode.cs
--- a/Engine/Source/Programs/AutomationTool/GUBP/Nodes/BuildNode.cs
+++ b/Engine/Source/Programs/AutomationTool/GUBP/Nodes/BuildNode.cs
@@ -89,7 +89,7 @@ namespace AutomationTool
 
 		public virtual void RetrieveBuildProducts(string GameNameIfAny, string StorageRootIfAny, TempStorageNodeInfo TempStorageNodeInfo)
 		{
-			CommandUtils.LogConsole("***** Retrieving GUBP Node {0} from {1}", Name, TempStorageNodeInfo.GetRelativeDirectory());
+			CommandUtils.Log("***** Retrieving GUBP Node {0} from {1}", Name, TempStorageNodeInfo.GetRelativeDirectory());
 			try
 			{
 				bool WasLocal;
diff --git a/Engine/Source/Programs/AutomationTool/GUBP/TempStorage.cs b/Engine/Source/Programs/AutomationTool/GUBP/TempStorage.cs
--- a/Engine/Source/Programs/AutomationTool/GUBP/TempStorage.cs
+++ b/Engine/Source/Programs/AutomationTool/GUBP/TempStorage.cs
@@ -515,7 +515,7 @@ namespace AutomationTool
                     }
                 };
 
-            CommandUtils.LogConsole("Cleaning temp storage for {0}...", TopDirectory);
+            CommandUtils.Log("Cleaning temp storage for {0}...", TopDirectory);
             int FoldersCleaned = 0;
             int FoldersManuallyCleaned = 0;
             int FoldersFailedCleaned = 0;
@@ -540,7 +540,7 @@ namespace AutomationTool
                 {
                     try
                     {
-                        CommandUtils.LogConsole("Deleting folder with old temp storage {0}...", OldManifestInfo.NodeNameDir);
+                        CommandUtils.Log("Deleting folder with old temp storage {0}...", OldManifestInfo.NodeNameDir);
                         Directory.Delete(OldManifestInfo.NodeNameDir, true);
                         FoldersCleaned++;
                     }
@@ -957,7 +957,7 @@ namespace AutomationTool
                 else
                 {
                     var SharedStorageNodeDir = SharedTempStorageDirectory(TempStorageNodeInfo, GameName);
-                    CommandUtils.LogConsole("Storing to {0}", SharedStorageNodeDir);
+                    CommandUtils.Log("Storing to {0}", SharedStorageNodeDir);
                     // this folder should not already exist, else we have concurrency or job duplication problems.
                     if (CommandUtils.DirectoryExists_NoExceptions(SharedStorageNodeDir))
                     {
@@ -1006,7 +1006,7 @@ namespace AutomationTool
                 var LocalManifest = LocalTempStorageManifestFilename(TempStorageNodeInfo);
                 if (CommandUtils.FileExists_NoExceptions(LocalManifest))
                 {
-                    CommandUtils.LogConsole("Found local manifest {0}", LocalManifest);
+                    CommandUtils.Log("Found local manifest {0}", LocalManifest);
                     var Local = TempStorageManifest.Load(LocalManifest);
                     var Files = Local.GetFiles(RootDir);
                     var LocalTest = TempStorageManifest.Create(Files, RootDir);
@@ -1023,7 +1023,7 @@ namespace AutomationTool
                 // We couldn't find the node storage locally, so get it from the shared location.
                 var SharedStorageNodeDir = SharedTempStorageDirectory(TempStorageNodeInfo, GameName);
 
-                CommandUtils.LogConsole("Attempting to retrieve from {0}", SharedStorageNodeDir);
+                CommandUtils.Log("Attempting to retrieve from {0}", SharedStorageNodeDir);
                 if (!CommandUtils.DirectoryExists_NoExceptions(SharedStorageNodeDir))
                 {
                     throw new AutomationException("Storage Block Does Not Exists! {0}", SharedStorageNodeDir);
diff --git a/Engine/Source/Programs/AutomationTool/HTML5/HTML5Platform.Automation.cs b/Engine/Source/Programs/AutomationTool/HTML5/HTML5Platform.Automation.cs
--- a/Engine/Source/Programs/AutomationTool/HTML5/HTML5Platform.Automation.cs
+++ b/Engine/Source/Programs/AutomationTool/HTML5/HTML5Platform.Automation.cs
@@ -25,9 +25,9 @@ public class HTML5Platform : Platform
 
 	public override void Package(ProjectParams Params, DeploymentContext SC, int WorkingCL)
 	{
-		LogConsole("Package {0}", Params.RawProjectPath);
+		Log("Package {0}", Params.RawProjectPath);
 
-		LogConsole("Setting Emscripten SDK for packaging..");
+		Log("Setting Emscripten SDK for packaging..");
 		HTML5SDKInfo.SetupEmscriptenTemp();
 		HTML5SDKInfo.SetUpEmscriptenConfigFile();
 
@@ -98,7 +98,7 @@ public class HTML5Platform : Platform
         string FullGameExePath = Path.Combine(Path.GetDirectoryName(Params.ProjectGameExeFilename), GameExe);
         if (!SC.IsCodeBasedProject && !FileExists_NoExceptions(FullGameExePath))
         {
-			LogConsole("Failed to find game application " + FullGameExePath);
+			Log("Failed to find game application " + FullGameExePath);
 	        throw new AutomationException(ExitCode.Error_MissingExecutable, "Stage Failed. Could not find application {0}. You may need to build the UE4 project with your target configuration and platform.", FullGameExePath);
         }
 
@@ -133,12 +133,12 @@ public class HTML5Platform : Platform
 			if (!ConfigCache.GetInt32("BuildSettings", "HeapSize", out ConfigHeapSize))
 			{
 				ConfigHeapSize = Params.IsCodeBasedProject ? 1024 : 512;
-				LogConsole("Could not find Heap Size setting in .ini for Client config {0}", Params.ClientConfigsToBuild[0].ToString());
+				Log("Could not find Heap Size setting in .ini for Client config {0}", Params.ClientConfigsToBuild[0].ToString());
 			}
 		}
 
 		HeapSize = (ulong)ConfigHeapSize * 1024L * 1024L; // convert to bytes.
-		LogConsole("Setting Heap size to {0} Mb ", ConfigHeapSize);
+		Log("Setting Heap size to {0} Mb ", ConfigHeapSize);
 
 
 		GenerateFileFromTemplate(TemplateFile, OutputFile, Params.ShortProjectName, Params.ClientConfigsToBuild[0].ToString(), Params.StageCommandline, !Params.IsCodeBasedProject, HeapSize);
@@ -183,7 +183,7 @@ public class HTML5Platform : Platform
 
 	void CompressFile(string Source, string Destination)
 	{
-		LogConsole(" Compressing " + Source); 
+		Log(" Compressing " + Source); 
 		bool DeleteSource = false; 
 
 		if(  Source == Destination )
@@ -545,7 +545,7 @@ public class HTML5Platform : Platform
 
 		if ( !AmazonIdentity )
 		{
-			LogConsole("Amazon S3 Incorrectly configured"); 
+			Log("Amazon S3 Incorrectly configured"); 
 			return; 
 		}
 
@@ -562,7 +562,7 @@ public class HTML5Platform : Platform
 		}
 
 		Task.WaitAll(UploadTasks.ToArray());
-		LogConsole("Upload Tasks finished.");
+		Log("Upload Tasks finished.");
 	}
 
 	private static IDictionary<string, string> MimeTypeMapping = new Dictionary<string, string>(StringComparer.InvariantCultureIgnoreCase) 
@@ -574,7 +574,7 @@ public class HTML5Platform : Platform
 
 	public void UploadToS3Worker(FileInfo Info, string KeyId, string AccessKey, string BucketName, string FolderName )
 	{
-		LogConsole(" Uploading " + Info.Name); 
+		Log(" Uploading " + Info.Name); 
 
 			// force upload files even if the timestamps haven't changed. 
 			string TimeStamp = string.Format("{0:r}", DateTime.UtcNow);
@@ -637,11 +637,11 @@ public class HTML5Platform : Platform
 			}
 			catch (Exception ex)
 			{
-				LogConsole("Could not connect to S3, incorrect S3 Keys? " + ex.ToString());
+				Log("Could not connect to S3, incorrect S3 Keys? " + ex.ToString());
 				throw ex;
 			}
 
-			LogConsole(Info.Name + " has been uploaded "); 
+			Log(Info.Name + " has been uploaded "); 
 	}
 
 #endregion 
diff --git a/Engine/Source/Programs/AutomationTool/IOS/IOSPlatform.Automation.cs b/Engine/Source/Programs/AutomationTool/IOS/IOSPlatform.Automation.cs
--- a/Engine/Source/Programs/AutomationTool/IOS/IOSPlatform.Automation.cs
+++ b/Engine/Source/Programs/AutomationTool/IOS/IOSPlatform.Automation.cs
@@ -127,7 +127,7 @@ public class IOSPlatform : Platform
 
 	public override void Package(ProjectParams Params, DeploymentContext SC, int WorkingCL)
 	{
-		LogConsole("Package {0}", Params.RawProjectPath);
+		Log("Package {0}", Params.RawProjectPath);
 
 		// ensure the ue4game binary exists, if applicable
 		string FullExePath = CombinePaths(Path.GetDirectoryName(Params.ProjectGameExeFilename), SC.StageExecutables[0] + (UnrealBuildTool.BuildHostPlatform.Current.Platform != UnrealTargetPlatform.Mac ? ".stub" : ""));
@@ -216,10 +216,10 @@ public class IOSPlatform : Platform
 			// package a .ipa from the now staged directory
 			var IPPExe = CombinePaths(CmdEnv.LocalRoot, "Engine/Binaries/DotNET/IOS/IPhonePackager.exe");
 
-			Log("ProjectName={0}", Params.ShortProjectName);
-			Log("ProjectStub={0}", ProjectStub);
-			Log("ProjectIPA={0}", ProjectIPA);
-			Log("IPPExe={0}", IPPExe);
+			LogLog("ProjectName={0}", Params.ShortProjectName);
+			LogLog("ProjectStub={0}", ProjectStub);
+			LogLog("ProjectIPA={0}", ProjectIPA);
+			LogLog("IPPExe={0}", IPPExe);
 
 			bool cookonthefly = Params.CookOnTheFly || Params.SkipCookOnTheFly;
 
diff --git a/Engine/Source/Programs/AutomationTool/Linux/LinuxPlatform.Automation.cs b/Engine/Source/Programs/AutomationTool/Linux/LinuxPlatform.Automation.cs
--- a/Engine/Source/Programs/AutomationTool/Linux/LinuxPlatform.Automation.cs
+++ b/Engine/Source/Programs/AutomationTool/Linux/LinuxPlatform.Automation.cs
@@ -277,7 +277,7 @@ chmod 700 $HOME/Desktop/{1}.desktop", DesiredGLVersion, SC.ShortProjectName, SC.
 				&& ((!String.IsNullOrEmpty(ProjParams.DeviceUsername) && !String.IsNullOrEmpty(ProjParams.DevicePassword))
 					|| !ProjParams.Unattended)) // Skip key generation in unattended mode if information is missing
 			{
-				LogConsole("Configuring Linux host");
+				Log("Configuring Linux host");
 
 				// Prompt for username if not already set
 				while (String.IsNullOrEmpty(ProjParams.DeviceUsername))
diff --git a/Engine/Source/Programs/AutomationTool/Scripts/AnalyzeThirdPartyLibs.Automation.cs b/Engine/Source/Programs/AutomationTool/Scripts/AnalyzeThirdPartyLibs.Automation.cs
--- a/Engine/Source/Programs/AutomationTool/Scripts/AnalyzeThirdPartyLibs.Automation.cs
+++ b/Engine/Source/Programs/AutomationTool/Scripts/AnalyzeThirdPartyLibs.Automation.cs
@@ -127,7 +127,7 @@ class AnalyzeThirdPartyLibs : BuildCommand
 
 	public override void ExecuteBuild()
 	{
-		LogConsole("************************* Analyze Third Party Libs");
+		Log("************************* Analyze Third Party Libs");
 
 		// figure out what batch/script to run
 		switch (UnrealBuildTool.BuildHostPlatform.Current.Platform)
@@ -195,19 +195,19 @@ class AnalyzeThirdPartyLibs : BuildCommand
 
 			long Size = Info.GetSize(Platforms);
 
-			LogConsole("Library {0} is {1}", Lib, ToMegabytes(Size));
+			Log("Library {0} is {1}", Lib, ToMegabytes(Size));
 
 			long Total = 0;
 			for (int Index = 0; Index < Platforms.Count; ++Index)
 			{
 				PlatformLibraryInfo Platform = Platforms[Index];
 				long Growth = Platform.TotalSize - LastSizes[Index];
-				LogConsole("  {0} is {1}", Platform.PlatformName, ToMegabytes(Growth));
+				Log("  {0} is {1}", Platform.PlatformName, ToMegabytes(Growth));
 
 				LastSizes[Index] = Platform.TotalSize;
 				Total += Growth;
 			}
-			LogConsole("  Platform neutral is probably {0} (specific sum {1})", ToMegabytes(Size - Total), ToMegabytes(Total));
+			Log("  Platform neutral is probably {0} (specific sum {1})", ToMegabytes(Size - Total), ToMegabytes(Total));
 
 			TotalSize += Size;
 		}
@@ -218,7 +218,7 @@ class AnalyzeThirdPartyLibs : BuildCommand
 		LargeFileExtensions.AddRange(new string[] { ".pdb", ".a", ".lib", ".dll", ".dylib", ".bc", ".so" });
 
 		// Hackery, look for big files (re-traverses everything)
-		LogConsole("----");
+		Log("----");
 		foreach (string Lib in LibsToEvaluate)
 		{
 			ThirdPartyLibraryInfo Info = new ThirdPartyLibraryInfo(Lib);
@@ -226,10 +226,10 @@ class AnalyzeThirdPartyLibs : BuildCommand
 		}
 
 		// Hackery, look for VS mixes (re-traverses everything)
-		LogConsole("----");
+		Log("----");
 		long TotalShadow12 = 0;
 
-		LogConsole("Listing VS2012 directories that are shadowed by a VS2013 dir");
+		Log("Listing VS2012 directories that are shadowed by a VS2013 dir");
 		foreach (string Lib in LibsToEvaluate)
 		{
 			string[] Dirs = Directory.GetDirectories(Lib, "*.*", SearchOption.AllDirectories);
@@ -270,24 +270,24 @@ class AnalyzeThirdPartyLibs : BuildCommand
 					TotalShadow12 += Size;
 
 					string GoodPath = (LibDir + "/" + Dupe).Replace("\\", "/");
-					LogConsole("{0}", GoodPath);
+					Log("{0}", GoodPath);
 				}
 			}
 		}
-		LogConsole("OVERALL {0} of VS2012 files are shadowed by a VS2013 dir", ToMegabytes(TotalShadow12));
+		Log("OVERALL {0} of VS2012 files are shadowed by a VS2013 dir", ToMegabytes(TotalShadow12));
 
-		LogConsole("----");
+		Log("----");
 		foreach (var Platform in Platforms)
 		{
-			LogConsole("  {0} is {1} (estimate)", Platform.PlatformName, ToMegabytes(Platform.TotalSize));
+			Log("  {0} is {1} (estimate)", Platform.PlatformName, ToMegabytes(Platform.TotalSize));
 		}
-		LogConsole("  OVERALL is {0} (accurate)", ToMegabytes(TotalSize));
+		Log("  OVERALL is {0} (accurate)", ToMegabytes(TotalSize));
 
 		// undo the LibDir push
 		CommandUtils.PopDir();
 
 
-		LogConsole("Listing VS2012 bin directories that are shadowed by a VS2013 dir");
+		Log("Listing VS2012 bin directories that are shadowed by a VS2013 dir");
 		//string[] BinaryDirs
 		//foreach (string Lib in BinaryDirs)
 		{
@@ -331,7 +331,7 @@ class AnalyzeThirdPartyLibs : BuildCommand
 					TotalShadow12 += Size;
 
 					string GoodPath = Dupe.Replace("\\", "/");
-					LogConsole("{0}", GoodPath);
+					Log("{0}", GoodPath);
 				}
 			}
 		}
diff --git a/Engine/Source/Programs/AutomationTool/Scripts/ArchiveCommand.Automation.cs b/Engine/Source/Programs/AutomationTool/Scripts/ArchiveCommand.Automation.cs
--- a/Engine/Source/Programs/AutomationTool/Scripts/ArchiveCommand.Automation.cs
+++ b/Engine/Source/Programs/AutomationTool/Scripts/ArchiveCommand.Automation.cs
@@ -57,7 +57,7 @@ public partial class Project : CommandUtils
 			return;
 		}
 
-		LogConsole("********** ARCHIVE COMMAND STARTED **********");
+		Log("********** ARCHIVE COMMAND STARTED **********");
 
 		if (!Params.NoClient)
 		{
@@ -81,7 +81,7 @@ public partial class Project : CommandUtils
 				SC.StageTargetPlatform.ProcessArchivedProject(Params, SC);
 			}
 		}
-		LogConsole("********** ARCHIVE COMMAND COMPLETED **********");
+		Log("********** ARCHIVE COMMAND COMPLETED **********");
     }
 
 	#endregion
diff --git a/Engine/Source/Programs/AutomationTool/Scripts/BlameKeyword.Automation.cs b/Engine/Source/Programs/AutomationTool/Scripts/BlameKeyword.Automation.cs
--- a/Engine/Source/Programs/AutomationTool/Scripts/BlameKeyword.Automation.cs
+++ b/Engine/Source/Programs/AutomationTool/Scripts/BlameKeyword.Automation.cs
@@ -45,7 +45,7 @@ class BlameKeyword : BuildCommand
 
 	public override void ExecuteBuild()
 	{
-		LogConsole("************************** BlameKeyword");
+		Log("************************** BlameKeyword");
 
 		var Extensions = new HashSet<string>(KnownExtensions, StringComparer.InvariantCultureIgnoreCase);
 		var PathToSearch = ParseParamValue("Path", CombinePaths(CmdEnv.LocalRoot, "Engine", "Source", "Runtime"));
@@ -58,10 +58,10 @@ class BlameKeyword : BuildCommand
 			throw new AutomationException("No keywords to look for");
 		}
 
-		LogConsole("Path={0}", PathToSearch);
-		LogConsole("Keywords={0}", String.Join(",", Keywords));
-		LogConsole("Timelapse={0}", bUseTimelapse);
-		LogConsole("Out={0}", OutFilename);
+		Log("Path={0}", PathToSearch);
+		Log("Keywords={0}", String.Join(",", Keywords));
+		Log("Timelapse={0}", bUseTimelapse);
+		Log("Out={0}", OutFilename);
 
 		// Find files to search
 		var FilenamesToSearch = new List<string>();
@@ -131,7 +131,7 @@ class BlameKeyword : BuildCommand
 		}
 
 		// Write out the csv
-		LogConsole("Creating blame csv: {0}", OutFilename);
+		Log("Creating blame csv: {0}", OutFilename);
 		WriteAllText(OutFilename, BlameCSV.ToString());
 	}
 }
diff --git a/Engine/Source/Programs/AutomationTool/Scripts/BuildCommonTools.Automation.cs b/Engine/Source/Programs/AutomationTool/Scripts/BuildCommonTools.Automation.cs
--- a/Engine/Source/Programs/AutomationTool/Scripts/BuildCommonTools.Automation.cs
+++ b/Engine/Source/Programs/AutomationTool/Scripts/BuildCommonTools.Automation.cs
@@ -15,7 +15,7 @@ public class BuildCommonTools : BuildCommand
 {
 	public override void ExecuteBuild()
 	{
-		LogConsole("************************* BuildCommonTools");
+		Log("************************* BuildCommonTools");
 
 		// Get the list of platform names
 		string[] PlatformNames = ParseParamValue("platforms", BuildHostPlatform.Current.Platform.ToString()).Split('+');
@@ -168,10 +168,10 @@ public class ZipProjectUp : BuildCommand
         string InstallDirectory = ParseParamValue("install", "");
         ProjectDirectory = Path.GetDirectoryName(ProjectDirectory);
 
-        LogConsole("Started zipping project up");
-        LogConsole("Project directory: {0}", ProjectDirectory);
-		LogConsole("Install directory: {0}", InstallDirectory);
-        LogConsole("Packaging up the project...");
+        Log("Started zipping project up");
+        Log("Project directory: {0}", ProjectDirectory);
+		Log("Install directory: {0}", InstallDirectory);
+        Log("Packaging up the project...");
 
         // Setup filters
         FileFilter Filter = new FileFilter();
@@ -182,6 +182,6 @@ public class ZipProjectUp : BuildCommand
 
         ZipFiles(InstallDirectory, ProjectDirectory, Filter);
 
-		LogConsole("Completed zipping project up");
+		Log("Completed zipping project up");
     }
 }
\ No newline at end of file
diff --git a/Engine/Source/Programs/AutomationTool/Scripts/BuildCookRun.Automation.cs b/Engine/Source/Programs/AutomationTool/Scripts/BuildCookRun.Automation.cs
--- a/Engine/Source/Programs/AutomationTool/Scripts/BuildCookRun.Automation.cs
+++ b/Engine/Source/Programs/AutomationTool/Scripts/BuildCookRun.Automation.cs
@@ -49,7 +49,7 @@ public class BuildCookRun : BuildCommand
 
 	protected ProjectParams SetupParams()
 	{
-		LogConsole("Setting up ProjectParams for {0}", ProjectPath);
+		Log("Setting up ProjectParams for {0}", ProjectPath);
 
 		var Params = new ProjectParams
 		(
@@ -128,13 +128,13 @@ public class BuildCookRun : BuildCommand
 	private string GetDefaultMap(ProjectParams Params)
 	{
 		const string EngineEntryMap = "/Engine/Maps/Entry";
-		LogConsole("Trying to find DefaultMap in ini files");
+		Log("Trying to find DefaultMap in ini files");
 		string DefaultMap = null;
 		var ProjectFolder = GetDirectoryName(Params.RawProjectPath);
 		var DefaultGameEngineConfig = CombinePaths(ProjectFolder, "Config", "DefaultEngine.ini");
 		if (FileExists(DefaultGameEngineConfig))
 		{
-			LogConsole("Looking for DefaultMap in {0}", DefaultGameEngineConfig);
+			Log("Looking for DefaultMap in {0}", DefaultGameEngineConfig);
 			DefaultMap = GetDefaultMapFromIni(DefaultGameEngineConfig, Params.DedicatedServer);
 			if (DefaultMap == null && Params.DedicatedServer)
 			{
@@ -146,7 +146,7 @@ public class BuildCookRun : BuildCommand
 			var BaseEngineConfig = CombinePaths(CmdEnv.LocalRoot, "Config", "BaseEngine.ini");
 			if (FileExists(BaseEngineConfig))
 			{
-				LogConsole("Looking for DefaultMap in {0}", BaseEngineConfig);
+				Log("Looking for DefaultMap in {0}", BaseEngineConfig);
 				DefaultMap = GetDefaultMapFromIni(BaseEngineConfig, Params.DedicatedServer);
 				if (DefaultMap == null && Params.DedicatedServer)
 				{
@@ -157,12 +157,12 @@ public class BuildCookRun : BuildCommand
 		// We check for null here becase null == not found
 		if (DefaultMap == null)
 		{
-			LogConsole("No DefaultMap found, assuming: {0}", EngineEntryMap);
+			Log("No DefaultMap found, assuming: {0}", EngineEntryMap);
 			DefaultMap = EngineEntryMap;
 		}
 		else
 		{
-			LogConsole("Found DefaultMap={0}", DefaultMap);
+			Log("Found DefaultMap={0}", DefaultMap);
 		}
 		return DefaultMap;
 	}
@@ -233,7 +233,7 @@ public class BuildCookRun : BuildCommand
 		}
 
 		var Dest = ParseParamValue("ForeignDest", CombinePaths(@"C:\testue4\foreign\", DestSample + "_ _Dir"));
-		LogConsole("Make a foreign sample {0} -> {1}", Src, Dest);
+		Log("Make a foreign sample {0} -> {1}", Src, Dest);
 
 		CloneDirectory(Src, Dest);
 
@@ -259,7 +259,7 @@ public class BuildCookRun : BuildCommand
 		}
 
 		var Dest = ParseParamValue("ForeignDest", CombinePaths(@"C:\testue4\foreign\", DestSample + "_ _Dir"));
-		LogConsole("Make a foreign sample {0} -> {1}", Src, Dest);
+		Log("Make a foreign sample {0} -> {1}", Src, Dest);
 
 		CloneDirectory(Src, Dest);
 		DeleteDirectory_NoExceptions(CombinePaths(Dest, "Intermediate"));
diff --git a/Engine/Source/Programs/AutomationTool/Scripts/BuildProjectCommand.Automation.cs b/Engine/Source/Programs/AutomationTool/Scripts/BuildProjectCommand.Automation.cs
--- a/Engine/Source/Programs/AutomationTool/Scripts/BuildProjectCommand.Automation.cs
+++ b/Engine/Source/Programs/AutomationTool/Scripts/BuildProjectCommand.Automation.cs
@@ -32,7 +32,7 @@ public partial class Project : CommandUtils
 			return;
 		}
 
-		LogConsole("********** BUILD COMMAND STARTED **********");
+		Log("********** BUILD COMMAND STARTED **********");
 
 		var UE4Build = new UE4Build(Command);
 		var Agenda = new UE4Build.BuildAgenda();
@@ -136,7 +136,7 @@ public partial class Project : CommandUtils
 			UE4Build.AddBuildProductsToChangelist(WorkingCL, UE4Build.BuildProductFiles);
 		}
 
-		LogConsole("********** BUILD COMMAND COMPLETED **********");
+		Log("********** BUILD COMMAND COMPLETED **********");
 	}
 
 	#endregion
diff --git a/Engine/Source/Programs/AutomationTool/Scripts/BuildThirdPartyLibs.Automation.cs b/Engine/Source/Programs/AutomationTool/Scripts/BuildThirdPartyLibs.Automation.cs
--- a/Engine/Source/Programs/AutomationTool/Scripts/BuildThirdPartyLibs.Automation.cs
+++ b/Engine/Source/Programs/AutomationTool/Scripts/BuildThirdPartyLibs.Automation.cs
@@ -63,7 +63,7 @@ class BuildThirdPartyLibs : BuildCommand
 
 	public override void ExecuteBuild()
 	{
-		LogConsole("************************* Build Third Party Libs");
+		Log("************************* Build Third Party Libs");
 
 		// figure out what batch/script to run
 		string CompileScriptName;
@@ -90,7 +90,7 @@ class BuildThirdPartyLibs : BuildCommand
 		{
 			WorkingCL = P4.CreateChange(P4Env.Client, String.Format("Third party libs built from changelist {0}", P4Env.Changelist));
 		}
-		LogConsole("Build from {0}    Working in {1}", P4Env.Changelist, WorkingCL);
+		Log("Build from {0}    Working in {1}", P4Env.Changelist, WorkingCL);
 
 		// go to the third party lib dir
 		string SearchLibraryDir = ParseParamValue("SearchDir", DefaultLibraryDir);
@@ -142,7 +142,7 @@ class BuildThirdPartyLibs : BuildCommand
 		// now go through and run each batch file, 
 		foreach (string Lib in LibsToCompile)
 		{
-			LogConsole("Building {0}", Lib);
+			Log("Building {0}", Lib);
 
 			// go into the lib dir
 			CommandUtils.PushDir(Lib);
@@ -166,7 +166,7 @@ class BuildThirdPartyLibs : BuildCommand
 		{
 			int SubmittedCL;
 			P4.Submit(WorkingCL, out SubmittedCL, true, true);
-			LogConsole("Submitted changelist {0}", SubmittedCL);
+			Log("Submitted changelist {0}", SubmittedCL);
 		}
 	}
 }
diff --git a/Engine/Source/Programs/AutomationTool/Scripts/CodeSurgery.Automation.cs b/Engine/Source/Programs/AutomationTool/Scripts/CodeSurgery.Automation.cs
--- a/Engine/Source/Programs/AutomationTool/Scripts/CodeSurgery.Automation.cs
+++ b/Engine/Source/Programs/AutomationTool/Scripts/CodeSurgery.Automation.cs
@@ -11,7 +11,7 @@ class CodeSurgery : BuildCommand
 {
 	public override void ExecuteBuild()
 	{
-		LogConsole("************************* Reworking stats code");
+		Log("************************* Reworking stats code");
 
 		var Wildcards = new List<string>() { "*.h", "*.cpp", "*.inl" };
 
@@ -44,7 +44,7 @@ class CodeSurgery : BuildCommand
 						!ThisFile.Contains(@"\ThirdParty\")
 						)
 					{
-						LogConsole("Source File: {0}", ThisFile);
+						Log("Source File: {0}", ThisFile);
 						AllFiles.Add(ThisFile);
 					}
 				}
@@ -124,7 +124,7 @@ class CodeSurgery : BuildCommand
 
 						if (bEnumLine || bDeclareLine)
 						{
-							LogConsole("{0} {1} Line: {2} : {3}", bEnumLine ? "Enum" : "Declare", TypeString, StatName, Line);
+							Log("{0} {1} Line: {2} : {3}", bEnumLine ? "Enum" : "Declare", TypeString, StatName, Line);
 
 							if (bEnumLine)
 							{
@@ -189,32 +189,32 @@ class CodeSurgery : BuildCommand
 		}
 
 		var ToCheckOuts = new HashSet<string>();
-		LogConsole("Stats *************************");
+		Log("Stats *************************");
 		foreach (var AllGoodStat in AllGoodStats)
 		{
-			LogConsole("{0}", AllGoodStat);
+			Log("{0}", AllGoodStat);
 			ToCheckOuts.Add(DeclareFiles[AllGoodStat]);
 			ToCheckOuts.Add(EnumFiles[AllGoodStat]);
 		}
-		LogConsole("Groups *************************");
+		Log("Groups *************************");
 		foreach (var AllGoodGroup in AllGoodGroups)
 		{
-			LogConsole("{0}", AllGoodGroup);
+			Log("{0}", AllGoodGroup);
 			ToCheckOuts.Add(DeclareFiles[AllGoodGroup]);
 			ToCheckOuts.Add(EnumFiles[AllGoodGroup]);
 		}
-		LogConsole("Broken *************************");
+		Log("Broken *************************");
 		foreach (var BrokenItem in Broken)
 		{
-			LogConsole("{0}", BrokenItem.Key);
+			Log("{0}", BrokenItem.Key);
 		}
-		LogConsole("*************************");
+		Log("*************************");
 
 		int WorkingCL = -1;
 		if (P4Enabled)
 		{
 			WorkingCL = P4.CreateChange(P4Env.Client, "Stat code surgery");
-			LogConsole("Working in {0}", WorkingCL);
+			Log("Working in {0}", WorkingCL);
 		}
 		else
 		{
@@ -229,20 +229,20 @@ class CodeSurgery : BuildCommand
 				CheckedOuts.Add(ToCheckOut);
 			}
 		}
-		LogConsole("Checked Out *************************");
+		Log("Checked Out *************************");
 		foreach (var CheckedOut in CheckedOuts)
 		{
-			LogConsole("{0}", CheckedOut);
+			Log("{0}", CheckedOut);
 		}
-		LogConsole("Failed to check out *************************");
+		Log("Failed to check out *************************");
 		foreach (var ToCheckOut in ToCheckOuts)
 		{
 			if (!CheckedOuts.Contains(ToCheckOut))
 			{
-				LogConsole("{0}", ToCheckOut);
+				Log("{0}", ToCheckOut);
 			}
 		}
-		LogConsole("*************************");
+		Log("*************************");
 		foreach (var AllGoodStat in AllGoodStats)
 		{
 			if (EnumFiles[AllGoodStat].EndsWith(".cpp", StringComparison.InvariantCultureIgnoreCase))
@@ -286,7 +286,7 @@ class CodeSurgery : BuildCommand
 			WriteAllText(EnumFiles[AllGoodGroup], EnumFileText);
 		}
 
-		LogConsole("*************************");
+		Log("*************************");
 
 	}
 }
@@ -296,7 +296,7 @@ class UpdateCopyright : BuildCommand
 {
 	public override void ExecuteBuild()
 	{
-		LogConsole("************************* UpdateCopyright");
+		Log("************************* UpdateCopyright");
 
 		var Wildcards = new List<string>() { "*.h", "*.cpp", "*.inl", "*.cs" };
 
@@ -369,7 +369,7 @@ class UpdateCopyright : BuildCommand
 				{
 					if (!ThisFile.Contains(@"\ThirdParty\"))
 					{
-						LogConsole("Source File: {0}", ThisFile);
+						Log("Source File: {0}", ThisFile);
 						AllFiles.Add(ThisFile);
 					}
 				}
@@ -380,7 +380,7 @@ class UpdateCopyright : BuildCommand
 		if (P4Enabled)
 		{
 			WorkingCL = P4.CreateChange(P4Env.Client, "Stat code surgery");
-			LogConsole("Working in {0}", WorkingCL);
+			Log("Working in {0}", WorkingCL);
 		}
 		else
 		{
@@ -391,11 +391,11 @@ class UpdateCopyright : BuildCommand
 		{
 			if (P4.Edit_NoExceptions(WorkingCL, ToCheckOut))
 			{
-				LogConsole("Checked out {0}", ToCheckOut);
+				Log("Checked out {0}", ToCheckOut);
 			}
 			else if (ToCheckOut.Contains("CodeSurgery"))
 			{
-				LogConsole("Couldn't check out {0}", ToCheckOut);
+				Log("Couldn't check out {0}", ToCheckOut);
 			}
 			else
 			{
@@ -425,15 +425,15 @@ class UpdateCopyright : BuildCommand
 				Missing.Add(ThisFile);
 			}
 		}
-		LogConsole("Normal Copyright *************************");
+		Log("Normal Copyright *************************");
 		foreach (var ThisFile in Normal)
 		{
-			LogConsole("  {0}", ThisFile);
+			Log("  {0}", ThisFile);
 		}
-		LogConsole("Missing Copyright *************************");
+		Log("Missing Copyright *************************");
 		foreach (var ThisFile in Missing)
 		{
-			LogConsole("  {0}", ThisFile);
+			Log("  {0}", ThisFile);
 		}
 	}
 }
diff --git a/Engine/Source/Programs/AutomationTool/Scripts/CookCommand.Automation.cs b/Engine/Source/Programs/AutomationTool/Scripts/CookCommand.Automation.cs
--- a/Engine/Source/Programs/AutomationTool/Scripts/CookCommand.Automation.cs
+++ b/Engine/Source/Programs/AutomationTool/Scripts/CookCommand.Automation.cs
@@ -27,7 +27,7 @@ public partial class Project : CommandUtils
 		}
 		Params.ValidateAndLog();
 
-		LogConsole("********** COOK COMMAND STARTED **********");
+		Log("********** COOK COMMAND STARTED **********");
 
 		string UE4EditorExe = HostPlatform.Current.GetUE4ExePath(Params.UE4Exe);
 		if (!FileExists(UE4EditorExe))
@@ -64,7 +64,7 @@ public partial class Project : CommandUtils
 
 				if (ServerProcess != null)
 				{
-					LogConsole("Waiting a few seconds for the server to start...");
+					Log("Waiting a few seconds for the server to start...");
 					Thread.Sleep(5000);
 				}
 			}
@@ -98,7 +98,7 @@ public partial class Project : CommandUtils
 
 			if (Params.Clean.HasValue && Params.Clean.Value && !Params.IterativeCooking)
 			{
-				LogConsole("Cleaning cooked data.");
+				Log("Cleaning cooked data.");
 				CleanupCookedData(PlatformsToCook.ToList(), Params);
 			}
 
@@ -109,11 +109,11 @@ public partial class Project : CommandUtils
 				Maps = Params.MapsToCook.ToArray();
                 foreach (var M in Maps)
                 {
-					LogConsole("HasMapsToCook " + M.ToString());
+					Log("HasMapsToCook " + M.ToString());
                 }
                 foreach (var M in Params.MapsToCook)
                 {
-					LogConsole("Params.HasMapsToCook " + M.ToString());
+					Log("Params.HasMapsToCook " + M.ToString());
                 }
 			}
 
@@ -241,7 +241,7 @@ public partial class Project : CommandUtils
 				else
 				{
 					// Delete cooked data (if any) as it may be incomplete / corrupted.
-					LogConsole("Cook failed. Deleting cooked data.");
+					Log("Cook failed. Deleting cooked data.");
 					CleanupCookedData(PlatformsToCook.ToList(), Params);
 					throw new AutomationException(ExitCode.Error_UnknownCookFailure, Ex, "Cook failed.");
 				}
@@ -256,7 +256,7 @@ public partial class Project : CommandUtils
                 catch ( Exception Ex )
                 {
                     // Delete cooked data (if any) as it may be incomplete / corrupted.
-                    LogConsole("Cook failed. Deleting cooked data.");
+                    Log("Cook failed. Deleting cooked data.");
                     CleanupCookedData(PlatformsToCook.ToList(), Params);
                     throw new AutomationException(ExitCode.Error_UnknownCookFailure, Ex, "Cook failed.");
                 }
@@ -265,7 +265,7 @@ public partial class Project : CommandUtils
 		}
 
 
-		LogConsole("********** COOK COMMAND COMPLETED **********");
+		Log("********** COOK COMMAND COMPLETED **********");
 	}
 
     private static void DiffCookedContent( ProjectParams Params)
@@ -291,7 +291,7 @@ public partial class Project : CommandUtils
             }
             catch(Exception )
             {
-                LogConsole("Failed deleting temporary directories "+TemporaryPakPath+" "+TemporaryFilesPath+" continuing.");
+                Log("Failed deleting temporary directories "+TemporaryPakPath+" "+TemporaryFilesPath+" continuing.");
             }
 
             Directory.CreateDirectory(TemporaryPakPath);
@@ -371,19 +371,19 @@ public partial class Project : CommandUtils
                     {
                         if ( SourceFile[Index] != DestFile[Index] )
                         {
-					        LogConsole("Diff cooked content failed on file " +SourceFilename + " when comparing against "+DestFilename + " at offset " + Index.ToString() );
+					        Log("Diff cooked content failed on file " +SourceFilename + " when comparing against "+DestFilename + " at offset " + Index.ToString() );
                             string SavedSourceFilename = CombinePaths( FailedContentDirectory, "Source" + Path.GetFileName(SourceFilename));
                             string SavedDestFilename = CombinePaths( FailedContentDirectory, "Dest" + Path.GetFileName(DestFilename));
                             File.Copy(SourceFilename, SavedSourceFilename);
                             File.Copy(DestFilename, SavedDestFilename);
-                            LogConsole("Content temporarily saved to " +SavedSourceFilename + " and "+SavedDestFilename + " at offset " + Index.ToString() );
+                            Log("Content temporarily saved to " +SavedSourceFilename + " and "+SavedDestFilename + " at offset " + Index.ToString() );
                             break;
                         }
                     }
                 }
                 else
                 {
-                    LogConsole("Diff cooked content failed on file " +SourceFilename + " when comparing against "+DestFilename + " files are different sizes " + SourceFile.LongLength.ToString() + " " + DestFile.LongLength.ToString() );
+                    Log("Diff cooked content failed on file " +SourceFilename + " when comparing against "+DestFilename + " files are different sizes " + SourceFile.LongLength.ToString() + " " + DestFile.LongLength.ToString() );
                 }
             } 
         }
diff --git a/Engine/Source/Programs/AutomationTool/Scripts/CopyBuildToStagingDirectory.Automation.cs b/Engine/Source/Programs/AutomationTool/Scripts/CopyBuildToStagingDirectory.Automation.cs
--- a/Engine/Source/Programs/AutomationTool/Scripts/CopyBuildToStagingDirectory.Automation.cs
+++ b/Engine/Source/Programs/AutomationTool/Scripts/CopyBuildToStagingDirectory.Automation.cs
@@ -105,7 +105,7 @@ public partial class Project : CommandUtils
 
 		var UnrealPakExe = CombinePaths(CmdEnv.LocalRoot, "Engine/Binaries/Win64/UnrealPak.exe");
 
-		LogConsole("Running UnrealPak *******");
+		Log("Running UnrealPak *******");
 		string CmdLine = CommandUtils.MakePathSafeToUseWithCommandLine(OutputLocation) + " -create=" + CommandUtils.MakePathSafeToUseWithCommandLine(UnrealPakResponseFileName);
 		if (!String.IsNullOrEmpty(EncryptionKeys))
 		{
@@ -126,34 +126,34 @@ public partial class Project : CommandUtils
         }
 		CmdLine += PlatformOptions;
 		RunAndLog(CmdEnv, UnrealPakExe, CmdLine, Options: ERunOptions.Default | ERunOptions.UTF8Output);
-		LogConsole("UnrealPak Done *******");
+		Log("UnrealPak Done *******");
 	}
 
 	static public void LogDeploymentContext(DeploymentContext SC)
 	{
-		Log("Deployment Context **************");
-		Log("ArchiveDirectory = {0}", SC.ArchiveDirectory);
-		Log("RawProjectPath = {0}", SC.RawProjectPath);
-		Log("IsCodeBasedUprojectFile = {0}", SC.IsCodeBasedProject);
-		Log("DedicatedServer = {0}", SC.DedicatedServer);
-		Log("Stage = {0}", SC.Stage);
-		Log("StageTargetPlatform = {0}", SC.StageTargetPlatform.PlatformType.ToString());
-		Log("LocalRoot = {0}", SC.LocalRoot);
-		Log("ProjectRoot = {0}", SC.ProjectRoot);
-		Log("PlatformDir = {0}", SC.PlatformDir);
-		Log("StageProjectRoot = {0}", SC.StageProjectRoot);
-		Log("ShortProjectName = {0}", SC.ShortProjectName);
-		Log("StageDirectory = {0}", SC.StageDirectory);
-		Log("SourceRelativeProjectRoot = {0}", SC.SourceRelativeProjectRoot);
-		Log("RelativeProjectRootForStage = {0}", SC.RelativeProjectRootForStage);
-		Log("RelativeProjectRootForUnrealPak = {0}", SC.RelativeProjectRootForUnrealPak);
-		Log("ProjectArgForCommandLines = {0}", SC.ProjectArgForCommandLines);
-		Log("RuntimeRootDir = {0}", SC.RuntimeRootDir);
-		Log("RuntimeProjectRootDir = {0}", SC.RuntimeProjectRootDir);
-		Log("UProjectCommandLineArgInternalRoot = {0}", SC.UProjectCommandLineArgInternalRoot);
-		Log("PakFileInternalRoot = {0}", SC.PakFileInternalRoot);
-		Log("UnrealFileServerInternalRoot = {0}", SC.UnrealFileServerInternalRoot);
-		Log("End Deployment Context **************");
+		LogLog("Deployment Context **************");
+		LogLog("ArchiveDirectory = {0}", SC.ArchiveDirectory);
+		LogLog("RawProjectPath = {0}", SC.RawProjectPath);
+		LogLog("IsCodeBasedUprojectFile = {0}", SC.IsCodeBasedProject);
+		LogLog("DedicatedServer = {0}", SC.DedicatedServer);
+		LogLog("Stage = {0}", SC.Stage);
+		LogLog("StageTargetPlatform = {0}", SC.StageTargetPlatform.PlatformType.ToString());
+		LogLog("LocalRoot = {0}", SC.LocalRoot);
+		LogLog("ProjectRoot = {0}", SC.ProjectRoot);
+		LogLog("PlatformDir = {0}", SC.PlatformDir);
+		LogLog("StageProjectRoot = {0}", SC.StageProjectRoot);
+		LogLog("ShortProjectName = {0}", SC.ShortProjectName);
+		LogLog("StageDirectory = {0}", SC.StageDirectory);
+		LogLog("SourceRelativeProjectRoot = {0}", SC.SourceRelativeProjectRoot);
+		LogLog("RelativeProjectRootForStage = {0}", SC.RelativeProjectRootForStage);
+		LogLog("RelativeProjectRootForUnrealPak = {0}", SC.RelativeProjectRootForUnrealPak);
+		LogLog("ProjectArgForCommandLines = {0}", SC.ProjectArgForCommandLines);
+		LogLog("RuntimeRootDir = {0}", SC.RuntimeRootDir);
+		LogLog("RuntimeProjectRootDir = {0}", SC.RuntimeProjectRootDir);
+		LogLog("UProjectCommandLineArgInternalRoot = {0}", SC.UProjectCommandLineArgInternalRoot);
+		LogLog("PakFileInternalRoot = {0}", SC.PakFileInternalRoot);
+		LogLog("UnrealFileServerInternalRoot = {0}", SC.UnrealFileServerInternalRoot);
+		LogLog("End Deployment Context **************");
 	}
 
 	public static Dictionary<string, string> ConvertToLower(Dictionary<string, string> Mapping)
@@ -235,7 +235,7 @@ public partial class Project : CommandUtils
 		}
 		var ThisPlatform = SC.StageTargetPlatform;
 
-		LogConsole("Creating Staging Manifest...");
+		Log("Creating Staging Manifest...");
 
         if (Params.HasDLCName)
         {
@@ -375,15 +375,15 @@ public partial class Project : CommandUtils
 			
 			// Engine & Game Plugins. Push the Engine/Source working directory so UBT code has a correct RelativeEnginePath in ReadAvailablePlugins
 			ProjectDescriptor Project = ProjectDescriptor.FromFile(SC.RawProjectPath);
-			Log("Searching for plugins with CurrentWorkingDir: " + Directory.GetCurrentDirectory());
-			Log("Searching for plugins in: " + SC.RawProjectPath);
+			LogLog("Searching for plugins with CurrentWorkingDir: " + Directory.GetCurrentDirectory());
+			LogLog("Searching for plugins in: " + SC.RawProjectPath);
 			List<PluginInfo> AvailablePlugins = Plugins.ReadAvailablePlugins(CombinePaths(SC.LocalRoot, "Engine"), SC.RawProjectPath);
 			foreach (PluginInfo Plugin in AvailablePlugins)
 			{
-				Log("Considering Plugin for Stage: " + Plugin.FileName);
+				LogLog("Considering Plugin for Stage: " + Plugin.FileName);
 				if (UProjectInfo.IsPluginEnabledForProject(Plugin, Project, SC.StageTargetPlatform.PlatformType))
 				{
-					Log("EnabledPlugin: " + Plugin.FileName);
+					LogLog("EnabledPlugin: " + Plugin.FileName);
 					SC.StageFiles(StagedFileType.UFS, Plugin.Directory, "*.uplugin", false, null, null, true, !Params.UsePak(SC.StageTargetPlatform), null, true, false);
 				}
 			}
@@ -559,7 +559,7 @@ public partial class Project : CommandUtils
 
 	public static void CopyManifestFilesToStageDir(Dictionary<string, string> Mapping, string StageDir, string ManifestName, List<string> CRCFiles)
 	{
-		LogConsole("Copying {0} to staging directory: {1}", ManifestName, StageDir);
+		Log("Copying {0} to staging directory: {1}", ManifestName, StageDir);
 		string ManifestPath = "";
 		string ManifestFile = "";
 		if (!String.IsNullOrEmpty(ManifestName))
@@ -637,7 +637,7 @@ public partial class Project : CommandUtils
 	/// <param name="SC"></param>
 	private static void CreatePakUsingStagingManifest(ProjectParams Params, DeploymentContext SC)
 	{
-		LogConsole("Creating pak using staging manifest.");
+		Log("Creating pak using staging manifest.");
 
 		DumpManifest(SC, CombinePaths(CmdEnv.LogFolder, "PrePak" + (SC.DedicatedServer ? "_Server" : "")));
 
@@ -660,7 +660,7 @@ public partial class Project : CommandUtils
 			var PakBlacklistFilename = CombinePaths(SC.ProjectRoot, "Build", SC.PlatformDir, string.Format("PakBlacklist-{0}.txt", SC.StageTargetConfigurations[0].ToString()));
 			if (File.Exists(PakBlacklistFilename))
 			{
-				LogConsole("Applying PAK blacklist file {0}", PakBlacklistFilename);
+				Log("Applying PAK blacklist file {0}", PakBlacklistFilename);
 				string[] BlacklistContents = File.ReadAllLines(PakBlacklistFilename);
 				foreach (string Candidate in BlacklistContents)
 				{
@@ -698,7 +698,7 @@ public partial class Project : CommandUtils
 
 				if (bExcludeFile)
 				{
-					LogConsole("Excluding {0}", Src);
+					Log("Excluding {0}", Src);
 					continue;
 				}
 			}
@@ -778,13 +778,13 @@ public partial class Project : CommandUtils
 
 				if (InternalUtils.SafeCopyFile(SourceOutputLocation, OutputLocation))
 				{
-					LogConsole("Copying source pak from {0} to {1} instead of creating new pak", SourceOutputLocation, OutputLocation);
+					Log("Copying source pak from {0} to {1} instead of creating new pak", SourceOutputLocation, OutputLocation);
 					bCopiedExistingPak = true;
 				}
 			}
 			if (!bCopiedExistingPak)
 			{
-				LogConsole("Failed to copy source pak from {0} to {1}, creating new pak", SourceOutputLocation, OutputLocation);
+				Log("Failed to copy source pak from {0} to {1}, creating new pak", SourceOutputLocation, OutputLocation);
 			}
 		}
 
@@ -910,7 +910,7 @@ public partial class Project : CommandUtils
 	/// <param name="SC"></param>
 	private static void CreatePaksUsingChunkManifests(ProjectParams Params, DeploymentContext SC)
 	{
-		LogConsole("Creating pak using streaming install manifests.");
+		Log("Creating pak using streaming install manifests.");
 		DumpManifest(SC, CombinePaths(CmdEnv.LogFolder, "PrePak" + (SC.DedicatedServer ? "_Server" : "")));
 
 		var TmpPackagingPath = GetTmpPackagingPath(Params, SC);
@@ -1059,7 +1059,7 @@ public partial class Project : CommandUtils
 	public static void ApplyStagingManifest(ProjectParams Params, DeploymentContext SC)
 	{
 		MaybeConvertToLowerCase(Params, SC);
-		LogConsole("Cleaning Stage Directory: {0}", SC.StageDirectory);
+		Log("Cleaning Stage Directory: {0}", SC.StageDirectory);
 		if (SC.Stage && !Params.NoCleanStage && !Params.SkipStage && !Params.IterativeDeploy)
 		{
             try
@@ -1130,7 +1130,7 @@ public partial class Project : CommandUtils
 			return;
 		}
 
-		LogConsole("Creating UE4CommandLine.txt");
+		Log("Creating UE4CommandLine.txt");
 		if (!string.IsNullOrEmpty(Params.StageCommandline) || !string.IsNullOrEmpty(Params.RunCommandline))
 		{
 			string FileHostParams = " ";
@@ -1626,7 +1626,7 @@ public partial class Project : CommandUtils
 		{
 			Params.ValidateAndLog();
 
-			LogConsole("********** STAGE COMMAND STARTED **********");
+			Log("********** STAGE COMMAND STARTED **********");
 
 			if (!Params.NoClient)
 			{
@@ -1683,7 +1683,7 @@ public partial class Project : CommandUtils
 					ApplyStagingManifest(Params, SC);
 				}
 			}
-			LogConsole("********** STAGE COMMAND COMPLETED **********");
+			Log("********** STAGE COMMAND COMPLETED **********");
 		}
 	}
 
diff --git a/Engine/Source/Programs/AutomationTool/Scripts/DeployCommand.Automation.cs b/Engine/Source/Programs/AutomationTool/Scripts/DeployCommand.Automation.cs
--- a/Engine/Source/Programs/AutomationTool/Scripts/DeployCommand.Automation.cs
+++ b/Engine/Source/Programs/AutomationTool/Scripts/DeployCommand.Automation.cs
@@ -19,7 +19,7 @@ public partial class Project : CommandUtils
 			return;
 		}
 
-		LogConsole("********** DEPLOY COMMAND STARTED **********");
+		Log("********** DEPLOY COMMAND STARTED **********");
 		
 		if (!Params.NoClient)
 		{
@@ -34,8 +34,8 @@ public partial class Project : CommandUtils
 					ClientCmdLine += "-SourceDir=\"" + CombinePaths(Params.BaseStageDirectory, SC.StageTargetPlatform.PlatformType.ToString()) + "\" ";
 					string ClientApp = CombinePaths(CmdEnv.LocalRoot, "Engine/Binaries/Win64/UnrealFrontend.exe");
 
-					LogConsole("Deploying via UFE:");
-					LogConsole("\tClientCmdLine: " + ClientCmdLine + "");
+					Log("Deploying via UFE:");
+					Log("\tClientCmdLine: " + ClientCmdLine + "");
 
 					//@todo UAT: Consolidate running of external applications like UFE (See 'RunProjectCommand' for other instances)
 					PushDir(Path.GetDirectoryName(ClientApp));
@@ -68,7 +68,7 @@ public partial class Project : CommandUtils
 			}
 		}
 
-		LogConsole("********** DEPLOY COMMAND COMPLETED **********");
+		Log("********** DEPLOY COMMAND COMPLETED **********");
 	}
 
 	#endregion
diff --git a/Engine/Source/Programs/AutomationTool/Scripts/FixupRedirects.Automation.cs b/Engine/Source/Programs/AutomationTool/Scripts/FixupRedirects.Automation.cs
--- a/Engine/Source/Programs/AutomationTool/Scripts/FixupRedirects.Automation.cs
+++ b/Engine/Source/Programs/AutomationTool/Scripts/FixupRedirects.Automation.cs
@@ -11,7 +11,7 @@ class FixupRedirects : BuildCommand
 	public override void ExecuteBuild()
 	{
 		var EditorExe = CombinePaths(CmdEnv.LocalRoot, @"Engine/Binaries/Win64/UE4Editor-Cmd.exe");
-		LogConsole("********** Running FixupRedirects: {0} -run=FixupRedirects -unattended -nopause -buildmachine -forcelogflush", EditorExe);
+		Log("********** Running FixupRedirects: {0} -run=FixupRedirects -unattended -nopause -buildmachine -forcelogflush", EditorExe);
 		var RunResult = Run(EditorExe, String.Format("-run=FixupRedirects -unattended -nopause -buildmachine -forcelogflush -autosubmit"));
 		if (RunResult.ExitCode != 0)
 		{
diff --git a/Engine/Source/Programs/AutomationTool/Scripts/IPhonePackager.Automation.cs b/Engine/Source/Programs/AutomationTool/Scripts/IPhonePackager.Automation.cs
--- a/Engine/Source/Programs/AutomationTool/Scripts/IPhonePackager.Automation.cs
+++ b/Engine/Source/Programs/AutomationTool/Scripts/IPhonePackager.Automation.cs
@@ -11,7 +11,7 @@ class IPhonePackager : BuildCommand
 {
 	public override void ExecuteBuild()
 	{
-		LogConsole("************************* Calling IPP");
+		Log("************************* Calling IPP");
 
 		Platform IOS = Platform.Platforms[UnrealTargetPlatform.IOS];
 		string Command = ParseParamValue("cmd", "");
diff --git a/Engine/Source/Programs/AutomationTool/Scripts/LauncherLocalization.Automation.cs b/Engine/Source/Programs/AutomationTool/Scripts/LauncherLocalization.Automation.cs
--- a/Engine/Source/Programs/AutomationTool/Scripts/LauncherLocalization.Automation.cs
+++ b/Engine/Source/Programs/AutomationTool/Scripts/LauncherLocalization.Automation.cs
@@ -19,7 +19,7 @@ class LauncherLocalization : BuildCommand
 
 		if (P4Enabled)
 		{
-			LogConsole("Sync necessary content to head revision");
+			Log("Sync necessary content to head revision");
 			P4.Sync(P4Env.BuildRootP4 + "/Engine/Config/...");
 			P4.Sync(P4Env.BuildRootP4 + "/Engine/Content/...");
 			P4.Sync(P4Env.BuildRootP4 + "/Engine/Source/...");
@@ -28,7 +28,7 @@ class LauncherLocalization : BuildCommand
             P4.Sync(P4Env.BuildRootP4 + "/Portal/Content/...");
             P4.Sync(P4Env.BuildRootP4 + "/Portal/Source/...");
 
-			LogConsole("Localize from label {0}", P4Env.LabelToSync);
+			Log("Localize from label {0}", P4Env.LabelToSync);
 		}
 
 		OneSkyConfigData OneSkyConfig = OneSkyConfigHelper.Find("OneSkyConfig_EpicGames");
@@ -72,9 +72,9 @@ class LauncherLocalization : BuildCommand
 		// Execute commandlet for each set of arguments.
 		foreach (var CommandletArguments in CommandletArgumentSets)
 		{
-			LogConsole("Localization for {0} {1}", EditorArguments, CommandletArguments);
+			Log("Localization for {0} {1}", EditorArguments, CommandletArguments);
 
-			LogConsole("Running UE4Editor to generate Localization data");
+			Log("Running UE4Editor to generate Localization data");
 
 			string Arguments = String.Format("-run=GatherText {0} {1}", EditorArguments, CommandletArguments);
 			var RunResult = Run(EditorExe, Arguments);
@@ -261,8 +261,8 @@ namespace EpicGames.OneSkyLocalization.Config
 
 		public void SpewValues()
 		{
-			CommandUtils.LogConsole("Name : {0}", Name);
-			CommandUtils.LogConsole("ApiKey : {0}", ApiKey);
+			CommandUtils.Log("Name : {0}", Name);
+			CommandUtils.Log("ApiKey : {0}", ApiKey);
 			//CommandUtils.LogConsole("ApiSecret : {0}", ApiSecret);  // This should probably never be revealed.
 		}
 	}
diff --git a/Engine/Source/Programs/AutomationTool/Scripts/ListMobileDevices.Automation.cs b/Engine/Source/Programs/AutomationTool/Scripts/ListMobileDevices.Automation.cs
--- a/Engine/Source/Programs/AutomationTool/Scripts/ListMobileDevices.Automation.cs
+++ b/Engine/Source/Programs/AutomationTool/Scripts/ListMobileDevices.Automation.cs
@@ -8,7 +8,7 @@ namespace AutomationScripts.Automation
 	{
 		public override void ExecuteBuild()
 		{
-			LogConsole("======= ListMobileDevices - Start =======");
+			Log("======= ListMobileDevices - Start =======");
 
 			var GlobalParams = new ProjectParams(
 				Command: this,
@@ -25,7 +25,7 @@ namespace AutomationScripts.Automation
 				throw new AutomationException("iOS is not yet implemented.");
 			}
 
-			LogConsole("======= ListMobileDevices - Done ========");
+			Log("======= ListMobileDevices - Done ========");
 		}
 
 		private static void GetConnectedDevices(ProjectParams Params, Platform TargetPlatform)
@@ -38,7 +38,7 @@ namespace AutomationScripts.Automation
 			{
 				foreach (var DeviceName in ConnectedDevices)
 				{
-					LogConsole("Device:{0}:{1}", PlatformName, DeviceName);
+					Log("Device:{0}:{1}", PlatformName, DeviceName);
 				}
 			}
 			catch
diff --git a/Engine/Source/Programs/AutomationTool/Scripts/Localisation.Automation.cs b/Engine/Source/Programs/AutomationTool/Scripts/Localisation.Automation.cs
--- a/Engine/Source/Programs/AutomationTool/Scripts/Localisation.Automation.cs
+++ b/Engine/Source/Programs/AutomationTool/Scripts/Localisation.Automation.cs
@@ -87,9 +87,9 @@ class Localise : BuildCommand
 		// Execute commandlet for each set of arguments.
 		foreach (var CommandletArguments in CommandletArgumentSets)
 		{
-			LogConsole("Localization for {0} {1}", EditorArguments, CommandletArguments);
+			Log("Localization for {0} {1}", EditorArguments, CommandletArguments);
 
-			LogConsole("Running UE4Editor to generate localization data");
+			Log("Running UE4Editor to generate localization data");
 
 			string Arguments = String.Format("-run=GatherText {0} {1}", EditorArguments, CommandletArguments);
 			var RunResult = Run(EditorExe, Arguments);
diff --git a/Engine/Source/Programs/AutomationTool/Scripts/MegaXGE.Automation.cs b/Engine/Source/Programs/AutomationTool/Scripts/MegaXGE.Automation.cs
--- a/Engine/Source/Programs/AutomationTool/Scripts/MegaXGE.Automation.cs
+++ b/Engine/Source/Programs/AutomationTool/Scripts/MegaXGE.Automation.cs
@@ -26,7 +26,7 @@ class MegaXGE : BuildCommand
 			WorkingCL = P4.CreateChange(P4Env.Client, String.Format("MegaXGE build from changelist {0} - Params: {1}", P4Env.Changelist, CmdLine));
 		}
 
-		LogConsole("************************* MegaXGE");
+		Log("************************* MegaXGE");
 
 		bool Clean = ParseParam("Clean");
 		string CleanToolLocation = CombinePaths(CmdEnv.LocalRoot, "Engine", "Build", "Batchfiles", "Clean.bat");
@@ -53,7 +53,7 @@ class MegaXGE : BuildCommand
 			}
 		}
 
-		LogConsole("*************************");
+		Log("*************************");
 		for (int Arg = 1; Arg < 100; Arg++)
 		{
 			string Parm = String.Format("Target{0}", Arg);
@@ -132,7 +132,7 @@ class MegaXGE : BuildCommand
 					foreach (var Configuration in Configurations)
 					{
 						Agenda.AddTargets(new string[] { CurTarget }, Platform, Configuration);
-						LogConsole("Target {0} {1} {2}", CurTarget, Platform.ToString(), Configuration.ToString());
+						Log("Target {0} {1} {2}", CurTarget, Platform.ToString(), Configuration.ToString());
 						if (Clean)
 						{
 							string Args = String.Format("{0} {1} {2}", CurTarget, Platform.ToString(), Configuration.ToString());
@@ -142,7 +142,7 @@ class MegaXGE : BuildCommand
 				}
 			}
 		}
-		LogConsole("*************************");
+		Log("*************************");
 
 		Agenda.DoRetries = ParseParam("Retry");
 		UE4Build.Build(Agenda, InUpdateVersionFiles: IsBuildMachine, InUseParallelExecutor: ParseParam("useparallelexecutor"), InShowProgress: ShowProgress);
diff --git a/Engine/Source/Programs/AutomationTool/Scripts/PackageCommand.Automation.cs b/Engine/Source/Programs/AutomationTool/Scripts/PackageCommand.Automation.cs
--- a/Engine/Source/Programs/AutomationTool/Scripts/PackageCommand.Automation.cs
+++ b/Engine/Source/Programs/AutomationTool/Scripts/PackageCommand.Automation.cs
@@ -25,7 +25,7 @@ public partial class Project : CommandUtils
 		}
 		if (DeployContextList.Count > 0 && !Params.SkipStage)
 		{
-			LogConsole("********** PACKAGE COMMAND STARTED **********");
+			Log("********** PACKAGE COMMAND STARTED **********");
 
 			foreach (var SC in DeployContextList)
 			{
@@ -35,7 +35,7 @@ public partial class Project : CommandUtils
 				}
 			}
 
-			LogConsole("********** PACKAGE COMMAND COMPLETED **********");
+			Log("********** PACKAGE COMMAND COMPLETED **********");
 		}
 	}
 
diff --git a/Engine/Source/Programs/AutomationTool/Scripts/RocketBuild.Automation.cs b/Engine/Source/Programs/AutomationTool/Scripts/RocketBuild.Automation.cs
--- a/Engine/Source/Programs/AutomationTool/Scripts/RocketBuild.Automation.cs
+++ b/Engine/Source/Programs/AutomationTool/Scripts/RocketBuild.Automation.cs
@@ -417,7 +417,7 @@ namespace Rocket
 			for (int Idx = 0; Idx < SourceFileNames.Length; Idx++)
 			{
 				CommandUtils.CreateDirectory(Path.GetDirectoryName(TargetFileNames[Idx]));
-				CommandUtils.LogConsole("Stripping symbols: {0} -> {1}", SourceFileNames[Idx], TargetFileNames[Idx]);
+				CommandUtils.Log("Stripping symbols: {0} -> {1}", SourceFileNames[Idx], TargetFileNames[Idx]);
 				ToolChain.StripSymbols(SourceFileNames[Idx], TargetFileNames[Idx]);
 			}
 		}
@@ -587,7 +587,7 @@ namespace Rocket
 			for (int Idx = 0; Idx < SourceFileNames.Length; Idx++)
 			{
 				CommandUtils.CreateDirectory(Path.GetDirectoryName(TargetFileNames[Idx]));
-				CommandUtils.LogConsole("Signing code: {0} -> {1}", SourceFileNames[Idx], TargetFileNames[Idx]);
+				CommandUtils.Log("Signing code: {0} -> {1}", SourceFileNames[Idx], TargetFileNames[Idx]);
 				CommandUtils.CopyFile(SourceFileNames[Idx], TargetFileNames[Idx]);
 			}
 
@@ -949,10 +949,10 @@ namespace Rocket
 			}
 
 			// Write the list to the log
-			CommandUtils.LogConsole("Files to be included in Rocket build:");
+			CommandUtils.Log("Files to be included in Rocket build:");
 			foreach(KeyValuePair<string, bool> SortedFile in SortedFiles)
 			{
-				CommandUtils.LogConsole("  {0}{1}", SortedFile.Key, SortedFile.Value ? " (stripped)" : "");
+				CommandUtils.Log("  {0}{1}", SortedFile.Key, SortedFile.Value ? " (stripped)" : "");
 			}
 		}
 
@@ -1057,11 +1057,11 @@ namespace Rocket
 		static void CopyManifestFilesToOutput(string ManifestPath, string InputDir, string OutputDir)
 		{
 			// Read the files from the manifest
-			CommandUtils.LogConsole("Reading manifest: '{0}'", ManifestPath);
+			CommandUtils.Log("Reading manifest: '{0}'", ManifestPath);
 			string[] Files = CommandUtils.ReadAllLines(ManifestPath).Select(x => x.Trim()).Where(x => x.Length > 0).ToArray();
 
 			// Create lists of source and target files
-			CommandUtils.LogConsole("Preparing file lists...");
+			CommandUtils.Log("Preparing file lists...");
 			var SourceFiles = Files.Select(x => CommandUtils.CombinePaths(InputDir, x)).ToList();
 			var TargetFiles = Files.Select(x => CommandUtils.CombinePaths(OutputDir, x)).ToList();
 
@@ -1099,7 +1099,7 @@ namespace Rocket
 		{
 			// Create a zip file containing the install
 			string FullZipFileName = Path.Combine(CommandUtils.CmdEnv.LocalRoot, "FullInstall" + StaticGetHostPlatformSuffix(HostPlatform) + ".zip");
-			CommandUtils.LogConsole("Creating {0}...", FullZipFileName);
+			CommandUtils.Log("Creating {0}...", FullZipFileName);
 			CommandUtils.ZipFiles(FullZipFileName, LocalDir, new FileFilter(FileFilterType.Include));
 
 			// Create a filter for the files we need just to run the editor
@@ -1119,11 +1119,11 @@ namespace Rocket
 
 			// Create a zip file containing the editor install
 			string EditorZipFileName = Path.Combine(CommandUtils.CmdEnv.LocalRoot, "EditorInstall" + StaticGetHostPlatformSuffix(HostPlatform) + ".zip");
-			CommandUtils.LogConsole("Creating {0}...", EditorZipFileName);
+			CommandUtils.Log("Creating {0}...", EditorZipFileName);
 			CommandUtils.ZipFiles(EditorZipFileName, LocalDir, EditorFilter);
 
 			// Copy the files to their final location
-			CommandUtils.LogConsole("Copying files to {0}", PublishDir);
+			CommandUtils.Log("Copying files to {0}", PublishDir);
 			InternalUtils.Robust_CopyFile(FullZipFileName, Path.Combine(PublishDir, Path.GetFileName(FullZipFileName)));
 			InternalUtils.Robust_CopyFile(EditorZipFileName, Path.Combine(PublishDir, Path.GetFileName(EditorZipFileName)));
 			CommandUtils.DeleteFile(FullZipFileName);
@@ -1283,7 +1283,7 @@ namespace Rocket
 				List<string> ProjectPakFiles = new List<string>();
 				foreach(BranchInfo.BranchUProject Project in Projects)
 				{
-					CommandUtils.LogConsole("Generating DDC data for {0} on {1}", Project.GameName, TargetPlatforms);
+					CommandUtils.Log("Generating DDC data for {0} on {1}", Project.GameName, TargetPlatforms);
 					CommandUtils.DDCCommandlet(Project.FilePath, EditorExe, null, TargetPlatforms, "-fill -DDC=CreateInstalledEnginePak -ProjectOnly");
 
 					string ProjectPakFile = CommandUtils.CombinePaths(Path.GetDirectoryName(OutputPakFile), String.Format("Compressed-{0}.ddp", Project.GameName));
@@ -1298,7 +1298,7 @@ namespace Rocket
 				}
 
 				// Generate DDC for the editor, and merge all the other PAK files in
-				CommandUtils.LogConsole("Generating DDC data for engine content on {0}", TargetPlatforms);
+				CommandUtils.Log("Generating DDC data for engine content on {0}", TargetPlatforms);
 				CommandUtils.DDCCommandlet(null, EditorExe, null, TargetPlatforms, "-fill -DDC=CreateInstalledEnginePak " + CommandUtils.MakePathSafeToUseWithCommandLine("-MergePaks=" + String.Join("+", ProjectPakFiles)));
 
 				// Copy the DDP file to the output path
diff --git a/Engine/Source/Programs/AutomationTool/Scripts/RunProjectCommand.Automation.cs b/Engine/Source/Programs/AutomationTool/Scripts/RunProjectCommand.Automation.cs
--- a/Engine/Source/Programs/AutomationTool/Scripts/RunProjectCommand.Automation.cs
+++ b/Engine/Source/Programs/AutomationTool/Scripts/RunProjectCommand.Automation.cs
@@ -93,7 +93,7 @@ public partial class Project : CommandUtils
 			return;
 		}
 
-		LogConsole("********** RUN COMMAND STARTED **********");
+		Log("********** RUN COMMAND STARTED **********");
 
 		var LogFolderOutsideOfSandbox = GetLogFolderOutsideOfSandbox();
 		if (!GlobalCommandLine.Installed && ServerProcess == null)
@@ -118,7 +118,7 @@ public partial class Project : CommandUtils
 			CopyLogsBackToLogFolder();
 		}
 
-		LogConsole("********** RUN COMMAND COMPLETED **********");
+		Log("********** RUN COMMAND COMPLETED **********");
 	}
 
 	private static void RunInternal(ProjectParams Params, string ServerLogFile, string ClientLogFile)
@@ -152,13 +152,13 @@ public partial class Project : CommandUtils
 
 		if (ServerProcess != null)
 		{
-			LogConsole("Waiting a few seconds for the server to start...");
+			Log("Waiting a few seconds for the server to start...");
 			Thread.Sleep(5000);
 		}
 
 		if (!Params.NoClient)
 		{
-			LogConsole("Starting Client....");
+			Log("Starting Client....");
 
 			var SC = CreateDeploymentContext(Params, false);
 
@@ -216,17 +216,17 @@ public partial class Project : CommandUtils
 				ProcessResult ClientProcess = null;
 				FileStream ClientProcessLog = null;
 				StreamReader ClientLogReader = null;
-				LogConsole("Starting Client for unattended test....");
+				Log("Starting Client for unattended test....");
 				ClientProcess = Run(ClientApp, ClientCmdLine + " -FORCELOGFLUSH -testexit=\"" + LookFor + "\"", null, ClientRunFlags | ERunOptions.NoWaitForExit);
 				while (!FileExists(ClientLogFile) && !ClientProcess.HasExited)
 				{
-					LogConsole("Waiting for client logging process to start...{0}", ClientLogFile);
+					Log("Waiting for client logging process to start...{0}", ClientLogFile);
 					Thread.Sleep(2000);
 				}
 				if (FileExists(ClientLogFile))
 				{
 					Thread.Sleep(2000);
-					LogConsole("Client logging process started...{0}", ClientLogFile);
+					Log("Client logging process started...{0}", ClientLogFile);
 					ClientProcessLog = File.Open(ClientLogFile, FileMode.Open, FileAccess.Read, FileShare.ReadWrite);
 					ClientLogReader = new StreamReader(ClientProcessLog);
 				}
@@ -259,7 +259,7 @@ public partial class Project : CommandUtils
 						if (AllClientOutput.LastIndexOf(LookFor) > AllClientOutput.IndexOf(LookFor))
 						{
 							WelcomedCorrectly = true;
-							LogConsole("Test complete...");
+							Log("Test complete...");
 							bKeepReading = false;
 						}
 						else if (Params.RunAutomationTests)
@@ -315,13 +315,13 @@ public partial class Project : CommandUtils
 					}
 					else if (bClientExited && (DateTime.UtcNow - ExitTime).TotalSeconds > 30)
 					{
-						LogConsole("Client exited and has been quiet for 30 seconds...exiting");
+						Log("Client exited and has been quiet for 30 seconds...exiting");
 						bKeepReading = false;
 					}
 				}
 				if (ClientProcess != null && !ClientProcess.HasExited)
 				{
-					LogConsole("Client is supposed to exit, lets wait a while for it to exit naturally...");
+					Log("Client is supposed to exit, lets wait a while for it to exit naturally...");
 					for (int i = 0; i < 120 && !ClientProcess.HasExited; i++)
 					{
 						Thread.Sleep(1000);
@@ -329,7 +329,7 @@ public partial class Project : CommandUtils
 				}
 				if (ClientProcess != null && !ClientProcess.HasExited)
 				{
-					LogConsole("Stopping client...");
+					Log("Stopping client...");
 					ClientProcess.StopProcess();
 					Thread.Sleep(10000);
 				}
@@ -408,7 +408,7 @@ public partial class Project : CommandUtils
 			{
 				while (!FileExists(ServerLogFile) && !ServerProcess.HasExited)
 				{
-					LogConsole("Waiting for logging process to start...");
+					Log("Waiting for logging process to start...");
 					Thread.Sleep(2000);
 				}
 				Thread.Sleep(1000);
@@ -434,33 +434,33 @@ public partial class Project : CommandUtils
 								if (ClientProcess == null &&
 										(AllServerOutput.Contains("Game Engine Initialized") || AllServerOutput.Contains("Unreal Network File Server is ready")))
 								{
-									LogConsole("Starting Client for unattended test....");
+									Log("Starting Client for unattended test....");
 									ClientProcess = Run(ClientApp, ClientCmdLine + " -FORCELOGFLUSH -testexit=\"" + LookFor + "\"", null, ClientRunFlags | ERunOptions.NoWaitForExit);
 									//@todo no testing is done on these
 									if (NumClients > 1 && NumClients < 9)
 									{
 										for (int i = 1; i < NumClients; i++)
 										{
-											LogConsole("Starting Extra Client....");
+											Log("Starting Extra Client....");
 											OtherClients.Add(Run(ClientApp, ClientCmdLine, null, ClientRunFlags | ERunOptions.NoWaitForExit));
 										}
 									}
 									while (!FileExists(ClientLogFile) && !ClientProcess.HasExited)
 									{
-										LogConsole("Waiting for client logging process to start...{0}", ClientLogFile);
+										Log("Waiting for client logging process to start...{0}", ClientLogFile);
 										Thread.Sleep(2000);
 									}
 									if (!ClientProcess.HasExited)
 									{
 										Thread.Sleep(2000);
-										LogConsole("Client logging process started...{0}", ClientLogFile);
+										Log("Client logging process started...{0}", ClientLogFile);
 										ClientProcessLog = File.Open(ClientLogFile, FileMode.Open, FileAccess.Read, FileShare.ReadWrite);
 										ClientLogReader = new StreamReader(ClientProcessLog);
 									}
 								}
 								else if (ClientProcess == null && !ServerProcess.HasExited)
 								{
-									LogConsole("Waiting for server to start....");
+									Log("Waiting for server to start....");
 									Thread.Sleep(2000);
 								}
 								if (ClientProcess != null && ClientProcess.HasExited)
@@ -491,12 +491,12 @@ public partial class Project : CommandUtils
 										{
 											if (Params.FakeClient)
 											{
-												LogConsole("Welcomed by server or client loaded, lets wait ten minutes...");
+												Log("Welcomed by server or client loaded, lets wait ten minutes...");
 												Thread.Sleep(60000 * 10);
 											}
 											else
 											{
-												LogConsole("Welcomed by server or client loaded, lets wait 30 seconds...");
+												Log("Welcomed by server or client loaded, lets wait 30 seconds...");
 												Thread.Sleep(30000);
 											}
 											WelcomedCorrectly = true;
@@ -539,7 +539,7 @@ public partial class Project : CommandUtils
 					AllClientOutput += Output;
 					if (ClientProcess == null && (AllClientOutput.Contains("Game Engine Initialized") || AllClientOutput.Contains("Unreal Network File Server is ready")))
 					{
-						LogConsole("Starting Client....");
+						Log("Starting Client....");
 						var SC = DeployContextList[0];
 						ClientProcess = SC.StageTargetPlatform.RunClient(ClientRunFlags | ERunOptions.NoWaitForExit, ClientApp, ClientCmdLine, Params);
 //						ClientProcess = Run(ClientApp, ClientCmdLine, null, ClientRunFlags | ERunOptions.NoWaitForExit);
@@ -547,7 +547,7 @@ public partial class Project : CommandUtils
 						{
 							for (int i = 1; i < NumClients; i++)
 							{
-								LogConsole("Starting Extra Client....");
+								Log("Starting Extra Client....");
 								ProcessResult NewClient = SC.StageTargetPlatform.RunClient(ClientRunFlags | ERunOptions.NoWaitForExit, ClientApp, ClientCmdLine, Params);
 								OtherClients.Add(NewClient);
 							}
@@ -556,7 +556,7 @@ public partial class Project : CommandUtils
 				}
 				else if (ClientProcess == null && !ServerProcess.HasExited)
 				{
-					LogConsole("Waiting for server to start....");
+					Log("Waiting for server to start....");
 					Thread.Sleep(2000);
 				}
 
@@ -568,7 +568,7 @@ public partial class Project : CommandUtils
 				if (ClientProcess != null && ClientProcess.HasExited)
 				{
 
-					LogConsole("Client exited, stopping server....");
+					Log("Client exited, stopping server....");
 					if (!GlobalCommandLine.NoKill)
 					{
 						ServerProcess.StopProcess();
@@ -579,7 +579,7 @@ public partial class Project : CommandUtils
 				return bKeepReading; // Keep reading
 			});
 		}
-		LogConsole("Server exited....");
+		Log("Server exited....");
 		if (ClientProcess != null && !ClientProcess.HasExited)
 		{
 			ClientProcess.StopProcess();
@@ -888,8 +888,8 @@ public partial class Project : CommandUtils
 			ClientCmdLine += "-Params=\"" + TempCmdLine + "\"";
 			ClientApp = CombinePaths(CmdEnv.LocalRoot, "Engine/Binaries/Win64/UnrealFrontend.exe");
 
-			LogConsole("Launching via UFE:");
-			LogConsole("\tClientCmdLine: " + ClientCmdLine + "");
+			Log("Launching via UFE:");
+			Log("\tClientCmdLine: " + ClientCmdLine + "");
 		}
 		else
 		{
@@ -910,7 +910,7 @@ public partial class Project : CommandUtils
 		{
 			if (String.IsNullOrEmpty(Output) == false)
 			{
-				LogConsole(Output);
+				Log(Output);
 			}
 			return true;
 		});
@@ -1053,7 +1053,7 @@ public partial class Project : CommandUtils
 #endif
 		var UnrealFileServerExe = HostPlatform.Current.GetUE4ExePath("UnrealFileServer.exe");
 
-		LogConsole("Running UnrealFileServer *******");
+		Log("Running UnrealFileServer *******");
 		var Args = String.Format("{0} -abslog={1} -unattended -CrashForUAT -FORCELOGFLUSH -log {2}",
 						CommandUtils.MakePathSafeToUseWithCommandLine(Params.RawProjectPath),
 						CommandUtils.MakePathSafeToUseWithCommandLine(ServerLogFile),
diff --git a/Engine/Source/Programs/AutomationTool/Scripts/TestSnippet.Automation.cs b/Engine/Source/Programs/AutomationTool/Scripts/TestSnippet.Automation.cs
--- a/Engine/Source/Programs/AutomationTool/Scripts/TestSnippet.Automation.cs
+++ b/Engine/Source/Programs/AutomationTool/Scripts/TestSnippet.Automation.cs
@@ -8,6 +8,6 @@ class TestSnippetUtils : CommandUtils
 {
 	public static void Test(params object[] Args)
 	{
-		LogConsole("Args[0]=" + Args[0].ToString());
+		Log("Args[0]=" + Args[0].ToString());
 	}
 }
\ No newline at end of file
diff --git a/Engine/Source/Programs/AutomationTool/Scripts/Tests.Automation.cs b/Engine/Source/Programs/AutomationTool/Scripts/Tests.Automation.cs
--- a/Engine/Source/Programs/AutomationTool/Scripts/Tests.Automation.cs
+++ b/Engine/Source/Programs/AutomationTool/Scripts/Tests.Automation.cs
@@ -20,8 +20,8 @@ class TestP4_Info : BuildCommand
 {
 	public override void ExecuteBuild()
 	{
-		LogConsole("P4CLIENT={0}", GetEnvVar("P4CLIENT"));
-		LogConsole("P4PORT={0}", GetEnvVar("P4PORT"));
+		Log("P4CLIENT={0}", GetEnvVar("P4CLIENT"));
+		Log("P4PORT={0}", GetEnvVar("P4PORT"));
 
 		var Result = P4.P4("info");
 		if (Result != 0)
@@ -74,7 +74,7 @@ class GitPullRequest : BuildCommand
 			throw new AutomationException("Unable to find Git.exe in the system path under a 'PortableGit' subdirectory.  Make sure the GitHub client is installed, and you are running this script from within a GitHub Command Shell.  We want to make sure we're using the correct version of Git, in case multiple versions are on the computer, which is why we check for a PortableGit folder that the GitHub Shell uses.");
 		}
 
-		LogConsole("Running {0} {1}", GitExePath, GitCommandLine);
+		Log("Running {0} {1}", GitExePath, GitCommandLine);
 
 		ProcessResult Result = Run(App: GitExePath, CommandLine: GitCommandLine, Options: (ERunOptions.NoLoggingOfRunCommand | ERunOptions.AllowSpew | ERunOptions.AppMustExist));
 		if (Result > 0 || Result < 0)
@@ -116,7 +116,7 @@ class GitPullRequest : BuildCommand
 				{
 					throw new AutomationException("Unrecognized branch.");
 				}
-				LogConsole("Depot {0}", Depot);
+				Log("Depot {0}", Depot);
 
 				Base = Base.Substring(0, Base.IndexOf(BaseEnd));
 				if (!Base.Contains(" "))
@@ -124,10 +124,10 @@ class GitPullRequest : BuildCommand
 					throw new AutomationException("Unrecognized commit3.");
 				}
 				Base = Base.Substring(Base.LastIndexOf(" "));
-				LogConsole("CL String {0}", Base);
+				Log("CL String {0}", Base);
 				CL = int.Parse(Base);
 			}
-			LogConsole("CL int {0}", CL);
+			Log("CL int {0}", CL);
 			if (CL < 2000000 || String.IsNullOrWhiteSpace(Depot))
 			{
 				throw new AutomationException("Unrecognized commit3.");
@@ -161,7 +161,7 @@ class GitPullRequest : BuildCommand
 					throw new AutomationException("Unrecognized commit5.");
 				}
 				CLStuff = CLStuff.Substring(0, CLStuff.IndexOf(" "));
-				LogConsole("CL String {0}", CLStuff);
+				Log("CL String {0}", CLStuff);
 				CL = int.Parse(CLStuff);
 
 				var BranchStuff = Base.Substring(Base.IndexOf(LiveStart) + LiveStart.Length, Base.IndexOf(LiveEnd) - Base.IndexOf(LiveStart) - LiveStart.Length);
@@ -170,7 +170,7 @@ class GitPullRequest : BuildCommand
 					throw new AutomationException("Unrecognized commit6.");
 				}
 				BranchStuff = BranchStuff.Substring(BranchStuff.LastIndexOf(" ") + 1);
-				LogConsole("Branch String {0}", BranchStuff);
+				Log("Branch String {0}", BranchStuff);
 				if (BranchStuff.StartsWith("4."))
 				{
 					Depot = "//depot/UE4-Releases/4." + BranchStuff.Substring(2, 1);
@@ -188,7 +188,7 @@ class GitPullRequest : BuildCommand
 					throw new AutomationException("Unrecognized branch2.");
 				}
 			}
-			LogConsole("CL int {0}", CL);
+			Log("CL int {0}", CL);
 			if (CL < 2000000 || String.IsNullOrWhiteSpace(Depot))
 			{
 				throw new AutomationException("Unrecognized commit3.");
@@ -256,7 +256,7 @@ class GitPullRequest : BuildCommand
 
 		foreach (var p in NewClient.View)
 		{
-			LogConsole("{0} = {1}", p.Key, p.Value);
+			Log("{0} = {1}", p.Key, p.Value);
 		}
 
 		string TestClient = P4Env.User + "_" + Environment.MachineName + "_PullRequests_" + Depot.Replace("/", "_");
@@ -386,7 +386,7 @@ class GitPullRequest : BuildCommand
 			}
 			catch (Exception Ex)
 			{
-				LogConsole(" Exception was {0}", LogUtils.FormatException(Ex));
+				Log(" Exception was {0}", LogUtils.FormatException(Ex));
 				Failures.Add(String.Format("PR {0} Failed with {1}", PR, LogUtils.FormatException(Ex)));
 			}
 		}
@@ -394,7 +394,7 @@ class GitPullRequest : BuildCommand
 
 		foreach (var Failed in Failures)
 		{
-			LogConsole("{0}", Failed);
+			Log("{0}", Failed);
 		}
 	}
 }
@@ -413,7 +413,7 @@ class TestSuccess : BuildCommand
 {
 	public override void ExecuteBuild()
 	{
-		LogConsole("TestSuccess message.");
+		Log("TestSuccess message.");
 	}
 }
 
@@ -422,11 +422,11 @@ class TestMessage : BuildCommand
 {
 	public override void ExecuteBuild()
 	{
-		LogConsole("TestMessage message.");
-		LogConsole("you must be in error");
-		LogConsole("Behold the Error of you ways");
-		LogConsole("ERROR, you are silly");
-		LogConsole("ERROR: Something must be broken");
+		Log("TestMessage message.");
+		Log("you must be in error");
+		Log("Behold the Error of you ways");
+		Log("ERROR, you are silly");
+		Log("ERROR: Something must be broken");
 	}
 }
 
@@ -435,9 +435,9 @@ class TestRecursion : BuildCommand
 {
 	public override void ExecuteBuild()
 	{
-		LogConsole("TestRecursion *********************");
+		Log("TestRecursion *********************");
 		string Params = ParseParamValue("Cmd");
-		LogConsole("Recursive Command: {0}", Params);
+		Log("Recursive Command: {0}", Params);
 		RunUAT(CmdEnv, Params);
 
 	}
@@ -447,7 +447,7 @@ class TestRecursionAuto : BuildCommand
 {
 	public override void ExecuteBuild()
 	{
-		LogConsole("TestRecursionAuto *********************");
+		Log("TestRecursionAuto *********************");
 		RunUAT(CmdEnv, "TestMessage");
 		RunUAT(CmdEnv, "TestMessage");
 		RunUAT(CmdEnv, "TestRecursion -Cmd=TestMessage");
@@ -461,7 +461,7 @@ class TestMacZip : BuildCommand
 {
 	public override void ExecuteBuild()
 	{
-		LogConsole("TestMacZip *********************");
+		Log("TestMacZip *********************");
 
 		if (UnrealBuildTool.Utils.IsRunningOnMono)
 		{
@@ -482,14 +482,14 @@ class TestP4_CreateChangelist : BuildCommand
 {
 	public override void ExecuteBuild()
 	{
-		LogConsole("P4CLIENT={0}", GetEnvVar("P4CLIENT"));
-		LogConsole("P4PORT={0}", GetEnvVar("P4PORT"));
+		Log("P4CLIENT={0}", GetEnvVar("P4CLIENT"));
+		Log("P4PORT={0}", GetEnvVar("P4PORT"));
 
 		var CLDescription = "AutomationTool TestP4";
 
-		LogConsole("Creating new changelist \"{0}\" using client \"{1}\"", CLDescription, GetEnvVar("P4CLIENT"));
+		Log("Creating new changelist \"{0}\" using client \"{1}\"", CLDescription, GetEnvVar("P4CLIENT"));
 		var ChangelistNumber = P4.CreateChange(Description: CLDescription);
-		LogConsole("Created changelist {0}", ChangelistNumber);
+		Log("Created changelist {0}", ChangelistNumber);
 	}
 }
 
@@ -500,12 +500,12 @@ class TestP4_StrandCheckout : BuildCommand
 	{
 		int WorkingCL = P4.CreateChange(P4Env.Client, String.Format("TestP4_StrandCheckout, head={0}", P4Env.Changelist));
 
-		LogConsole("Build from {0}    Working in {1}", P4Env.Changelist, WorkingCL);
+		Log("Build from {0}    Working in {1}", P4Env.Changelist, WorkingCL);
 
 		List<string> Sign = new List<string>();
 		Sign.Add(CombinePaths(CmdEnv.LocalRoot, @"\Engine\Binaries\DotNET\AgentInterface.dll"));
 
-		LogConsole("Signing and adding {0} build products to changelist {1}...", Sign.Count, WorkingCL);
+		Log("Signing and adding {0} build products to changelist {1}...", Sign.Count, WorkingCL);
 
 		CodeSign.SignMultipleIfEXEOrDLL(this, Sign);
 		foreach (var File in Sign)
@@ -528,13 +528,13 @@ class TestP4_LabelDescription : BuildCommand
 	{
 		string Label = GetEnvVar("SBF_LabelFromUser");
 		string Desc;
-		LogConsole("LabelDescription {0}", Label);
+		Log("LabelDescription {0}", Label);
 		var Result = P4.LabelDescription(Label, out Desc);
 		if (!Result)
 		{
 			throw new AutomationException("Could not get label description");
 		}
-		LogConsole("Result***\n{0}\n***\n", Desc);
+		Log("Result***\n{0}\n***\n", Desc);
 	}
 }
 
@@ -607,8 +607,8 @@ class TestP4_ClientOps : BuildCommand
 					throw new AutomationException("{0} does not contain {1} and it is supposed to.", depot, P4Env.BuildRootP4);
 				}
 				string local = CombinePaths(depot.Replace(SlashRoot, LocalRoot));
-				LogConsole("{0}", depot);
-				LogConsole("    {0}", local);
+				Log("{0}", depot);
+				Log("    {0}", local);
 			}
 		}
 
@@ -633,7 +633,7 @@ class CleanDDC : BuildCommand
 {
 	public override void ExecuteBuild()
 	{
-		LogConsole("*********************** Clean DDC");
+		Log("*********************** Clean DDC");
 
 		bool DoIt = ParseParam("DoIt");
 
@@ -654,10 +654,10 @@ class CleanDDC : BuildCommand
 						if (FileExists_NoExceptions(file))
 						{
 							FileInfo Info = new FileInfo(file);
-							LogConsole("{0}", file);
+							Log("{0}", file);
 							if ((DateTime.UtcNow - Info.LastWriteTimeUtc).TotalDays > 20)
 							{
-								LogConsole("  is old.");
+								Log("  is old.");
 								if (DoIt)
 								{
 									DeleteFile_NoExceptions(file);
@@ -665,7 +665,7 @@ class CleanDDC : BuildCommand
 							}
 							else
 							{
-								LogConsole("  is NOT old.");
+								Log("  is NOT old.");
 							}
 						}
 					}
@@ -680,7 +680,7 @@ class TestTestFarm : BuildCommand
 {
 	public override void ExecuteBuild()
 	{
-		LogConsole("*********************** TestTestFarm");
+		Log("*********************** TestTestFarm");
 
 		string Exe = CombinePaths(CmdEnv.LocalRoot, "Engine", "Binaries", "Win64", "UE4Editor.exe");
 		string ClientLogFile = CombinePaths(CmdEnv.LogFolder, "HoverGameRun");
@@ -693,7 +693,7 @@ class TestTestFarm : BuildCommand
 			if (!String.IsNullOrEmpty(Output) &&
 					Output.Contains("Game Engine Initialized"))
 			{
-				LogConsole("It looks started, let me kill it....");
+				Log("It looks started, let me kill it....");
 				App.StopProcess();
 			}
 			return true; //keep reading
@@ -706,10 +706,10 @@ class TestArguments : BuildCommand
 {
 	public override void ExecuteBuild()
 	{
-		LogConsole("List of provided arguments: ");
+		Log("List of provided arguments: ");
 		for (int Index = 0; Index < Params.Length; ++Index)
 		{
-			LogConsole("{0}={1}", Index, Params[Index].ToString());
+			Log("{0}={1}", Index, Params[Index].ToString());
 		}
 
 		object[] TestArgs = new object[] { "NoXGE", "SomeArg", "Map=AwesomeMap", "run=whatever", "Stuff" };
@@ -762,7 +762,7 @@ public class TestCombinePaths : BuildCommand
 		for (int Index = 0; Index < Results.Length; ++Index)
 		{
 			var CombinedPath = Results[Index];
-			LogConsole(CombinedPath);
+			Log(CombinedPath);
 			if (CombinedPath != Results[0])
 			{
 				throw new AutomationException("Path {0} does not match path 0: {1} (expected: {2})", Index, CombinedPath, Results[0]);
@@ -786,7 +786,7 @@ class TestLog : BuildCommand
 {
 	public override void ExecuteBuild()
 	{
-		LogConsole("************************* ShooterGameRun");
+		Log("************************* ShooterGameRun");
 		string MapToRun = ParseParamValue("map", "/game/maps/sanctuary");
 
 		PushDir(CombinePaths(CmdEnv.LocalRoot, @"Engine/Binaries/Win64/"));
@@ -837,7 +837,7 @@ class TestGamePerf : BuildCommand
 
 	public override void ExecuteBuild()
 	{
-		LogConsole("*********************** TestGamePerf");
+		Log("*********************** TestGamePerf");
 
 		string UserName = "test.automation";
 		string Password = "JJGfh9CX";
@@ -1037,13 +1037,13 @@ class TestGamePerf : BuildCommand
 						{
 							if (!File.Exists(ClientLogFile))
 							{
-								LogConsole("Creating log file");
+								Log("Creating log file");
 
 								File.Create(ClientLogFile).Close();
 
-								LogConsole(ClientLogFile);
+								Log(ClientLogFile);
 
-								LogConsole("Log file created");
+								Log("Log file created");
 							}
 
 							RunAndLog(Exe, CmdLine, ClientLogFile);
@@ -1081,7 +1081,7 @@ class TestGamePerf : BuildCommand
 
 						catch (Exception Err)
 						{
-							LogConsole(Err.Message);
+							Log(Err.Message);
 						}
 					}
 				}
@@ -1090,7 +1090,7 @@ class TestGamePerf : BuildCommand
 
 		catch (Exception Err)
 		{
-			LogConsole(Err.Message);
+			Log(Err.Message);
 		}
 	}
 }
@@ -1123,10 +1123,10 @@ public class TestUATBuildProducts : BuildCommand
 		UE4Build Build = new UE4Build(this);
 		Build.AddUATFilesToBuildProducts();
 
-		LogConsole("Build products:");
+		Log("Build products:");
 		foreach (var Product in Build.BuildProductFiles)
 		{
-			LogConsole("  " + Product);
+			Log("  " + Product);
 		}
 	}
 }
@@ -1138,37 +1138,37 @@ public class TestWatchdogTimer : BuildCommand
 {
 	public override void ExecuteBuild()
 	{
-		LogConsole("Starting 1st timer (1s). This should not crash.");
+		Log("Starting 1st timer (1s). This should not crash.");
 		using (var SafeTimer = new WatchdogTimer(1))
 		{
 			// Wait 500ms
-			LogConsole("Started {0}", SafeTimer.GetProcessName());
+			Log("Started {0}", SafeTimer.GetProcessName());
 			Thread.Sleep(500);
 		}
-		LogConsole("First timer disposed successfully.");
+		Log("First timer disposed successfully.");
 
-		LogConsole("Starting 2nd timer (2s). This should throw an exception after 1 second.");
+		Log("Starting 2nd timer (2s). This should throw an exception after 1 second.");
 		try
 		{
 			using (var CrashTimer = new WatchdogTimer(2))
 			{
 				// Wait 5s (this will trigger the watchdog timer)
-				LogConsole("Started {0}", CrashTimer.GetProcessName());
+				Log("Started {0}", CrashTimer.GetProcessName());
 				Thread.Sleep(1000);
 				throw new Exception("Test exceptions under WatchdogTimer");
 			}
 		}
 		catch (Exception Ex)
 		{
-			LogConsole("Triggered exception guarded by WatchdogTimer:");
-			LogConsole(Ex.Message);
+			Log("Triggered exception guarded by WatchdogTimer:");
+			Log(Ex.Message);
 		}
 
-		LogConsole("Starting 3rd timer (2s). This should crash after 2 seconds.");
+		Log("Starting 3rd timer (2s). This should crash after 2 seconds.");
 		using (var CrashTimer = new WatchdogTimer(2))
 		{
 			// Wait 5s (this will trigger the watchdog timer)
-			LogConsole("Started {0}", CrashTimer.GetProcessName());
+			Log("Started {0}", CrashTimer.GetProcessName());
 			Thread.Sleep(5000);
 		}
 	}
@@ -1194,7 +1194,7 @@ class TestOSSCommands : BuildCommand
 		}
 		catch (Exception Err)
 		{
-			LogConsole(Err.Message);
+			Log(Err.Message);
 		}
 
 		FileInfo[] LogFiles = new DirectoryInfo(LogDirectory).GetFiles();
@@ -1302,15 +1302,15 @@ public class UBT : BuildCommand
 
 		Agenda.AddTargets(Targets.ToArray(), Platform, Configuration);
 
-		LogConsole("UBT Buid");
-		LogConsole("Targets={0}", String.Join(",", Targets));
-		LogConsole("Platform={0}", Platform);
-		LogConsole("Configuration={0}", Configuration);
-		LogConsole("Clean={0}", Clean);
+		Log("UBT Buid");
+		Log("Targets={0}", String.Join(",", Targets));
+		Log("Platform={0}", Platform);
+		Log("Configuration={0}", Configuration);
+		Log("Clean={0}", Clean);
 
 		Build.Build(Agenda, InUpdateVersionFiles: false);
 
-		LogConsole("UBT Completed");
+		Log("UBT Completed");
 	}
 }
 [Help("Zeroes engine versions in ObjectVersion.cpp and Version.h and checks them in.")]
@@ -1325,21 +1325,21 @@ public class ZeroEngineVersions : BuildCommand
 			var Stat = P4.FStat(VersionFilename);
 			if (Stat.IsValid && Stat.Action != P4Action.None)
 			{
-				LogConsole("Reverting {0}", VersionFilename);
+				Log("Reverting {0}", VersionFilename);
 				P4.Revert(VersionFilename);
 			}
 
-			LogConsole("Gettting engine version files @{0}", P4Env.Changelist);
+			Log("Gettting engine version files @{0}", P4Env.Changelist);
 			P4.Sync(String.Format("-f {0}@{1}", VersionFilename, P4Env.Changelist));
 		}
 
-		LogConsole("Checking if engine version files need to be reset...");
+		Log("Checking if engine version files need to be reset...");
 		List<string> FilesToSubmit = new List<string>();
 		{
 			VersionFileUpdater VersionH = new VersionFileUpdater(VersionFilename);
 			if (VersionH.Contains("#define ENGINE_VERSION 0") == false)
 			{
-				LogConsole("Zeroing out engine versions in {0}", VersionFilename);
+				Log("Zeroing out engine versions in {0}", VersionFilename);
 				VersionH.ReplaceLine("#define BRANCH_NAME ", "\"" + P4Env.BranchName + "\"");
 				VersionH.ReplaceLine("#define BUILT_FROM_CHANGELIST ", "0");
 
@@ -1355,14 +1355,14 @@ public class ZeroEngineVersions : BuildCommand
 			{
 				P4.Edit(CL, Filename);
 			}
-			LogConsole("Submitting CL #{0}...", CL);
+			Log("Submitting CL #{0}...", CL);
 			int SubmittedCL;
 			P4.Submit(CL, out SubmittedCL, false, true);
-			LogConsole("CL #{0} submitted as {1}", CL, SubmittedCL);
+			Log("CL #{0} submitted as {1}", CL, SubmittedCL);
 		}
 		else
 		{
-			LogConsole("Engine version files are set to 0.");
+			Log("Engine version files are set to 0.");
 		}
 	}
 }
@@ -1432,7 +1432,7 @@ public class SyncSource : BuildCommand
 		foreach (var AlwaysSyncSubFolder in AlwaysSync)
 		{
 			var SyncCmdLine = CombinePaths(PathSeparator.Depot, DepotPath, AlwaysSyncSubFolder, "...") + Changelist;
-			LogConsole("Syncing {0}", SyncCmdLine, Changelist);
+			Log("Syncing {0}", SyncCmdLine, Changelist);
 			SafeSync(SyncCmdLine);
 		}
 
@@ -1448,7 +1448,7 @@ public class SyncSource : BuildCommand
 				if (!ExclusionList.Contains(SubDirName))
 				{
 					var SyncCmdLine = CombinePaths(PathSeparator.Depot, SubDir, "...") + Changelist;
-					LogConsole("Syncing {0}", SyncCmdLine);
+					Log("Syncing {0}", SyncCmdLine);
 					SafeSync(SyncCmdLine);
 				}
 			}
@@ -1525,7 +1525,7 @@ class DumpBranch : BuildCommand
 
 	public override void ExecuteBuild()
 	{
-		LogConsole("************************* DumpBranch");
+		Log("************************* DumpBranch");
 
 		var HostPlatforms = new List<UnrealTargetPlatform>();
 		HostPlatforms.Add(UnrealTargetPlatform.Win64);
@@ -1566,7 +1566,7 @@ class TestBlame : BuildCommand
 		var Filename = ParseParamValue("File", "//depot/UE4/Engine/Source/Runtime/PakFile/Public/IPlatformFilePak.h");
 		var OutFilename = ParseParamValue("Out");
 
-		LogConsole("Creating blame file for {0}", Filename);
+		Log("Creating blame file for {0}", Filename);
 		P4Connection.BlameLineInfo[] Result = null;
 		if (ParseParam("Timelapse"))
 		{
@@ -1580,7 +1580,7 @@ class TestBlame : BuildCommand
 		foreach (var BlameLine in Result)
 		{
 			var ResultLine = String.Format("#{0} in {1} by {2}: {3}", BlameLine.Revision.Revision, BlameLine.Revision.Changelist, BlameLine.Revision.User, BlameLine.Line);
-			LogConsole(ResultLine);
+			Log(ResultLine);
 			BlameResult.AppendLine(ResultLine);
 		}
 		if (String.IsNullOrEmpty(OutFilename) == false)
@@ -1607,7 +1607,7 @@ class TestChanges : BuildCommand
 			}
 			foreach (var Record in ChangeRecords)
 			{
-				LogConsole("{0} {1} {2}", Record.CL, Record.UserEmail, Record.Summary);
+				Log("{0} {1} {2}", Record.CL, Record.UserEmail, Record.Summary);
 			}
 		}
 		{
@@ -1618,7 +1618,7 @@ class TestChanges : BuildCommand
 			}
 			foreach (var Record in ChangeRecords)
 			{
-				LogConsole("{0} {1} {2}", Record.CL, Record.UserEmail, Record.Summary);
+				Log("{0} {1} {2}", Record.CL, Record.UserEmail, Record.Summary);
 			}
 		}
 	}
@@ -1629,7 +1629,7 @@ class TestKillAll : BuildCommand
 {
 	public override void ExecuteBuild()
 	{
-		LogConsole("*********************** TestKillAll");
+		Log("*********************** TestKillAll");
 
 		string Exe = CombinePaths(CmdEnv.LocalRoot, "Engine", "Binaries", "Win64", "UE4Editor.exe");
 		string ClientLogFile = CombinePaths(CmdEnv.LogFolder, "HoverGameRun");
@@ -1645,7 +1645,7 @@ class TestCleanFormalBuilds : BuildCommand
 {
 	public override void ExecuteBuild()
 	{
-		LogConsole("*********************** TestCleanFormalBuilds");
+		Log("*********************** TestCleanFormalBuilds");
 		var Dir = ParseParamValue("Dir", @"P:\Builds\Soul\Soul_Android_Shipping_MakeBuild\++depot+UE4-CL-2077154");
 		var CLString = ParseParamValue("CL", "2077154");
 		CleanFormalBuilds(Dir, CLString);
@@ -1658,7 +1658,7 @@ class TestStopProcess : BuildCommand
 {
 	public override void ExecuteBuild()
 	{
-		LogConsole("*********************** TestStopProcess");
+		Log("*********************** TestStopProcess");
 
 		string Exe = CombinePaths(CmdEnv.LocalRoot, "Engine", "Binaries", "Win64", "UE4Editor.exe");
 		string ClientLogFile = CombinePaths(CmdEnv.LogFolder, "HoverGameRun");
@@ -1666,13 +1666,13 @@ class TestStopProcess : BuildCommand
 
 		for (int Attempt = 0; Attempt < 5; ++Attempt)
 		{
-			LogConsole("Attempt: {0}", Attempt);
+			Log("Attempt: {0}", Attempt);
 			var Proc = Run(Exe, CmdLine, null, ERunOptions.AllowSpew | ERunOptions.NoWaitForExit | ERunOptions.AppMustExist | ERunOptions.NoStdOutRedirect);
 			Thread.Sleep(10000);
 			Proc.StopProcess();
 		}
 
-		LogConsole("One final attempt to test KillAll");
+		Log("One final attempt to test KillAll");
 		Run(Exe, CmdLine, null, ERunOptions.AllowSpew | ERunOptions.NoWaitForExit | ERunOptions.AppMustExist | ERunOptions.NoStdOutRedirect);
 		Thread.Sleep(10000);
 	}
diff --git a/Engine/Source/Programs/AutomationTool/Scripts/UE4BuildUtils.cs b/Engine/Source/Programs/AutomationTool/Scripts/UE4BuildUtils.cs
--- a/Engine/Source/Programs/AutomationTool/Scripts/UE4BuildUtils.cs
+++ b/Engine/Source/Programs/AutomationTool/Scripts/UE4BuildUtils.cs
@@ -51,7 +51,7 @@ public class UE4BuildUtils : CommandUtils
 			throw new AutomationException("Target is required when calling UE4BuildUtils.BuildProduct");
 		}
 
-		LogConsole("Building {0}", Target.TargetName);
+		Log("Building {0}", Target.TargetName);
 
 		if (Command == null)
 		{
diff --git a/Engine/Source/Programs/AutomationTool/Scripts/UnrealSync.Automation.cs b/Engine/Source/Programs/AutomationTool/Scripts/UnrealSync.Automation.cs
--- a/Engine/Source/Programs/AutomationTool/Scripts/UnrealSync.Automation.cs
+++ b/Engine/Source/Programs/AutomationTool/Scripts/UnrealSync.Automation.cs
@@ -44,19 +44,19 @@ namespace AutomationScripts.Automation
 
 			if (QueryType == QueryType.Promoted)
 			{
-				LogConsole("Promoted labels for {0}.", BranchAndGameName);
+				Log("Promoted labels for {0}.", BranchAndGameName);
 
 				Print(GetPromotedLabels(BranchPath, GameName), Ticks);
 			}
 			else if(QueryType == QueryType.Promotable)
 			{
-				LogConsole("Promotable labels for {0}.", BranchAndGameName);
+				Log("Promotable labels for {0}.", BranchAndGameName);
 
 				Print(GetPromotableLabels(BranchPath, GameName), Ticks);
 			}
 			else
 			{
-				LogConsole("All labels for {0}.", BranchPath);
+				Log("All labels for {0}.", BranchPath);
 
 				Print(GetBranchLabels(BranchPath), Ticks);
 			}
@@ -114,7 +114,7 @@ namespace AutomationScripts.Automation
 		{
 			foreach(var Label in Labels.OrderByDescending((Label) => Label.Date))
 			{
-				LogConsole(Label.Name + (Ticks ? (" " + Label.Date.Ticks.ToString()) : ""));
+				Log(Label.Name + (Ticks ? (" " + Label.Date.Ticks.ToString()) : ""));
 			}
 		}
 
