Allow build commands to fail with an exit code, without throwing exceptions and dumping a callstack. GUBP nodes handle call print callstacks near to the point that they're thrown, but still need to do EC/temp storage clean-up work before the command can execute, so we don't want unneccessary spam by the generic exception handler.

diff --git a/Engine/Source/Programs/AutomationTool/AutomationUtils/Automation.cs b/Engine/Source/Programs/AutomationTool/AutomationUtils/Automation.cs
--- a/Engine/Source/Programs/AutomationTool/AutomationUtils/Automation.cs
+++ b/Engine/Source/Programs/AutomationTool/AutomationUtils/Automation.cs
@@ -302,7 +302,7 @@ AutomationTool.exe [-verbose] [-compileonly] [-p4] Command0 [-Arg0 -Arg1 -Arg2 
 		/// Main method.
 		/// </summary>
 		/// <param name="CommandLine">Command line</param>
-		public static void Process(string[] CommandLine)
+		public static ExitCode Process(string[] CommandLine)
 		{
 			// Initial check for local or build machine runs BEFORE we parse the command line (We need this value set
 			// in case something throws the exception while parsing the command line)
@@ -326,7 +326,7 @@ AutomationTool.exe [-verbose] [-compileonly] [-p4] Command0 [-Arg0 -Arg1 -Arg2 
 			if (CommandsToExecute.Count == 0 && GlobalCommandLine.Help)
 			{
 				DisplayHelp();
-				return;
+				return ExitCode.Success;
 			}
 
 			// Disable AutoSDKs if specified on the command line
@@ -359,19 +359,19 @@ AutomationTool.exe [-verbose] [-compileonly] [-p4] Command0 [-Arg0 -Arg1 -Arg2 
 			if (GlobalCommandLine.CompileOnly)
 			{
 				Log.TraceInformation("Compilation successful, exiting (CompileOnly)");
-				return;
+				return ExitCode.Success;
 			}
 
 			if (GlobalCommandLine.List)
 			{
 				ListAvailableCommands(Compiler.Commands);
-				return;
+				return ExitCode.Success;
 			}
 
 			if (GlobalCommandLine.Help)
 			{
 				DisplayHelp(CommandsToExecute, Compiler.Commands);
-				return;
+				return ExitCode.Success;
 			}
 
 			// Enable or disable P4 support
@@ -384,9 +384,7 @@ AutomationTool.exe [-verbose] [-compileonly] [-p4] Command0 [-Arg0 -Arg1 -Arg2 
 			}
 
 			// Find and execute commands.
-			Execute(CommandsToExecute, Compiler.Commands);
-
-			return;
+			return Execute(CommandsToExecute, Compiler.Commands);
 		}
 
 		/// <summary>
@@ -394,7 +392,7 @@ AutomationTool.exe [-verbose] [-compileonly] [-p4] Command0 [-Arg0 -Arg1 -Arg2 
 		/// </summary>
 		/// <param name="CommandsToExecute"></param>
 		/// <param name="Commands"></param>
-		private static void Execute(List<CommandInfo> CommandsToExecute, CaselessDictionary<Type> Commands)
+		private static ExitCode Execute(List<CommandInfo> CommandsToExecute, CaselessDictionary<Type> Commands)
 		{
 			for (int CommandIndex = 0; CommandIndex < CommandsToExecute.Count; ++CommandIndex)
 			{
@@ -405,24 +403,32 @@ AutomationTool.exe [-verbose] [-compileonly] [-p4] Command0 [-Arg0 -Arg1 -Arg2 
 				{
 					throw new AutomationException("Failed to find command {0}", CommandInfo.CommandName);
 				}
-				else
+
+				BuildCommand Command = (BuildCommand)Activator.CreateInstance(CommandType);
+				Command.Params = CommandInfo.Arguments.ToArray();
+				try
+				{
+					ExitCode Result = Command.Execute();
+					if(Result != ExitCode.Success)
+					{
+						return Result;
+					}
+					CommandUtils.Log("BUILD SUCCESSFUL");
+				}
+				finally
 				{
-					BuildCommand Command = (BuildCommand)Activator.CreateInstance(CommandType);
-					Command.Params = CommandInfo.Arguments.ToArray();
-					Command.Execute();
 					// dispose of the class if necessary
+					var CommandDisposable = Command as IDisposable;
+					if (CommandDisposable != null)
 					{
-						var CommandDisposable = Command as IDisposable;
-						if (CommandDisposable != null)
-						{
-							CommandDisposable.Dispose();
-						}
+						CommandDisposable.Dispose();
 					}
-
-					// Make sure there's no directories on the stack.
-					CommandUtils.ClearDirStack();
 				}
+
+				// Make sure there's no directories on the stack.
+				CommandUtils.ClearDirStack();
 			}
+			return ExitCode.Success;
 		}
 
 		#endregion
diff --git a/Engine/Source/Programs/AutomationTool/AutomationUtils/BuildCommand.cs b/Engine/Source/Programs/AutomationTool/AutomationUtils/BuildCommand.cs
--- a/Engine/Source/Programs/AutomationTool/AutomationUtils/BuildCommand.cs
+++ b/Engine/Source/Programs/AutomationTool/AutomationUtils/BuildCommand.cs
@@ -11,31 +11,21 @@ namespace AutomationTool
 	/// </summary>
 	public abstract class BuildCommand : CommandUtils
 	{
-
-		#region Interface
-
 		/// <summary>
 		/// Build command entry point.  Throws AutomationExceptions on failure.
 		/// </summary>
-		public abstract void ExecuteBuild();
-
-		#endregion
+		public virtual void ExecuteBuild()
+		{
+			throw new AutomationException("Either Execute() or ExecuteBuild() should be implemented for {0}", GetType().Name);
+		}
 
 		/// <summary>
 		/// Command entry point.
 		/// </summary>
-		public void Execute()
+		public virtual ExitCode Execute()
 		{
-			try
-			{
-				ExecuteBuild();
-			}
-			catch
-			{
-				LogError("BUILD FAILED");
-				throw;
-			}
-			Log("BUILD SUCCESSFUL");
+			ExecuteBuild();
+			return ExitCode.Success;
 		}
 
 		/// <summary>
@@ -43,14 +33,14 @@ namespace AutomationTool
 		/// </summary>
 		/// <typeparam name="T"></typeparam>
 		/// <param name="ParentCommand"></param>
-		public static void Execute<T>(BuildCommand ParentCommand) where T : BuildCommand, new()
+		public static ExitCode Execute<T>(BuildCommand ParentCommand) where T : BuildCommand, new()
 		{
 			T Command = new T();
 			if (ParentCommand != null)
 			{
 				Command.Params = ParentCommand.Params;
 			}
-			Command.Execute();
+			return Command.Execute();
 		}
 	}
 }
\ No newline at end of file
diff --git a/Engine/Source/Programs/AutomationTool/AutomationUtils/Utils.cs b/Engine/Source/Programs/AutomationTool/AutomationUtils/Utils.cs
--- a/Engine/Source/Programs/AutomationTool/AutomationUtils/Utils.cs
+++ b/Engine/Source/Programs/AutomationTool/AutomationUtils/Utils.cs
@@ -684,12 +684,11 @@ namespace AutomationTool
 		/// </summary>
 		/// <param name="Main"></param>
 		/// <param name="Param"></param>
-		public static void RunSingleInstance(Action<object> Main, object Param)
+		public static ExitCode RunSingleInstance(Func<object, ExitCode> Main, object Param)
 		{
 			if (Environment.GetEnvironmentVariable("uebp_UATMutexNoWait") == "1")
 			{
-				Main(Param);
-                return;
+				return Main(Param);
 			}
 			var bCreatedMutex = false;
             var EntryAssemblyLocation = Assembly.GetEntryAssembly().GetOriginalLocation();
@@ -706,9 +705,11 @@ namespace AutomationTool
 					Log.TraceVerbose("No other instance of {0} is running.", EntryAssemblyLocation);
 				}
 
-				Main(Param);
+				ExitCode Result = Main(Param);
 
 				SingleInstanceMutex.ReleaseMutex();
+
+				return Result;
 			}
 		}
 
diff --git a/Engine/Source/Programs/AutomationTool/GUBP/GUBP.cs b/Engine/Source/Programs/AutomationTool/GUBP/GUBP.cs
--- a/Engine/Source/Programs/AutomationTool/GUBP/GUBP.cs
+++ b/Engine/Source/Programs/AutomationTool/GUBP/GUBP.cs
@@ -148,7 +148,7 @@ public partial class GUBP : BuildCommand
 	/// <summary>
 	/// Main entry point for GUBP
 	/// </summary>
-    public override void ExecuteBuild()
+    public override ExitCode Execute()
     {
 		Log("************************* GUBP");
 
@@ -312,19 +312,23 @@ public partial class GUBP : BuildCommand
 			}
 
 			PrintDetailedChanges(History, P4Env.Changelist);
+			return ExitCode.Success;
 		}
 		else
 		{
 			if(CommanderSetup)
 			{
-				DoCommanderSetup(EC, Graph.BuildNodes, Graph.AggregateNodes, OrderedToDo, TimeIndex, TimeQuantum, bSkipTriggers, bFake, bFakeEC, ExplicitTrigger, UnfinishedTriggers, JobInfo.IsPreflight);
+				return DoCommanderSetup(EC, Graph.BuildNodes, Graph.AggregateNodes, OrderedToDo, TimeIndex, TimeQuantum, bSkipTriggers, bFake, bFakeEC, ExplicitTrigger, UnfinishedTriggers, JobInfo.IsPreflight);
 			}
-			else if(!ParseParam("ListOnly"))
+			else if(ParseParam("ListOnly"))
 			{
-                ExecuteNodes(EC, OrderedToDo, bFake, bFakeEC, bSaveSharedTempStorage, JobInfo);
+				return ExitCode.Success;
+			}
+			else
+			{
+				return ExecuteNodes(EC, OrderedToDo, bFake, bFakeEC, bSaveSharedTempStorage, JobInfo);
 			}
 		}
-		PrintRunTime();
 	}
 
 	/// <summary>
@@ -1298,7 +1302,7 @@ public partial class GUBP : BuildCommand
 		}
 	}
 
-    private void DoCommanderSetup(ElectricCommander EC, IEnumerable<BuildNode> AllNodes, IEnumerable<AggregateNode> AllAggregates, List<BuildNode> OrderedToDo, int TimeIndex, int TimeQuantum, bool bSkipTriggers, bool bFake, bool bFakeEC, TriggerNode ExplicitTrigger, List<TriggerNode> UnfinishedTriggers, bool bPreflightBuild)
+    private ExitCode DoCommanderSetup(ElectricCommander EC, IEnumerable<BuildNode> AllNodes, IEnumerable<AggregateNode> AllAggregates, List<BuildNode> OrderedToDo, int TimeIndex, int TimeQuantum, bool bSkipTriggers, bool bFake, bool bFakeEC, TriggerNode ExplicitTrigger, List<TriggerNode> UnfinishedTriggers, bool bPreflightBuild)
 	{
 		List<BuildNode> SortedNodes = TopologicalSort(new HashSet<BuildNode>(AllNodes), null, SubSort: false, DoNotConsiderCompletion: true);
 		Log("******* {0} GUBP Nodes", SortedNodes.Count);
@@ -1334,9 +1338,10 @@ public partial class GUBP : BuildCommand
 		}
 
 		EC.DoCommanderSetup(AllNodes, AllAggregates, FilteredOrderedToDo, SortedNodes, TimeIndex, TimeQuantum, bSkipTriggers, bFake, bFakeEC, ExplicitTrigger, UnfinishedTriggers);
+		return ExitCode.Success;
 	}
 
-    void ExecuteNodes(ElectricCommander EC, List<BuildNode> OrderedToDo, bool bFake, bool bFakeEC, bool bSaveSharedTempStorage, JobInfo JobInfo)
+    ExitCode ExecuteNodes(ElectricCommander EC, List<BuildNode> OrderedToDo, bool bFake, bool bFakeEC, bool bSaveSharedTempStorage, JobInfo JobInfo)
 	{
         Dictionary<string, BuildNode> BuildProductToNodeMap = new Dictionary<string, BuildNode>();
 		foreach (BuildNode NodeToDo in OrderedToDo)
@@ -1360,135 +1365,51 @@ public partial class GUBP : BuildCommand
 				StorageRootIfAny = CmdEnv.LocalRoot;
 			}
 
-            // this is kinda complicated
-            bool SaveSuccessRecords = (IsBuildMachine || bFakeEC) && // no real reason to make these locally except for fakeEC tests
-                (!(NodeToDo is TriggerNode) || NodeToDo.IsSticky); // trigger nodes are run twice, one to start the new workflow and once when it is actually triggered, we will save reconds for the latter
-
-            Log("***** Running GUBP Node {0} -> {1} : {2}", NodeToDo.Name, GameNameIfAny, TempStorageNodeInfo.GetRelativeDirectory());
             if (NodeToDo.IsComplete)
             {
+				// Just fetch the build products from temp storage
+				Log("***** Retrieving GUBP Node {0} -> {1} : {2}", NodeToDo.Name, GameNameIfAny, TempStorageNodeInfo.GetRelativeDirectory());
 				NodeToDo.RetrieveBuildProducts(GameNameIfAny, StorageRootIfAny, TempStorageNodeInfo);
             }
             else
             {
+				// this is kinda complicated
+				bool SaveSuccessRecords = (IsBuildMachine || bFakeEC) && // no real reason to make these locally except for fakeEC tests
+					(!(NodeToDo is TriggerNode) || NodeToDo.IsSticky); // trigger nodes are run twice, one to start the new workflow and once when it is actually triggered, we will save reconds for the latter
+
+                // Record that the node has started. We save our status to a new temp storage location specifically named with a suffix so we can find it later.
                 if (SaveSuccessRecords) 
                 {
-                    // We save our status to a new temp storage location specifically named with a suffix so we can find it later.
                     EC.SaveStatus(new TempStorageNodeInfo(JobInfo, NodeToDo.Name + StartedTempStorageSuffix), bSaveSharedTempStorage, GameNameIfAny);
                 }
-				double BuildDuration = 0.0;
-                try
-                {
-                    if (bFake)
-                    {
-                        Log("***** FAKE!! Building GUBP Node {0} for {1}", NodeToDo.Name, TempStorageNodeInfo.GetRelativeDirectory());
-                        NodeToDo.DoFakeBuild();
-                    }
-                    else
-                    {
-                        Log("***** Building GUBP Node {0} for {1}", NodeToDo.Name, TempStorageNodeInfo.GetRelativeDirectory());
-						DateTime StartTime = DateTime.UtcNow;
-						using(TelemetryStopwatch DoBuildStopwatch = new TelemetryStopwatch("DoBuild.{0}", NodeToDo.Name))
-						{
-							NodeToDo.DoBuild();
-						}
-						BuildDuration = (DateTime.UtcNow - StartTime).TotalMilliseconds / 1000;
-                    }
-
-					NodeToDo.ArchiveBuildProducts(GameNameIfAny, StorageRootIfAny, TempStorageNodeInfo, !bSaveSharedTempStorage);
-                }
-                catch (Exception Ex)
-                {
-                    //@todo: This is a dup of non-exception code. Consolidate this!!
-					NodeHistory History = null;
-
-                    if (SaveSuccessRecords)
-                    {
-						using(TelemetryStopwatch UpdateNodeHistoryStopwatch = new TelemetryStopwatch("UpdateNodeHistory"))
-						{
-                            History = FindNodeHistory(NodeToDo, JobInfo);
-						}
-						using(TelemetryStopwatch SaveNodeStatusStopwatch = new TelemetryStopwatch("SaveNodeStatus"))
-						{
-                            // We save our status to a new temp storage location specifically named with a suffix so we can find it later.
-                            EC.SaveStatus(new TempStorageNodeInfo(JobInfo, NodeToDo.Name + FailedTempStorageSuffix), bSaveSharedTempStorage, GameNameIfAny, ParseParamValue("MyJobStepId"));
-						}
-						using(TelemetryStopwatch UpdateECPropsStopwatch = new TelemetryStopwatch("UpdateECProps"))
-						{
-							EC.UpdateECProps(NodeToDo);
-						}
-                        
-						if (IsBuildMachine)
-						{
-							using(TelemetryStopwatch GetFailEmailsStopwatch = new TelemetryStopwatch("GetFailEmails"))
-							{
-                                GetFailureEmails(EC, NodeToDo, History, JobInfo);
-							}
-						}
-						EC.UpdateECBuildTime(NodeToDo, BuildDuration);
-                    }
-
-					Log("{0}", Ex);
-
-
-                    if (History != null)
-                    {
-						Log("Changes since last green *********************************");
-						Log("");
-						Log("");
-						Log("");
-                        PrintDetailedChanges(History, P4Env.Changelist);
-						Log("End changes since last green");
-                    }
-
-                    string FailInfo = "";
-                    FailInfo += "********************************* Main log file";
-                    FailInfo += Environment.NewLine + Environment.NewLine;
-                    FailInfo += LogUtils.GetLogTail();
-                    FailInfo += Environment.NewLine + Environment.NewLine + Environment.NewLine;
-
-
-
-                    string OtherLog = "See logfile for details: '";
-                    if (FailInfo.Contains(OtherLog))
-                    {
-                        string LogFile = FailInfo.Substring(FailInfo.IndexOf(OtherLog) + OtherLog.Length);
-                        if (LogFile.Contains("'"))
-                        {
-                            LogFile = CombinePaths(CmdEnv.LogFolder, LogFile.Substring(0, LogFile.IndexOf("'")));
-                            if (FileExists_NoExceptions(LogFile))
-                            {
-                                FailInfo += "********************************* Sub log file " + LogFile;
-                                FailInfo += Environment.NewLine + Environment.NewLine;
 
-                                FailInfo += LogUtils.GetLogTail(LogFile);
-                                FailInfo += Environment.NewLine + Environment.NewLine + Environment.NewLine;
-                            }
-                        }
-                    }
+				// Execute the node
+				DateTime StartTime = DateTime.UtcNow;
+				bool bResult = ExecuteNode(NodeToDo, bFake);
+				TimeSpan ElapsedTime = DateTime.UtcNow - StartTime;
 
-                    string Filename = CombinePaths(CmdEnv.LogFolder, "LogTailsAndChanges.log");
-                    WriteAllText(Filename, FailInfo);
+				// Archive the build products if the node succeeded
+				if(bResult)
+				{
+					NodeToDo.ArchiveBuildProducts(GameNameIfAny, StorageRootIfAny, TempStorageNodeInfo, !bSaveSharedTempStorage);
+				}
 
-                    throw(Ex);
-                }
-                if (SaveSuccessRecords) 
+				// Record that the node has finished
+				NodeHistory History = null;
+                if (SaveSuccessRecords)
                 {
-					NodeHistory History = null;
 					using(TelemetryStopwatch UpdateNodeHistoryStopwatch = new TelemetryStopwatch("UpdateNodeHistory"))
 					{
                         History = FindNodeHistory(NodeToDo, JobInfo);
 					}
 					using(TelemetryStopwatch SaveNodeStatusStopwatch = new TelemetryStopwatch("SaveNodeStatus"))
 					{
-                        // We save our status to a new temp storage location specifically named with a suffix so we can find it later.
-                        EC.SaveStatus(new TempStorageNodeInfo(JobInfo, NodeToDo.Name + SucceededTempStorageSuffix), bSaveSharedTempStorage, GameNameIfAny);
+                        EC.SaveStatus(new TempStorageNodeInfo(JobInfo, NodeToDo.Name + (bResult? SucceededTempStorageSuffix : FailedTempStorageSuffix)), bSaveSharedTempStorage, GameNameIfAny, ParseParamValue("MyJobStepId"));
 					}
 					using(TelemetryStopwatch UpdateECPropsStopwatch = new TelemetryStopwatch("UpdateECProps"))
 					{
 						EC.UpdateECProps(NodeToDo);
 					}
-                    
 					if (IsBuildMachine)
 					{
 						using(TelemetryStopwatch GetFailEmailsStopwatch = new TelemetryStopwatch("GetFailEmails"))
@@ -1496,7 +1417,24 @@ public partial class GUBP : BuildCommand
                             GetFailureEmails(EC, NodeToDo, History, JobInfo);
 						}
 					}
-					EC.UpdateECBuildTime(NodeToDo, BuildDuration);
+					EC.UpdateECBuildTime(NodeToDo, ElapsedTime.TotalSeconds);
+                }
+
+				// If it failed, print the diagnostic info and quit
+				if(!bResult)
+				{
+                    if (History != null)
+                    {
+						Log("Changes since last green *********************************");
+						Log("");
+						Log("");
+						Log("");
+                        PrintDetailedChanges(History, P4Env.Changelist);
+						Log("End changes since last green");
+                    }
+
+					WriteFailureLog(CombinePaths(CmdEnv.LogFolder, "LogTailsAndChanges.log"));
+					return ExitCode.Error_Unknown;
                 }
             }
             foreach (string Product in NodeToDo.BuildProducts)
@@ -1507,7 +1445,57 @@ public partial class GUBP : BuildCommand
                 }
                 BuildProductToNodeMap.Add(Product, NodeToDo);
             }
-        }        
-        PrintRunTime();
-    }   
+        }
+		return ExitCode.Success;
+    }
+
+	bool ExecuteNode(BuildNode NodeToDo, bool bFake)
+	{
+		try
+		{
+			if (bFake)
+			{
+				NodeToDo.DoFakeBuild();
+			}
+			else
+			{
+				NodeToDo.DoBuild();
+			}
+			return true;
+		}
+		catch(Exception Ex)
+		{
+			Log("{0}", Ex.ToString());
+			return false;
+		}
+	}
+
+	static void WriteFailureLog(string FileName)
+	{
+		string FailInfo = "";
+		FailInfo += "********************************* Main log file";
+		FailInfo += Environment.NewLine + Environment.NewLine;
+		FailInfo += LogUtils.GetLogTail();
+		FailInfo += Environment.NewLine + Environment.NewLine + Environment.NewLine;
+
+		string OtherLog = "See logfile for details: '";
+		if (FailInfo.Contains(OtherLog))
+		{
+			string LogFile = FailInfo.Substring(FailInfo.IndexOf(OtherLog) + OtherLog.Length);
+			if (LogFile.Contains("'"))
+			{
+				LogFile = CombinePaths(CmdEnv.LogFolder, LogFile.Substring(0, LogFile.IndexOf("'")));
+				if (FileExists_NoExceptions(LogFile))
+				{
+					FailInfo += "********************************* Sub log file " + LogFile;
+					FailInfo += Environment.NewLine + Environment.NewLine;
+
+					FailInfo += LogUtils.GetLogTail(LogFile);
+					FailInfo += Environment.NewLine + Environment.NewLine + Environment.NewLine;
+				}
+			}
+		}
+
+		WriteAllText(FileName, FailInfo);
+	}
 }
diff --git a/Engine/Source/Programs/AutomationTool/Program.cs b/Engine/Source/Programs/AutomationTool/Program.cs
--- a/Engine/Source/Programs/AutomationTool/Program.cs
+++ b/Engine/Source/Programs/AutomationTool/Program.cs
@@ -51,7 +51,7 @@ namespace AutomationTool
 				Log.TraceVerbose("{0} ver. {1}", Version.ProductName, Version.ProductVersion);
 
 				// Don't allow simultaneous execution of AT (in the same branch)
-				InternalUtils.RunSingleInstance(MainProc, CommandLine);
+				ReturnCode = InternalUtils.RunSingleInstance(MainProc, CommandLine);
 
 			}
 			catch (AutomationException Ex)
@@ -126,10 +126,11 @@ namespace AutomationTool
 			Trace.Close();
 		}
 
-		static void MainProc(object Param)
+		static ExitCode MainProc(object Param)
 		{
-			Automation.Process((string[])Param);
+			ExitCode Result = Automation.Process((string[])Param);
 			ShouldKillProcesses = Automation.ShouldKillProcesses;
+			return Result;
 		}
 
 		static bool ShouldKillProcesses = true;
