Correct case for Mac .dSYM in AutomationTool as well.

diff --git a/Engine/Source/Programs/AutomationTool/Mac/MacPlatform.Automation.cs b/Engine/Source/Programs/AutomationTool/Mac/MacPlatform.Automation.cs
--- a/Engine/Source/Programs/AutomationTool/Mac/MacPlatform.Automation.cs
+++ b/Engine/Source/Programs/AutomationTool/Mac/MacPlatform.Automation.cs
@@ -429,7 +429,7 @@ public class MacPlatform : Platform
 	}
 	public override List<string> GetDebugFileExtentions()
 	{
-		return new List<string> { ".dsym" };
+		return new List<string> { ".dSYM" };
 	}
 	public override bool CanHostPlatform(UnrealTargetPlatform Platform)
 	{
