Clean up output for detailed changes in GUBP.

diff --git a/Engine/Source/Programs/AutomationTool/GUBP/GUBP.cs b/Engine/Source/Programs/AutomationTool/GUBP/GUBP.cs
--- a/Engine/Source/Programs/AutomationTool/GUBP/GUBP.cs
+++ b/Engine/Source/Programs/AutomationTool/GUBP/GUBP.cs
@@ -432,7 +432,8 @@ public partial class GUBP : BuildCommand
 	/// <param name="History">History for this node</param>
     static void PrintDetailedChanges(NodeHistory History, int CurrentCL)
     {
-        DateTime StartTime = DateTime.UtcNow;
+		Log("");
+		Log("**** Changes since last green *********************************");
 
 		// Find all the changelists that we're interested in
 		SortedSet<int> BuildChanges = new SortedSet<int>();
@@ -477,7 +478,8 @@ public partial class GUBP : BuildCommand
 			LastCL = BuildCL;
 		}
 
-        Log("Took {0:0.0}s to get P4 history", (DateTime.UtcNow - StartTime).TotalSeconds);
+		Log("");
+		Log("***************************************************************");
     }
 
     void PrintNodes(GUBP bp, List<BuildNode> Nodes, IEnumerable<AggregateNode> Aggregates, List<TriggerNode> UnfinishedTriggers, int TimeQuantum)
@@ -1425,12 +1427,7 @@ public partial class GUBP : BuildCommand
 				{
                     if (History != null)
                     {
-						Log("Changes since last green *********************************");
-						Log("");
-						Log("");
-						Log("");
                         PrintDetailedChanges(History, P4Env.Changelist);
-						Log("End changes since last green");
                     }
 
 					WriteFailureLog(CombinePaths(CmdEnv.LogFolder, "LogTailsAndChanges.log"));
