Tweak the list of recent changes in the build failure report. Now shows in descending order (like Perforce).

diff --git a/Engine/Source/Programs/AutomationTool/GUBP/GUBP.cs b/Engine/Source/Programs/AutomationTool/GUBP/GUBP.cs
--- a/Engine/Source/Programs/AutomationTool/GUBP/GUBP.cs
+++ b/Engine/Source/Programs/AutomationTool/GUBP/GUBP.cs
@@ -432,8 +432,8 @@ public partial class GUBP : BuildCommand
 	/// <param name="History">History for this node</param>
     static void PrintDetailedChanges(NodeHistory History, int CurrentCL)
     {
+		Log("**** Changes since last succeeded *****************************");
 		Log("");
-		Log("**** Changes since last green *********************************");
 
 		// Find all the changelists that we're interested in
 		SortedSet<int> BuildChanges = new SortedSet<int>();
@@ -442,10 +442,11 @@ public partial class GUBP : BuildCommand
 
 		// Find all the changelists that we're interested in
         List<P4Connection.ChangeRecord> ChangeRecords = GetSourceChangeRecords(BuildChanges.First(), BuildChanges.Last());
+		ChangeRecords.Reverse();
 
 		// Print all the changes in the set
 		int LastCL = BuildChanges.First();
-		foreach(int BuildCL in BuildChanges)
+		foreach(int BuildCL in BuildChanges.OrderByDescending(x => x))
 		{
 			// Show all the changes in this range
 			foreach(P4Connection.ChangeRecord Record in ChangeRecords.Where(x => x.CL > LastCL && x.CL <= BuildCL))
@@ -458,11 +459,11 @@ public partial class GUBP : BuildCommand
 			string BuildStatus;
 			if(BuildCL == CurrentCL)
 			{
-				BuildStatus = "this sync";
+				BuildStatus = "<<<< This change";
 			}
             else if (History.AllSucceeded.Contains(BuildCL))
             {
-				BuildStatus = "built";
+				BuildStatus = "SUCCEEDED";
             }
             else if (History.AllFailed.Contains(BuildCL))
             {
