Another attempt at fixing node history dump.

diff --git a/Engine/Source/Programs/AutomationTool/GUBP/GUBP.cs b/Engine/Source/Programs/AutomationTool/GUBP/GUBP.cs
--- a/Engine/Source/Programs/AutomationTool/GUBP/GUBP.cs
+++ b/Engine/Source/Programs/AutomationTool/GUBP/GUBP.cs
@@ -445,21 +445,16 @@ public partial class GUBP : BuildCommand
 		ChangeRecords.Reverse();
 
 		// Print all the changes in the set
-		int LastCL = BuildChanges.First();
-		foreach(int BuildCL in BuildChanges.OrderByDescending(x => x))
+		int[] BuildChangesArray = BuildChanges.ToArray();
+		for(int Idx = BuildChangesArray.Length - 1; Idx >= 0; Idx--)
 		{
-			// Show all the changes in this range
-			foreach(P4Connection.ChangeRecord Record in ChangeRecords.Where(x => x.CL > LastCL && x.CL <= BuildCL))
-			{
-				string[] Lines = Record.Summary.Split(new char[]{ '\r', '\n' }, StringSplitOptions.RemoveEmptyEntries);
-				Log("             {0} {1} {2}", Record.CL, Record.UserEmail, (Lines.Length > 0)? Lines[0] : "");
-			}
+			int BuildCL = BuildChangesArray[Idx];
 
 			// Show the status of this build
 			string BuildStatus;
 			if(BuildCL == CurrentCL)
 			{
-				BuildStatus = "<<<< This change";
+				BuildStatus = "<<<< THIS CHANGE";
 			}
             else if (History.AllSucceeded.Contains(BuildCL))
             {
@@ -471,12 +466,20 @@ public partial class GUBP : BuildCommand
             }
 			else
 			{
-				BuildStatus = "running";
+				BuildStatus = "still running";
 			}
 			Log(" {0} {1} {2}", (BuildCL == CurrentCL)? ">>>>" : "    ", BuildCL, BuildStatus.ToString());
 
-			// Update the last CL now that we've output this one
-			LastCL = BuildCL;
+			// Show all the changes in this range
+			if(Idx > 0)
+			{
+				int PrevCL = BuildChangesArray[Idx - 1];
+				foreach(P4Connection.ChangeRecord Record in ChangeRecords.Where(x => x.CL > PrevCL && x.CL <= BuildCL))
+				{
+					string[] Lines = Record.Summary.Split(new char[]{ '\r', '\n' }, StringSplitOptions.RemoveEmptyEntries);
+					Log("             {0} {1} {2}", Record.CL, Record.UserEmail, (Lines.Length > 0)? Lines[0] : "");
+				}
+			}
 		}
 
 		Log("");
