Another minor tweak to the output of failed changes in GUBP.

diff --git a/Engine/Source/Programs/AutomationTool/GUBP/GUBP.cs b/Engine/Source/Programs/AutomationTool/GUBP/GUBP.cs
--- a/Engine/Source/Programs/AutomationTool/GUBP/GUBP.cs
+++ b/Engine/Source/Programs/AutomationTool/GUBP/GUBP.cs
@@ -454,15 +454,15 @@ public partial class GUBP : BuildCommand
 			string BuildStatus;
 			if(BuildCL == CurrentCL)
 			{
-				BuildStatus = "<<<< THIS CHANGE";
+				BuildStatus = "THIS CHANGE";
 			}
-            else if (History.AllSucceeded.Contains(BuildCL))
-            {
-				BuildStatus = "SUCCEEDED";
-            }
             else if (History.AllFailed.Contains(BuildCL))
             {
 				BuildStatus = "FAILED";
+            }
+            else if (History.AllSucceeded.Contains(BuildCL))
+            {
+				BuildStatus = "SUCCEEDED";
             }
 			else
 			{
