[NO NON-PII COMMIT MESSAGE]

diff --git a/Engine/Source/Programs/AutomationTool/GUBP/TempStorage.cs b/Engine/Source/Programs/AutomationTool/GUBP/TempStorage.cs
--- a/Engine/Source/Programs/AutomationTool/GUBP/TempStorage.cs
+++ b/Engine/Source/Programs/AutomationTool/GUBP/TempStorage.cs
@@ -10,6 +10,7 @@ using System.Threading.Tasks;
 using System.Xml.Linq;
 using UnrealBuildTool;
 using Ionic.Zip;
+using System.IO.Compression;
 
 namespace AutomationTool
 {
@@ -868,37 +869,77 @@ namespace AutomationTool
             // as ZipFile cannot unzip files in parallel from a single archive.
             // We can safely assume the build system will not be doing more important things at the same time, so we simply use all our logical cores,
             // which has shown to be optimal via profiling, and limits the number of resulting zip files to the number of logical cores.
-            var Threads = (
-                from CoreNum in Enumerable.Range(0, bZipInParallel ? Environment.ProcessorCount : 1)
-                let ZipFileName = Path.Combine(StagingDir ?? OutputDir, string.Format("{0}{1}.zip", ZipBasename, bZipInParallel ? "-" + CoreNum.ToString("00") : ""))
-                select new Thread(() =>
-                {
-                    // Create one zip per thread using the given basename
-                    using (var ZipArchive = new ZipFile(ZipFileName) { CompressionLevel = Ionic.Zlib.CompressionLevel.BestSpeed })
+            // 
+            // Sadly, mono implemention of System.IO.Compression is really poor (as of 2015/Aug), causing OOM when parallel zipping a large set of files.
+            // However, Ionic is MUCH slower than .NET's native implementation (2x+ slower in our build farm), so we stick to the faster solution on PC.
+            // The code duplication in the threadprocs is unfortunate here, and hopefully we can settle on .NET's implementation on both platforms eventually.
+            List<Thread> ZipThreads;
+
+            if (Utils.IsRunningOnMono)
+            {
+                ZipThreads = (
+                    from CoreNum in Enumerable.Range(0, bZipInParallel ? Environment.ProcessorCount : 1)
+                    let ZipFileName = Path.Combine(StagingDir ?? OutputDir, string.Format("{0}{1}.zip", ZipBasename, bZipInParallel ? "-" + CoreNum.ToString("00") : ""))
+                    select new Thread(() =>
                     {
+                        // Create one zip per thread using the given basename
+                        using (var ZipArchive = new Ionic.Zip.ZipFile(ZipFileName) { CompressionLevel = Ionic.Zlib.CompressionLevel.BestSpeed })
+                        {
 
-                        // pull from the queue until we are out of files.
-                        string File;
-                        while (FilesToZip.TryDequeue(out File))
+                            // pull from the queue until we are out of files.
+                            string File;
+                            while (FilesToZip.TryDequeue(out File))
+                            {
+                                // use fastest compression. In our best case we are CPU bound, so this is a good tradeoff,
+                                // cutting overall time by 2/3 while only modestly increasing the compression ratio (22.7% -> 23.8% for RootEditor PDBs).
+                                // This is in cases of a super hot cache, so the operation was largely CPU bound.
+                                ZipArchive.AddFile(File, Path.GetDirectoryName(CommandUtils.StripBaseDirectory(File, RootDir)));
+                            }
+                            ZipArchive.Save();
+                        }
+                        Interlocked.Add(ref ZipFilesTotalSize, new FileInfo(ZipFileName).Length);
+                        // if we are using a staging dir, copy to the final location and delete the staged copy.
+                        if (StagingDir != null)
                         {
-                            // use fastest compression. In our best case we are CPU bound, so this is a good tradeoff,
-                            // cutting overall time by 2/3 while only modestly increasing the compression ratio (22.7% -> 23.8% for RootEditor PDBs).
-                            // This is in cases of a super hot cache, so the operation was largely CPU bound.
-                            ZipArchive.AddFile(File, Path.GetDirectoryName(CommandUtils.StripBaseDirectory(File, RootDir)));
+                            CommandUtils.CopyFile(ZipFileName, CommandUtils.MakeRerootedFilePath(ZipFileName, StagingDir, OutputDir));
+                            CommandUtils.DeleteFile(true, ZipFileName);
                         }
-                        ZipArchive.Save();
-                    }
-                    Interlocked.Add(ref ZipFilesTotalSize, new FileInfo(ZipFileName).Length);
-                    // if we are using a staging dir, copy to the final location and delete the staged copy.
-                    if (StagingDir != null)
+                    })).ToList();
+            }
+            else
+            {
+                ZipThreads = (
+                    from CoreNum in Enumerable.Range(0, bZipInParallel ? Environment.ProcessorCount : 1)
+                    let ZipFileName = Path.Combine(StagingDir ?? OutputDir, string.Format("{0}{1}.zip", ZipBasename, bZipInParallel ? "-" + CoreNum.ToString("00") : ""))
+                    select new Thread(() =>
                     {
-                        CommandUtils.CopyFile(ZipFileName, CommandUtils.MakeRerootedFilePath(ZipFileName, StagingDir, OutputDir));
-                        CommandUtils.DeleteFile(true, ZipFileName);
-                    }
-                })).ToList();
+                        // Create one zip per thread using the given basename
+                        using (var ZipArchive = System.IO.Compression.ZipFile.Open(ZipFileName, System.IO.Compression.ZipArchiveMode.Create))
+                        {
 
-            Threads.ForEach(thread => thread.Start());
-            Threads.ForEach(thread => thread.Join());
+                            // pull from the queue until we are out of files.
+                            string File;
+                            while (FilesToZip.TryDequeue(out File))
+                            {
+                                // use fastest compression. In our best case we are CPU bound, so this is a good tradeoff,
+                                // cutting overall time by 2/3 while only modestly increasing the compression ratio (22.7% -> 23.8% for RootEditor PDBs).
+                                // This is in cases of a super hot cache, so the operation was largely CPU bound.
+                                // Also, sadly, mono appears to have a bug where nothing you can do will properly set the LastWriteTime on the created entry,
+                                // so we have to ignore timestamps on files extracted from a zip, since it may have been created on a Mac.
+                                ZipArchive.CreateEntryFromFile(File, CommandUtils.StripBaseDirectory(File, RootDir), System.IO.Compression.CompressionLevel.Fastest);
+                            }
+                        }
+                        Interlocked.Add(ref ZipFilesTotalSize, new FileInfo(ZipFileName).Length);
+                        // if we are using a staging dir, copy to the final location and delete the staged copy.
+                        if (StagingDir != null)
+                        {
+                            CommandUtils.CopyFile(ZipFileName, CommandUtils.MakeRerootedFilePath(ZipFileName, StagingDir, OutputDir));
+                            CommandUtils.DeleteFile(true, ZipFileName);
+                        }
+                    })).ToList();
+            }
+            ZipThreads.ForEach(thread => thread.Start());
+            ZipThreads.ForEach(thread => thread.Join());
             return new ParallelZipResult(ZipTimer.ElapsedTime, ZipFilesTotalSize);
         }
 
@@ -916,15 +957,56 @@ namespace AutomationTool
         {
             var UnzipTimer = DateTimeStopwatch.Start();
             long ZipFilesTotalSize = 0L;
-            Parallel.ForEach(Directory.EnumerateFiles(FolderWithZipFiles, ZipBasename + "*.zip").ToList(),
-                (ZipFilename) =>
-                {
-                    Interlocked.Add(ref ZipFilesTotalSize, new FileInfo(ZipFilename).Length);
-                    using (var ZipArchive = ZipFile.Read(ZipFilename))
+            // Sadly, mono implemention of System.IO.Compression is really poor (as of 2015/Aug), causing OOM when parallel zipping a large set of files.
+            // However, Ionic is MUCH slower than .NET's native implementation (2x+ slower in our build farm), so we stick to the faster solution on PC.
+            // The code duplication in the threadprocs is unfortunate here, and hopefully we can settle on .NET's implementation on both platforms eventually.
+            if (Utils.IsRunningOnMono)
+            {
+                Parallel.ForEach(Directory.EnumerateFiles(FolderWithZipFiles, ZipBasename + "*.zip").ToList(),
+                    (ZipFilename) =>
                     {
-                        ZipArchive.ExtractAll(RootDir, ExtractExistingFileAction.OverwriteSilently);
-                    }
-                });
+                        Interlocked.Add(ref ZipFilesTotalSize, new FileInfo(ZipFilename).Length);
+                        using (var ZipArchive = Ionic.Zip.ZipFile.Read(ZipFilename))
+                        {
+                            ZipArchive.ExtractAll(RootDir, Ionic.Zip.ExtractExistingFileAction.OverwriteSilently);
+                        }
+                    });
+            }
+            else
+            {
+                Parallel.ForEach(Directory.EnumerateFiles(FolderWithZipFiles, ZipBasename + "*.zip").ToList(),
+                    (ZipFilename) =>
+                    {
+                        Interlocked.Add(ref ZipFilesTotalSize, new FileInfo(ZipFilename).Length);
+                        // unzip the files manually instead of caling ZipFile.ExtractToDirectory() because we need to overwrite readonly files. Because of this, creating the directories is up to us as well.
+                        using (var ZipArchive = System.IO.Compression.ZipFile.OpenRead(ZipFilename))
+                        {
+                            foreach (var Entry in ZipArchive.Entries)
+                            {
+                                // Use CommandUtils.CombinePaths to ensure directory separators get converted correctly. On mono on *nix, if the path has backslashes it will not convert it.
+                                var ExtractedFilename = CommandUtils.CombinePaths(RootDir, Entry.FullName);
+                                // Zips can contain empty dirs. Ours usually don't have them, but we should support it.
+                                if (Path.GetFileName(ExtractedFilename).Length == 0)
+                                {
+                                    Directory.CreateDirectory(ExtractedFilename);
+                                }
+                                else
+                                {
+                                    // We must delete any existing file, even if it's readonly. .Net does not do this by default.
+                                    if (File.Exists(ExtractedFilename))
+                                    {
+                                        InternalUtils.SafeDeleteFile(ExtractedFilename, true);
+                                    }
+                                    else
+                                    {
+                                        Directory.CreateDirectory(Path.GetDirectoryName(ExtractedFilename));
+                                    }
+                                    Entry.ExtractToFile(ExtractedFilename, true);
+                                }
+                            }
+                        }
+                    });
+            }
             return new ParallelZipResult(UnzipTimer.ElapsedTime, ZipFilesTotalSize);
         }
 
