- Added "-UploadSymbols" which can be done while packaging (well, you don't actually need to package, but it's usually done alongisde packaging)

diff --git a/Engine/Source/Programs/AutomationTool/AutomationUtils/Platform.cs b/Engine/Source/Programs/AutomationTool/AutomationUtils/Platform.cs
--- a/Engine/Source/Programs/AutomationTool/AutomationUtils/Platform.cs
+++ b/Engine/Source/Programs/AutomationTool/AutomationUtils/Platform.cs
@@ -184,6 +184,11 @@ namespace AutomationTool
 			// do nothing in the default case
 		}
 
+		public virtual void UploadSymbols(ProjectParams Params, DeploymentContext SC)
+		{
+			LogWarning("{0} does not implement UploadSymbols...", PlatformType);
+		}
+
 		/// <summary>
 		/// Get the platform-specific name for the executable (with out the file extension)
 		/// </summary>
diff --git a/Engine/Source/Programs/AutomationTool/AutomationUtils/ProjectParams.cs b/Engine/Source/Programs/AutomationTool/AutomationUtils/ProjectParams.cs
--- a/Engine/Source/Programs/AutomationTool/AutomationUtils/ProjectParams.cs
+++ b/Engine/Source/Programs/AutomationTool/AutomationUtils/ProjectParams.cs
@@ -298,6 +298,7 @@ namespace AutomationTool
             this.RunTimeoutSeconds = InParams.RunTimeoutSeconds;
 			this.bIsCodeBasedProject = InParams.bIsCodeBasedProject;
 			this.bCodeSign = InParams.bCodeSign;
+			this.UploadSymbols = InParams.UploadSymbols;
 		}
 
 		/// <summary>
@@ -403,7 +404,8 @@ namespace AutomationTool
             bool? IterativeDeploy = null,
 			bool? FastCook = null,
 			bool? IgnoreCookErrors = null,
-			bool? CodeSign = null
+			bool? CodeSign = null,
+			bool? UploadSymbols = null
 			)
 		{
 			//
@@ -560,6 +562,7 @@ namespace AutomationTool
 			this.IterativeDeploy = GetParamValueIfNotSpecified(Command, IterativeDeploy, this.IterativeDeploy, new string[] {"iterativedeploy", "iterate" } );
 			this.FastCook = GetParamValueIfNotSpecified(Command, FastCook, this.FastCook, "FastCook");
 			this.IgnoreCookErrors = GetParamValueIfNotSpecified(Command, IgnoreCookErrors, this.IgnoreCookErrors, "IgnoreCookErrors");
+			this.UploadSymbols = GetParamValueIfNotSpecified(Command, UploadSymbols, this.UploadSymbols, "uploadsymbols");
 			this.Device = ParseParamValueIfNotSpecified(Command, Device, "device", String.Empty).Trim(new char[] { '\"' });
 
 			// strip the platform prefix the specified device.
@@ -1403,6 +1406,9 @@ namespace AutomationTool
 		[Help("SpecifiedArchitecture", "Determine a specific Minimum OS")]
 		public string SpecifiedArchitecture;
 
+		[Help("UploadSymbols", "upload symbols while packaging")]
+		public bool UploadSymbols { get; set; }
+
 		#endregion
 
 		#region Deploy
@@ -2140,6 +2146,7 @@ namespace AutomationTool
 				CommandUtils.LogLog("bUsesCEF3={0}", bUsesCEF3);
 				CommandUtils.LogLog("bUsesSlate={0}", bUsesSlate);
                 CommandUtils.LogLog("bDebugBuildsActuallyUseDebugCRT={0}", bDebugBuildsActuallyUseDebugCRT);
+				CommandUtils.LogLog("UploadSymbols={0}", UploadSymbols);
 				CommandUtils.LogLog("Project Params **************");
 			}
 			bLogged = true;
diff --git a/Engine/Source/Programs/AutomationTool/Scripts/PackageCommand.Automation.cs b/Engine/Source/Programs/AutomationTool/Scripts/PackageCommand.Automation.cs
--- a/Engine/Source/Programs/AutomationTool/Scripts/PackageCommand.Automation.cs
+++ b/Engine/Source/Programs/AutomationTool/Scripts/PackageCommand.Automation.cs
@@ -23,6 +23,17 @@ public partial class Project : CommandUtils
 		{
 			DeployContextList.AddRange(CreateDeploymentContext(Params, true, false));
 		}
+
+		// before we package up the build, allow a symbol upload (this isn't actually tied to packaging, 
+		// but logically it's a package-time thing)
+		foreach (var SC in DeployContextList)
+		{
+			if (Params.UploadSymbols)
+			{
+				SC.StageTargetPlatform.UploadSymbols(Params, SC);
+			}
+		}
+
 		if (DeployContextList.Count > 0 && !Params.SkipStage)
 		{
 			Log("********** PACKAGE COMMAND STARTED **********");
