Fix TempStorage unzip failing on Mac on zips created on PC. It was failing to interpret the backslashes that I was using on PC to indicate directories.

diff --git a/Engine/Source/Programs/AutomationTool/GUBP/TempStorage.cs b/Engine/Source/Programs/AutomationTool/GUBP/TempStorage.cs
--- a/Engine/Source/Programs/AutomationTool/GUBP/TempStorage.cs
+++ b/Engine/Source/Programs/AutomationTool/GUBP/TempStorage.cs
@@ -893,7 +893,7 @@ namespace AutomationTool
                                 // use fastest compression. In our best case we are CPU bound, so this is a good tradeoff,
                                 // cutting overall time by 2/3 while only modestly increasing the compression ratio (22.7% -> 23.8% for RootEditor PDBs).
                                 // This is in cases of a super hot cache, so the operation was largely CPU bound.
-                                ZipArchive.AddFile(File, Path.GetDirectoryName(CommandUtils.StripBaseDirectory(File, RootDir)));
+                                ZipArchive.AddFile(File, CommandUtils.ConvertSeparators(PathSeparator.Slash, Path.GetDirectoryName(CommandUtils.StripBaseDirectory(File, RootDir))));
                             }
                             ZipArchive.Save();
                         }
@@ -926,7 +926,7 @@ namespace AutomationTool
                                 // This is in cases of a super hot cache, so the operation was largely CPU bound.
                                 // Also, sadly, mono appears to have a bug where nothing you can do will properly set the LastWriteTime on the created entry,
                                 // so we have to ignore timestamps on files extracted from a zip, since it may have been created on a Mac.
-                                ZipArchive.CreateEntryFromFile(File, CommandUtils.StripBaseDirectory(File, RootDir), System.IO.Compression.CompressionLevel.Fastest);
+                                ZipArchive.CreateEntryFromFile(File, CommandUtils.ConvertSeparators(PathSeparator.Slash, CommandUtils.StripBaseDirectory(File, RootDir)), System.IO.Compression.CompressionLevel.Fastest);
                             }
                         }
                         Interlocked.Add(ref ZipFilesTotalSize, new FileInfo(ZipFileName).Length);
