Clean temp storage early in GUBP startup. Must happen BEFORE caching completion state.

diff --git a/Engine/Source/Programs/AutomationTool/GUBP/GUBP.cs b/Engine/Source/Programs/AutomationTool/GUBP/GUBP.cs
--- a/Engine/Source/Programs/AutomationTool/GUBP/GUBP.cs
+++ b/Engine/Source/Programs/AutomationTool/GUBP/GUBP.cs
@@ -166,7 +166,11 @@ public partial class GUBP : BuildCommand
 			HostPlatforms.Add(UnrealTargetPlatform.Linux);
 		}
 
-	    bool bCleanLocalTempStorage = ParseParam("CleanLocal");
+	    if(ParseParam("CleanLocal"))
+        {
+            TempStorage.DeleteLocalTempStorage();
+        }
+
         bool bSkipTriggers = ParseParam("SkipTriggers");
         bool bFake = ParseParam("fake");
         bool bFakeEC = ParseParam("FakeEC");
@@ -231,11 +235,6 @@ public partial class GUBP : BuildCommand
 		FindCompletionState(Graph.BuildNodes, JobInfo, LocalOnly);
 		ComputeDependentFrequencies(Graph.BuildNodes);
 
-        if (bCleanLocalTempStorage)  // shared temp storage can never be wiped
-        {
-            TempStorage.DeleteLocalTempStorage();
-        }
-
 		int TimeIndex = ParseParamInt("TimeIndex", 0);
 		if (TimeIndex == 0)
 		{
