Move JobInfo out of the GUBP class.

diff --git a/Engine/Source/Programs/AutomationTool/GUBP/GUBP.Automation.csproj b/Engine/Source/Programs/AutomationTool/GUBP/GUBP.Automation.csproj
--- a/Engine/Source/Programs/AutomationTool/GUBP/GUBP.Automation.csproj
+++ b/Engine/Source/Programs/AutomationTool/GUBP/GUBP.Automation.csproj
@@ -63,6 +63,7 @@
     <Compile Include="BuildGraph.cs" />
     <Compile Include="ElectricCommander.cs" />
     <Compile Include="GUBP.cs" />
+    <Compile Include="JobInfo.cs" />
     <Compile Include="LegacyBranchSetup.cs" />
     <Compile Include="LegacyNodes.cs" />
     <Compile Include="Nodes\AggregateNode.cs" />
diff --git a/Engine/Source/Programs/AutomationTool/GUBP/GUBP.cs b/Engine/Source/Programs/AutomationTool/GUBP/GUBP.cs
--- a/Engine/Source/Programs/AutomationTool/GUBP/GUBP.cs
+++ b/Engine/Source/Programs/AutomationTool/GUBP/GUBP.cs
@@ -49,88 +49,6 @@ public partial class GUBP : BuildCommand
 	const string FailedTempStorageSuffix = "_Failed";
 	const string SucceededTempStorageSuffix = "_Succeeded";
 
-    /// <summary>
-    /// Contains some global attributes of the GUBP job. A GUBP job consists of various nodes that may be executed on many different machines at the same time.
-    /// These attributes are constant throughout the job, for any node in the job.
-    /// </summary>
-    /// <remarks>
-    /// This does not contain ALL the job attributes, just the ones that were added during incremental refactoring.
-    /// 
-    /// This class is immutable to ensure no side-effect code an tamper with it after the job starts.
-    /// </remarks>
-    public class JobInfo
-    {
-        /// <summary>
-        /// Used as the prefix for temp storage for the job.
-        /// 
-        /// If P4 is enabled, this is P4Env.BuildRootEscaped, ie the branch name safe to create a file or folder from. If P4 is not enabled, it is NoP4, 
-        /// but we never write to the shared temp storage location in those cases.
-        /// For testing, arbitrary values can be used.
-        /// </summary>
-        public string BranchNameForTempStorage { get; private set; }
-
-        /// <summary>
-        /// The changelist associated with this job (ie, the code CL that was sync'ed to run this job.
-        /// The -CL parameter can be used to override the P4 environment for testing.
-        /// If P4 is not enabled, this value will be 0.
-        /// </summary>
-        public int Changelist { get; private set; }
-
-        /// <summary>
-        /// If this is a preflight vuild, this is the changelist that contains the shelve being preflighted. 
-        /// Otherwise, will be 0.
-        /// </summary>
-        public int PreflightShelveChangelist { get; private set; }
-
-        /// <summary>
-        /// If this is a preflight build, contains a unique ID to identify the preflight job.
-        /// Used to disambiguate multiple preflight runs from the same changelist and shelve, since the
-        /// code in the shelve can change from run to run.
-        /// </summary>
-        public int PreflightUID { get; private set; }
-
-        /// <summary>
-        /// Will be true if this job is a preflight run. Shortcut for checking if the shelve is > 0
-        /// </summary>
-        public bool IsPreflight { get { return PreflightShelveChangelist > 0; } }
-
-        /// <summary>
-        /// If this is a preflight run, returns a suffix used to describe the preflight shelve and UID. Used
-        /// to uniquely identify a temp storage block for a particular preflight run.
-        /// </summary>
-        /// <returns></returns>
-        public string GetPreflightSuffix()
-        {
-            return IsPreflight ? string.Format("-PF-{0}-{1}", PreflightShelveChangelist, PreflightUID) : "";
-        }
-
-        /// <summary>
-        /// Quick method to create a new JobInfo using all the same attributes, but a different CL.
-        /// 
-        /// Used by some GUBP code to search for and use previous temp storage blocks when checking for node history.
-        /// </summary>
-        /// <param name="NewCL">The CL to use instead of the current Job's CL.</param>
-        /// <returns>A new JobInfo instance with all the same attributes as the current instance, but a different CL.</returns>
-        public JobInfo CreateWithNewChangelist(int NewCL)
-        {
-            return new JobInfo(BranchNameForTempStorage, NewCL, PreflightShelveChangelist, PreflightUID);
-        }
-
-        /// <summary>
-        /// Public ctor. This is an immutable class, so all attributes must be set a construction time.
-        /// </summary>
-        /// <param name="BranchNameForTempStorage">See similarly named attribute.</param>
-        /// <param name="Changelist">See similarly named attribute.</param>
-        /// <param name="PreflightShelveChangelist">See similarly named attribute.</param>
-        /// <param name="PreflightUID">See similarly named attribute.</param>
-        public JobInfo(string BranchNameForTempStorage, int Changelist, int PreflightShelveChangelist, int PreflightUID)
-        {
-            this.BranchNameForTempStorage = BranchNameForTempStorage;
-            this.Changelist = Changelist;
-            this.PreflightShelveChangelist = PreflightShelveChangelist;
-            this.PreflightUID = PreflightUID;
-        }
-    }
     
     class NodeHistory
 	{
diff --git a/Engine/Source/Programs/AutomationTool/GUBP/JobInfo.cs b/Engine/Source/Programs/AutomationTool/GUBP/JobInfo.cs
new file mode 100644
--- /dev/null
+++ b/Engine/Source/Programs/AutomationTool/GUBP/JobInfo.cs
@@ -0,0 +1,91 @@
+﻿using System;
+using System.Collections.Generic;
+using System.Linq;
+using System.Text;
+using System.Threading.Tasks;
+
+namespace AutomationTool
+{
+    /// <summary>
+    /// Contains some global attributes of the GUBP job. A GUBP job consists of various nodes that may be executed on many different machines at the same time.
+    /// These attributes are constant throughout the job, for any node in the job.
+    /// </summary>
+    /// <remarks>
+    /// This does not contain ALL the job attributes, just the ones that were added during incremental refactoring.
+    /// 
+    /// This class is immutable to ensure no side-effect code an tamper with it after the job starts.
+    /// </remarks>
+    public class JobInfo
+    {
+        /// <summary>
+        /// Used as the prefix for temp storage for the job.
+        /// 
+        /// If P4 is enabled, this is P4Env.BuildRootEscaped, ie the branch name safe to create a file or folder from. If P4 is not enabled, it is NoP4, 
+        /// but we never write to the shared temp storage location in those cases.
+        /// For testing, arbitrary values can be used.
+        /// </summary>
+        public string BranchNameForTempStorage { get; private set; }
+
+        /// <summary>
+        /// The changelist associated with this job (ie, the code CL that was sync'ed to run this job.
+        /// The -CL parameter can be used to override the P4 environment for testing.
+        /// If P4 is not enabled, this value will be 0.
+        /// </summary>
+        public int Changelist { get; private set; }
+
+        /// <summary>
+        /// If this is a preflight vuild, this is the changelist that contains the shelve being preflighted. 
+        /// Otherwise, will be 0.
+        /// </summary>
+        public int PreflightShelveChangelist { get; private set; }
+
+        /// <summary>
+        /// If this is a preflight build, contains a unique ID to identify the preflight job.
+        /// Used to disambiguate multiple preflight runs from the same changelist and shelve, since the
+        /// code in the shelve can change from run to run.
+        /// </summary>
+        public int PreflightUID { get; private set; }
+
+        /// <summary>
+        /// Will be true if this job is a preflight run. Shortcut for checking if the shelve is > 0
+        /// </summary>
+        public bool IsPreflight { get { return PreflightShelveChangelist > 0; } }
+
+        /// <summary>
+        /// If this is a preflight run, returns a suffix used to describe the preflight shelve and UID. Used
+        /// to uniquely identify a temp storage block for a particular preflight run.
+        /// </summary>
+        /// <returns></returns>
+        public string GetPreflightSuffix()
+        {
+            return IsPreflight ? string.Format("-PF-{0}-{1}", PreflightShelveChangelist, PreflightUID) : "";
+        }
+
+        /// <summary>
+        /// Quick method to create a new JobInfo using all the same attributes, but a different CL.
+        /// 
+        /// Used by some GUBP code to search for and use previous temp storage blocks when checking for node history.
+        /// </summary>
+        /// <param name="NewCL">The CL to use instead of the current Job's CL.</param>
+        /// <returns>A new JobInfo instance with all the same attributes as the current instance, but a different CL.</returns>
+        public JobInfo CreateWithNewChangelist(int NewCL)
+        {
+            return new JobInfo(BranchNameForTempStorage, NewCL, PreflightShelveChangelist, PreflightUID);
+        }
+
+        /// <summary>
+        /// Public ctor. This is an immutable class, so all attributes must be set a construction time.
+        /// </summary>
+        /// <param name="BranchNameForTempStorage">See similarly named attribute.</param>
+        /// <param name="Changelist">See similarly named attribute.</param>
+        /// <param name="PreflightShelveChangelist">See similarly named attribute.</param>
+        /// <param name="PreflightUID">See similarly named attribute.</param>
+        public JobInfo(string BranchNameForTempStorage, int Changelist, int PreflightShelveChangelist, int PreflightUID)
+        {
+            this.BranchNameForTempStorage = BranchNameForTempStorage;
+            this.Changelist = Changelist;
+            this.PreflightShelveChangelist = PreflightShelveChangelist;
+            this.PreflightUID = PreflightUID;
+        }
+    }
+}
diff --git a/Engine/Source/Programs/AutomationTool/GUBP/TempStorage.cs b/Engine/Source/Programs/AutomationTool/GUBP/TempStorage.cs
--- a/Engine/Source/Programs/AutomationTool/GUBP/TempStorage.cs
+++ b/Engine/Source/Programs/AutomationTool/GUBP/TempStorage.cs
@@ -32,7 +32,7 @@ namespace AutomationTool
         /// <summary>
         /// The GUBP global job attributes.
         /// </summary>
-        public GUBP.JobInfo JobInfo { get; private set; }
+        public JobInfo JobInfo { get; private set; }
 
         /// <summary>
         /// The storage name associated with this node. It is usually the name of the GUBP node running, but sometimes certain suffixes are added to indicate node completion status, etc.
@@ -44,7 +44,7 @@ namespace AutomationTool
         /// </summary>
         /// <param name="JobInfo"></param>
         /// <param name="NodeStorageName"></param>
-        public TempStorageNodeInfo(GUBP.JobInfo JobInfo, string NodeStorageName)
+        public TempStorageNodeInfo(JobInfo JobInfo, string NodeStorageName)
         {
             this.JobInfo = JobInfo;
             this.NodeStorageName = NodeStorageName;
@@ -1188,7 +1188,7 @@ namespace AutomationTool
         static internal void TestTempStorage(CommandEnvironment CmdEnv)
         {
             // We are not a real GUBP job, so fake the values.
-            var TestTempStorageInfo = new TempStorageNodeInfo(new GUBP.JobInfo("Test", 0, 0, 0), "Test");
+            var TestTempStorageInfo = new TempStorageNodeInfo(new JobInfo("Test", 0, 0, 0), "Test");
 
             // Delete any local and shared temp storage that may exist.
             DeleteLocalTempStorage();
