Better way to locate IncrediBuild xgConsole.exe (avoids illegal characters in path)
#ue4

diff --git a/Engine/Source/Programs/AutomationTool/AutomationUtils/UE4Build.cs b/Engine/Source/Programs/AutomationTool/AutomationUtils/UE4Build.cs
--- a/Engine/Source/Programs/AutomationTool/AutomationUtils/UE4Build.cs
+++ b/Engine/Source/Programs/AutomationTool/AutomationUtils/UE4Build.cs
@@ -737,16 +737,26 @@ namespace AutomationTool
 		{
 			if (string.IsNullOrEmpty(XGEConsoleExePath))
 			{
-				string[] PathDirs = Environment.GetEnvironmentVariable("PATH").Split(Path.PathSeparator);
-				foreach (string PathDir in PathDirs)
+				try
 				{
-					string FullPath = Path.Combine(PathDir, "xgConsole.exe");
-					if (FileExists(FullPath))
+					Process proc = new Process();
+					proc.StartInfo.UseShellExecute = false;
+					proc.StartInfo.FileName = "where";
+					proc.StartInfo.Arguments = "xgConsole.exe";
+					proc.StartInfo.RedirectStandardOutput = true;
+					proc.Start();
+					string result = proc.StandardOutput.ReadToEnd();
+					proc.WaitForExit();
+
+					if (proc.ExitCode == 0)
 					{
-						XGEConsoleExePath = FullPath;
-						break;
+						XGEConsoleExePath = result.Substring(0, result.IndexOf(Environment.NewLine));
 					}
 				}
+				catch (Exception)
+				{
+					// where not found
+				}
 			}
 			return XGEConsoleExePath;
 		}
