Switch path parsing to use the system path separator, rather than hard-coding ';'.

diff --git a/Engine/Source/Programs/AutomationTool/Scripts/Tests.Automation.cs b/Engine/Source/Programs/AutomationTool/Scripts/Tests.Automation.cs
--- a/Engine/Source/Programs/AutomationTool/Scripts/Tests.Automation.cs
+++ b/Engine/Source/Programs/AutomationTool/Scripts/Tests.Automation.cs
@@ -52,7 +52,7 @@ class GitPullRequest : BuildCommand
 			return Path.GetFullPath(ExeName);
 		}
 
-		foreach (string BasePath in Environment.GetEnvironmentVariable("PATH").Split(';'))
+		foreach (string BasePath in Environment.GetEnvironmentVariable("PATH").Split(Path.PathSeparator))
 		{
 			var FullPath = Path.Combine(BasePath, ExeName);
 			if (ExpectedPathSubstring == null || FullPath.IndexOf(ExpectedPathSubstring, StringComparison.InvariantCultureIgnoreCase) != -1)
