fix for not failing when the ddc stats file is not generated.

diff --git a/Engine/Source/Programs/AutomationTool/AutomationUtils/CommandletUtils.cs b/Engine/Source/Programs/AutomationTool/AutomationUtils/CommandletUtils.cs
--- a/Engine/Source/Programs/AutomationTool/AutomationUtils/CommandletUtils.cs
+++ b/Engine/Source/Programs/AutomationTool/AutomationUtils/CommandletUtils.cs
@@ -206,14 +206,17 @@ namespace AutomationTool
 			// copy the ddc access log to the destination folder
 			string DestDDCLogFile = LogUtils.GetUniqueLogName(CombinePaths(CmdEnv.LogFolder, Commandlet + "-DDC"));
 			string LocalDDCLogFile = CombinePaths(CmdEnv.EngineSavedFolder, "CookRunDDC.txt");
-			if (!CommandUtils.CopyFile_NoExceptions(LocalDDCLogFile, DestDDCLogFile))
+			if (!CommandUtils.FileExists(LocalDDCLogFile) || !CommandUtils.CopyFile_NoExceptions(LocalDDCLogFile, DestDDCLogFile))
 			{
 				CommandUtils.LogWarning("Commandlet {0} failed to copy the local log file from {1} to {2}. The log file will be lost.", Commandlet, LocalDDCLogFile, DestDDCLogFile);
 			}
 
 			// Whether it was copied correctly or not, delete the local log as it was only a temporary file. 
 			CommandUtils.DeleteFile_NoExceptions(LocalLogFile);
-			CommandUtils.DeleteFile_NoExceptions(LocalDDCLogFile);
+			if (CommandUtils.FileExists(LocalDDCLogFile))
+			{
+				CommandUtils.DeleteFile_NoExceptions(LocalDDCLogFile);
+			}
 
 			if (RunResult.ExitCode != 0)
 			{
