Save cooker stats from build machine builds.
Help track down irratic cooking times.

diff --git a/Engine/Source/Programs/AutomationTool/AutomationUtils/CommandletUtils.cs b/Engine/Source/Programs/AutomationTool/AutomationUtils/CommandletUtils.cs
--- a/Engine/Source/Programs/AutomationTool/AutomationUtils/CommandletUtils.cs
+++ b/Engine/Source/Programs/AutomationTool/AutomationUtils/CommandletUtils.cs
@@ -211,6 +211,17 @@ namespace AutomationTool
 				CommandUtils.LogWarning("Commandlet {0} failed to copy the local log file from {1} to {2}. The log file will be lost.", Commandlet, LocalDDCLogFile, DestDDCLogFile);
 			}
 
+            string ProjectStatsDirectory = CombinePaths(Path.GetDirectoryName(ProjectName), "Saved", "Stats");
+            string DestCookerStats = CmdEnv.LogFolder;
+            foreach ( var StatsFile in Directory.EnumerateFiles(ProjectStatsDirectory, "*.csv") )
+            {
+                if (!CommandUtils.CopyFile_NoExceptions(StatsFile, CombinePaths(DestCookerStats, Path.GetFileName(StatsFile))))
+                {
+                    CommandUtils.LogWarning("Commandlet {0} failed to copy the local log file from {1} to {2}. The log file will be lost.", Commandlet, LocalDDCLogFile, DestDDCLogFile);
+                }
+            }
+            
+
 			// Whether it was copied correctly or not, delete the local log as it was only a temporary file. 
 			CommandUtils.DeleteFile_NoExceptions(LocalLogFile);
 			if (CommandUtils.FileExists(LocalDDCLogFile))
