When copying the html5 helper command template remove any read only file attributes.
Fixes issues when attempting to write to these files.

diff --git a/Engine/Source/Programs/AutomationTool/HTML5/HTML5Platform.Automation.cs b/Engine/Source/Programs/AutomationTool/HTML5/HTML5Platform.Automation.cs
--- a/Engine/Source/Programs/AutomationTool/HTML5/HTML5Platform.Automation.cs
+++ b/Engine/Source/Programs/AutomationTool/HTML5/HTML5Platform.Automation.cs
@@ -321,10 +321,12 @@ public class HTML5Platform : Platform
 				Directory.CreateDirectory(Path.GetDirectoryName(InOutputFile));
 				if (File.Exists(InOutputFile))
 				{
+					File.SetAttributes(InOutputFile, File.GetAttributes(InOutputFile) & ~FileAttributes.ReadOnly);
 					File.Delete(InOutputFile);
 				}
 				File.Copy(InTemplateFile, InOutputFile);
-				using (var CmdFile = File.OpenWrite(InOutputFile)) 
+				File.SetAttributes(InOutputFile, File.GetAttributes(InOutputFile) & ~FileAttributes.ReadOnly);
+				using (var CmdFile = File.Open(InOutputFile, FileMode.OpenOrCreate | FileMode.Truncate)) 
 				{
 					Byte[] BytesToWrite = new UTF8Encoding(true).GetBytes(outputContents.ToString());
 					CmdFile.Write(BytesToWrite, 0, BytesToWrite.Length);
