Merging UAT stream support from //UE4/Main.

diff --git a/Engine/Source/Programs/AutomationTool/AutomationUtils/LocalP4Environment.cs b/Engine/Source/Programs/AutomationTool/AutomationUtils/LocalP4Environment.cs
--- a/Engine/Source/Programs/AutomationTool/AutomationUtils/LocalP4Environment.cs
+++ b/Engine/Source/Programs/AutomationTool/AutomationUtils/LocalP4Environment.cs
@@ -140,31 +140,39 @@ namespace AutomationTool
 		/// <param name="ClientRootPath">Client root</param>
 		private static void DetectRootPaths(P4Connection Connection, string LocalRootPath, P4ClientInfo ThisClient, out string BuildRootPath, out string ClientRootPath)
 		{
-			// Figure out the build root
-			string KnownFilePathFromRoot = CommandEnvironment.KnownFileRelativeToRoot;
-			string KnownLocalPath = CommandUtils.MakePathSafeToUseWithCommandLine(CommandUtils.CombinePaths(PathSeparator.Slash, LocalRootPath, KnownFilePathFromRoot));
-			ProcessResult P4Result = Connection.P4(String.Format("files -m 1 {0}", KnownLocalPath), AllowSpew: false);
-
-			string KnownFileDepotMapping = P4Result.Output;
-
-			// Get the build root
-			Log.TraceVerbose("Looking for {0} in {1}", KnownFilePathFromRoot, KnownFileDepotMapping);
-			int EndIdx = KnownFileDepotMapping.IndexOf(KnownFilePathFromRoot, StringComparison.CurrentCultureIgnoreCase);
-			if (EndIdx < 0)
-			{
-				EndIdx = KnownFileDepotMapping.IndexOf(CommandUtils.ConvertSeparators(PathSeparator.Slash, KnownFilePathFromRoot), StringComparison.CurrentCultureIgnoreCase);
-			}
-			// Get the root path without the trailing path separator
-			BuildRootPath = KnownFileDepotMapping.Substring(0, EndIdx - 1);
-
-			// Get the client root
-			if (LocalRootPath.StartsWith(CommandUtils.CombinePaths(PathSeparator.Slash, ThisClient.RootPath, "/"), StringComparison.InvariantCultureIgnoreCase) || LocalRootPath == ThisClient.RootPath)
+			if(!String.IsNullOrEmpty(ThisClient.Stream))
 			{
-				ClientRootPath = CommandUtils.CombinePaths(PathSeparator.Depot, String.Format("//{0}/", ThisClient.Name), LocalRootPath.Substring(ThisClient.RootPath.Length));
+				BuildRootPath = ThisClient.Stream;
+				ClientRootPath = String.Format("//{0}", ThisClient.Name);
 			}
 			else
 			{
-				throw new AutomationException("LocalRootPath ({0}) does not start with the client root path ({1})", LocalRootPath, ThisClient.RootPath);
+				// Figure out the build root
+				string KnownFilePathFromRoot = CommandEnvironment.KnownFileRelativeToRoot;
+				string KnownLocalPath = CommandUtils.MakePathSafeToUseWithCommandLine(CommandUtils.CombinePaths(PathSeparator.Slash, LocalRootPath, KnownFilePathFromRoot));
+				ProcessResult P4Result = Connection.P4(String.Format("files -m 1 {0}", KnownLocalPath), AllowSpew: false);
+
+				string KnownFileDepotMapping = P4Result.Output;
+
+				// Get the build root
+				Log.TraceVerbose("Looking for {0} in {1}", KnownFilePathFromRoot, KnownFileDepotMapping);
+				int EndIdx = KnownFileDepotMapping.IndexOf(KnownFilePathFromRoot, StringComparison.CurrentCultureIgnoreCase);
+				if (EndIdx < 0)
+				{
+					EndIdx = KnownFileDepotMapping.IndexOf(CommandUtils.ConvertSeparators(PathSeparator.Slash, KnownFilePathFromRoot), StringComparison.CurrentCultureIgnoreCase);
+				}
+				// Get the root path without the trailing path separator
+				BuildRootPath = KnownFileDepotMapping.Substring(0, EndIdx - 1);
+
+				// Get the client root
+				if (LocalRootPath.StartsWith(CommandUtils.CombinePaths(PathSeparator.Slash, ThisClient.RootPath, "/"), StringComparison.InvariantCultureIgnoreCase) || LocalRootPath == ThisClient.RootPath)
+				{
+					ClientRootPath = CommandUtils.CombinePaths(PathSeparator.Depot, String.Format("//{0}/", ThisClient.Name), LocalRootPath.Substring(ThisClient.RootPath.Length));
+				}
+				else
+				{
+					throw new AutomationException("LocalRootPath ({0}) does not start with the client root path ({1})", LocalRootPath, ThisClient.RootPath);
+				}
 			}
 		}
 
diff --git a/Engine/Source/Programs/AutomationTool/AutomationUtils/P4Environment.cs b/Engine/Source/Programs/AutomationTool/AutomationUtils/P4Environment.cs
--- a/Engine/Source/Programs/AutomationTool/AutomationUtils/P4Environment.cs
+++ b/Engine/Source/Programs/AutomationTool/AutomationUtils/P4Environment.cs
@@ -58,7 +58,6 @@ namespace AutomationTool
 		public string BuildRootEscaped { get; protected set; }
 		public string LabelToSync { get; protected set; }
 		public string LabelPrefix { get; protected set; }
-		public string BranchName { get; protected set; }
 
 
 		internal P4Environment(P4Connection Connection, CommandEnvironment CmdEnv)
@@ -110,22 +109,6 @@ namespace AutomationTool
 					// In local runs, changelist initialization can be really slow!
 					VerifyChangelistStringAndSetChangelistNumber();
 				}
-
-				// Setup branch name
-				string DepotSuffix = "//depot/";
-				if (BuildRootP4.StartsWith(DepotSuffix))
-				{
-					BranchName = BuildRootP4.Substring(DepotSuffix.Length);
-				}
-				else
-				{
-					throw new AutomationException("Needs update to work with a stream");
-				}
-
-				if (String.IsNullOrWhiteSpace(BranchName))
-				{
-					throw new AutomationException("BUILD FAILED no branch name.");
-				}
 			}
 
 			LogSettings();
diff --git a/Engine/Source/Programs/AutomationTool/AutomationUtils/P4Utils.cs b/Engine/Source/Programs/AutomationTool/AutomationUtils/P4Utils.cs
--- a/Engine/Source/Programs/AutomationTool/AutomationUtils/P4Utils.cs
+++ b/Engine/Source/Programs/AutomationTool/AutomationUtils/P4Utils.cs
@@ -73,6 +73,7 @@ namespace AutomationTool
 		public string RootPath;
 		public string Host;
 		public string Owner;
+		public string Stream;
 		public DateTime Access;
         public P4LineEnd LineEnd;
         public P4ClientOption Options;
@@ -85,6 +86,7 @@ namespace AutomationTool
 				&& RootPath == Other.RootPath 
 				&& Host == Other.Host 
 				&& Owner == Other.Owner 
+				&& Stream == Other.Stream
 				&& LineEnd == Other.LineEnd 
 				&& Options == Other.Options 
 				&& SubmitOptions == Other.SubmitOptions
@@ -2441,6 +2443,7 @@ namespace AutomationTool
 					Info.RootPath = CommandUtils.ConvertSeparators(PathSeparator.Default, Info.RootPath);
 				}
 				Tags.TryGetValue("Owner", out Info.Owner);
+				Tags.TryGetValue("Stream", out Info.Stream);
 				string AccessTime;
 				Tags.TryGetValue("Access", out AccessTime);
 				if (!String.IsNullOrEmpty(AccessTime))
diff --git a/Engine/Source/Programs/AutomationTool/Scripts/Tests.Automation.cs b/Engine/Source/Programs/AutomationTool/Scripts/Tests.Automation.cs
--- a/Engine/Source/Programs/AutomationTool/Scripts/Tests.Automation.cs
+++ b/Engine/Source/Programs/AutomationTool/Scripts/Tests.Automation.cs
@@ -1340,7 +1340,6 @@ public class ZeroEngineVersions : BuildCommand
 			if (VersionH.Contains("#define ENGINE_VERSION 0") == false)
 			{
 				Log("Zeroing out engine versions in {0}", VersionFilename);
-				VersionH.ReplaceLine("#define BRANCH_NAME ", "\"" + P4Env.BranchName + "\"");
 				VersionH.ReplaceLine("#define BUILT_FROM_CHANGELIST ", "0");
 
 				VersionH.Commit();
