UE-20341 - fixed issues with needing to cook server when cooking client

diff --git a/Engine/Source/Programs/AutomationTool/AutomationUtils/ProjectParams.cs b/Engine/Source/Programs/AutomationTool/AutomationUtils/ProjectParams.cs
--- a/Engine/Source/Programs/AutomationTool/AutomationUtils/ProjectParams.cs
+++ b/Engine/Source/Programs/AutomationTool/AutomationUtils/ProjectParams.cs
@@ -507,10 +507,10 @@ namespace AutomationTool
 			this.FileServer = GetParamValueIfNotSpecified(Command, FileServer, this.FileServer, "fileserver");
 			this.DedicatedServer = GetParamValueIfNotSpecified(Command, DedicatedServer, this.DedicatedServer, "dedicatedserver", "server");
 			this.Client = GetParamValueIfNotSpecified(Command, Client, this.Client, "client");
-			if( this.Client )
+			/*if( this.Client )
 			{
 				this.DedicatedServer = true;
-			}
+			}*/
 			this.NoClient = GetParamValueIfNotSpecified(Command, NoClient, this.NoClient, "noclient");
 			this.LogWindow = GetParamValueIfNotSpecified(Command, LogWindow, this.LogWindow, "logwindow");
 			this.Stage = GetParamValueIfNotSpecified(Command, Stage, this.Stage, "stage");
diff --git a/Engine/Source/Programs/AutomationTool/Scripts/CookCommand.Automation.cs b/Engine/Source/Programs/AutomationTool/Scripts/CookCommand.Automation.cs
--- a/Engine/Source/Programs/AutomationTool/Scripts/CookCommand.Automation.cs
+++ b/Engine/Source/Programs/AutomationTool/Scripts/CookCommand.Automation.cs
@@ -87,7 +87,7 @@ public partial class Project : CommandUtils
 				{
 					// Use the data platform, sometimes we will copy another platform's data
 					var DataPlatform = Params.GetCookedDataPlatformForClientTarget(ClientPlatform);
-					PlatformsToCook.Add(Params.GetTargetPlatformInstance(DataPlatform).GetCookPlatform(false, Params.HasDedicatedServerAndClient, Params.CookFlavor));
+					PlatformsToCook.Add(Params.GetTargetPlatformInstance(DataPlatform).GetCookPlatform(false, Params.Client, Params.CookFlavor));
 				}
 			}
 			if (Params.DedicatedServer)
