Added some more warnings when stats directory was not found.

diff --git a/Engine/Source/Programs/AutomationTool/AutomationUtils/CommandletUtils.cs b/Engine/Source/Programs/AutomationTool/AutomationUtils/CommandletUtils.cs
--- a/Engine/Source/Programs/AutomationTool/AutomationUtils/CommandletUtils.cs
+++ b/Engine/Source/Programs/AutomationTool/AutomationUtils/CommandletUtils.cs
@@ -223,6 +223,10 @@ namespace AutomationTool
                     }
                 }
             }
+            else
+            {
+                CommandUtils.LogWarning("Failed to find directory {0} will not save stats", ProjectStatsDirectory);
+            }
 
 			// Whether it was copied correctly or not, delete the local log as it was only a temporary file. 
 			CommandUtils.DeleteFile_NoExceptions(LocalLogFile);
