Dedicated servers should stage needed engine content for crash reporter

Merging 2671465

//depot/UE4-UT/...

to //depot/UE4/...

diff --git a/Engine/Source/Programs/AutomationTool/Scripts/CopyBuildToStagingDirectory.Automation.cs b/Engine/Source/Programs/AutomationTool/Scripts/CopyBuildToStagingDirectory.Automation.cs
--- a/Engine/Source/Programs/AutomationTool/Scripts/CopyBuildToStagingDirectory.Automation.cs
+++ b/Engine/Source/Programs/AutomationTool/Scripts/CopyBuildToStagingDirectory.Automation.cs
@@ -480,7 +480,7 @@ public partial class Project : CommandUtils
 
             // CrashReportClient is a standalone slate app that does not look in the generated pak file, so it needs the Content/Slate and Shaders/StandaloneRenderer folders Non-UFS
             // @todo Make CrashReportClient more portable so we don't have to do this
-            if (SC.bStageCrashReporter && UnrealBuildTool.UnrealBuildTool.PlatformSupportsCrashReporter(SC.StageTargetPlatform.PlatformType) && !SC.DedicatedServer)
+            if (SC.bStageCrashReporter && UnrealBuildTool.UnrealBuildTool.PlatformSupportsCrashReporter(SC.StageTargetPlatform.PlatformType))
             {
                 //If the .dat file needs to be staged as NonUFS for non-Windows/Linux hosts we need to change the casing as we do with the build properties file above.
                 SC.StageFiles(StagedFileType.NonUFS, CombinePaths(SC.LocalRoot, "Engine/Content/Slate"));
