Don't use AppendLine on StringBuilder is it outputs system dependent line endings. Manually output \\n instead.
Fixes UE-20380

diff --git a/Engine/Source/Programs/AutomationTool/HTML5/HTML5Platform.Automation.cs b/Engine/Source/Programs/AutomationTool/HTML5/HTML5Platform.Automation.cs
--- a/Engine/Source/Programs/AutomationTool/HTML5/HTML5Platform.Automation.cs
+++ b/Engine/Source/Programs/AutomationTool/HTML5/HTML5Platform.Automation.cs
@@ -309,7 +309,7 @@ public class HTML5Platform : Platform
 					LineStr = LineStr.Replace("${unreal_mono_pkg_path_base}", InMonoPathParent);
 				}
 
-				outputContents.AppendLine(LineStr);
+				outputContents.Append(LineStr + '\n');
 			}
 		}
 
