Don't try to expant empty architectures, only missing one.
Several platforms have empty architectures and we shouldn't call build code when staging if we can avoid it.

diff --git a/Engine/Source/Programs/AutomationTool/Scripts/CopyBuildToStagingDirectory.Automation.cs b/Engine/Source/Programs/AutomationTool/Scripts/CopyBuildToStagingDirectory.Automation.cs
--- a/Engine/Source/Programs/AutomationTool/Scripts/CopyBuildToStagingDirectory.Automation.cs
+++ b/Engine/Source/Programs/AutomationTool/Scripts/CopyBuildToStagingDirectory.Automation.cs
@@ -1544,7 +1544,7 @@ public partial class Project : CommandUtils
 					foreach (UnrealTargetPlatform ReceiptPlatform in SubPlatformsToStage)
 					{
                         string Architecture = Params.SpecifiedArchitecture;
-                        if (string.IsNullOrEmpty(Architecture))
+                        if (Architecture == null)
                         {
                             Architecture = "";
                             var BuildPlatform = UEBuildPlatform.GetBuildPlatform(ReceiptPlatform, true);
