Improve Amazon S3 upload performance

Use .NET ThreadPool to have finer control over the number of upload threads
Use AWSSDK's PutFileAsync method over the non-async version

diff --git a/Engine/Source/Programs/AutomationTool/AutomationUtils/MCPPublic.cs b/Engine/Source/Programs/AutomationTool/AutomationUtils/MCPPublic.cs
--- a/Engine/Source/Programs/AutomationTool/AutomationUtils/MCPPublic.cs
+++ b/Engine/Source/Programs/AutomationTool/AutomationUtils/MCPPublic.cs
@@ -14,6 +14,7 @@ using UnrealBuildTool;
 namespace EpicGames.MCP.Automation
 {
 	using EpicGames.MCP.Config;
+	using System.Threading.Tasks;
 
     /// <summary>
     /// Utility class to provide commit/rollback functionality via an RAII-like functionality.
@@ -723,34 +724,54 @@ namespace EpicGames.MCP.Automation
 		abstract public byte[] GetFile(string Container, string Identifier, out string ContentType);
 
 		/// <summary>
-		/// Posts a file to the cloud storage provider
-		/// NOTE: the method returns void, rather than bSuccess as you might imagine, because of an apparent bug in VS2013 debugger, where an AccessViolationExeption is
-		/// thrown when stepping in to an overridden method, with a return value, and at least one out parameter, when optimizations are enabled.
+		/// Posts a file to the cloud storage provider.
+		/// </summary>
+		/// <param name="Container">The name of the folder or container in which to store the file.</param>
+		/// <param name="Identifier">The identifier or filename of the file to write.</param>
+		/// <param name="Contents">A byte array containing the data to write.</param>
+		/// <param name="ContentType">The MIME type of the file being uploaded.</param>
+		/// <param name="bOverwrite">If true, will overwrite an existing file.  If false, will throw an exception if the file exists.</param>
+		/// <param name="bMakePublic">Specified whether the file should be made public readable.</param>
+		public PostFileResult PostFile(string Container, string Identifier, byte[] Contents, string ContentType = null, bool bOverwrite = true, bool bMakePublic = false)
+		{
+			return PostFileAsync(Container, Identifier, Contents, ContentType, bOverwrite, bMakePublic).Result;
+		}
+
+		/// <summary>
+		/// Posts a file to the cloud storage provider asynchronously.
 		/// </summary>
 		/// <param name="Container">The name of the folder or container in which to store the file.</param>
 		/// <param name="Identifier">The identifier or filename of the file to write.</param>
 		/// <param name="Contents">A byte array containing the data to write.</param>
-		/// <param name="ObjectURL">An OUTPUT parameter which will be set to the URL of the uploaded file on success.</param>
-		/// <param name="bSuccess">An OUTPUT parameter which will be set to true if the write succeeds, false otherwise.</param>
 		/// <param name="ContentType">The MIME type of the file being uploaded.</param>
 		/// <param name="bOverwrite">If true, will overwrite an existing file.  If false, will throw an exception if the file exists.</param>
 		/// <param name="bMakePublic">Specified whether the file should be made public readable.</param>
-		abstract public void PostFile(string Container, string Identifier, byte[] Contents, out string ObjectURL, out bool bSuccess, string ContentType = null, bool bOverwrite = true, bool bMakePublic = false);
+		abstract public Task<PostFileResult> PostFileAsync(string Container, string Identifier, byte[] Contents, string ContentType = null, bool bOverwrite = true, bool bMakePublic = false);
 
 		/// <summary>
 		/// Posts a file to the cloud storage provider.
-		/// NOTE: the method returns void, rather than bSuccess as you might imagine, because of an apparent bug in VS2013 debugger, where an AccessViolationExeption is
-		/// thrown when stepping in to an overridden method, with a return value, and at least one out parameter, when optimizations are enabled.
 		/// </summary>
 		/// <param name="Container">The name of the folder or container in which to store the file.</param>
 		/// <param name="Identifier">The identifier or filename of the file to write.</param>
 		/// <param name="SourceFilePath">The full path of the file to upload.</param>
-		/// <param name="ObjectURL">An OUTPUT parameter which will be set to the URL of the uploaded file on success.</param>
-		/// <param name="bSuccess">An OUTPUT parameter which will be set to true if the write succeeds, false otherwise.</param>
 		/// <param name="ContentType">The MIME type of the file being uploaded.</param>
 		/// <param name="bOverwrite">If true, will overwrite an existing file.  If false, will throw an exception if the file exists.</param>
 		/// <param name="bMakePublic">Specified whether the file should be made public readable.</param>
-		abstract public void PostFile(string Container, string Identifier, string SourceFilePath, out string ObjectURL, out bool bSuccess, string ContentType = null, bool bOverwrite = true, bool bMakePublic = false);
+		public PostFileResult PostFile(string Container, string Identifier, string SourceFilePath, string ContentType = null, bool bOverwrite = true, bool bMakePublic = false)
+		{
+			return PostFileAsync(Container, Identifier, SourceFilePath, ContentType, bOverwrite, bMakePublic).Result;
+		}
+
+		/// <summary>
+		/// Posts a file to the cloud storage provider asynchronously.
+		/// </summary>
+		/// <param name="Container">The name of the folder or container in which to store the file.</param>
+		/// <param name="Identifier">The identifier or filename of the file to write.</param>
+		/// <param name="SourceFilePath">The full path of the file to upload.</param>
+		/// <param name="ContentType">The MIME type of the file being uploaded.</param>
+		/// <param name="bOverwrite">If true, will overwrite an existing file.  If false, will throw an exception if the file exists.</param>
+		/// <param name="bMakePublic">Specified whether the file should be made public readable.</param>
+		abstract public Task<PostFileResult> PostFileAsync(string Container, string Identifier, string SourceFilePath, string ContentType = null, bool bOverwrite = true, bool bMakePublic = false);
 
 		/// <summary>
 		/// Deletes a file from cloud storage
@@ -820,6 +841,19 @@ namespace EpicGames.MCP.Automation
 		/// <param name="stagingInfo">Staging info representing the build to check.</param>
 		/// <returns>True if the manifest exists in cloud storage, false otherwise.</returns>
 		abstract public bool IsManifestOnCloudStorage(string Container, BuildPatchToolStagingInfo StagingInfo);
+
+		public class PostFileResult
+		{
+			/// <summary>
+			/// Set to the URL of the uploaded file on success
+			/// </summary>
+			public string ObjectURL { get; set; }
+
+			/// <summary>
+			/// Set to true if the write succeeds, false otherwise.
+			/// </summary>
+			public bool bSuccess { get; set; }
+		}
 	}
 }
 
