UE-20642 - removed ForceNonUnity, ForceMonolithic, and ForceDebugInfo from BuildTarget class as they weren't being used anywhere
#uat

diff --git a/Engine/Source/Programs/AutomationTool/AutomationUtils/UE4Build.cs b/Engine/Source/Programs/AutomationTool/AutomationUtils/UE4Build.cs
--- a/Engine/Source/Programs/AutomationTool/AutomationUtils/UE4Build.cs
+++ b/Engine/Source/Programs/AutomationTool/AutomationUtils/UE4Build.cs
@@ -599,15 +599,6 @@ namespace AutomationTool
 			/// Configuration to build
 			public UnrealBuildTool.UnrealTargetConfiguration Config;
 
-			/// If true, force monolithic
-			public bool ForceMonolithic;
-
-			/// If true, force non-unity
-			public bool ForceNonUnity;
-
-			/// If true, force debug info, even in development builds (useful for distributed builds where we need symbols)
-			public bool ForceDebugInfo;
-
 			/// Extra UBT args
 			public string UBTArgs;
 		}
@@ -657,11 +648,8 @@ namespace AutomationTool
 			/// <param name="InPlatform">Platform</param>
 			/// <param name="InConfiguration">Configuration</param>
 			/// <param name="InUprojectPath">Path to optional uproject file</param>
-			/// <param name="bForceMonolithic">Force monolithic build.</param>
-			/// <param name="bForceNonUnity">Force non-unity</param>
-			/// <param name="bForceDebugInfo">Force debug info even in development builds</param>
 			/// <param name="InAddArgs">Specifies additional arguments for UBT</param>
-			public void AddTarget(string TargetName, UnrealBuildTool.UnrealTargetPlatform InPlatform, UnrealBuildTool.UnrealTargetConfiguration InConfiguration, string InUprojectPath = null, bool bForceMonolithic = false, bool bForceNonUnity = false, bool bForceDebugInfo = false, string InAddArgs = "")
+			public void AddTarget(string TargetName, UnrealBuildTool.UnrealTargetPlatform InPlatform, UnrealBuildTool.UnrealTargetConfiguration InConfiguration, string InUprojectPath = null, string InAddArgs = "")
 			{
 				// Is this platform a compilable target?
 				if (!Platform.Platforms[InPlatform].CanBeCompiled())
@@ -675,9 +663,6 @@ namespace AutomationTool
 					Platform = InPlatform,
 					Config = InConfiguration,
 					UprojectPath = InUprojectPath,
-					ForceMonolithic = bForceMonolithic,
-					ForceNonUnity = bForceNonUnity,
-					ForceDebugInfo = bForceDebugInfo,
 					UBTArgs = InAddArgs,
 				});
 			}
@@ -689,10 +674,8 @@ namespace AutomationTool
 			/// <param name="InPlatform">Platform</param>
 			/// <param name="InConfiguration">Configuration</param>
 			/// <param name="InUprojectPath">Path to optional uproject file</param>
-			/// <param name="bForceMonolithic">Force monolithic build.</param>
-			/// <param name="bForceNonUnity">Force non-unity</param>
-			/// <param name="bForceDebugInfo">Force debug info even in development builds</param>
-			public void AddTargets(string[] TargetNames, UnrealBuildTool.UnrealTargetPlatform InPlatform, UnrealBuildTool.UnrealTargetConfiguration InConfiguration, string InUprojectPath = null, bool bForceMonolithic = false, bool bForceNonUnity = false, bool bForceDebugInfo = false, string InAddArgs = "")
+			/// <param name="InAddArgs">Specifies additional arguments for UBT</param>
+			public void AddTargets(string[] TargetNames, UnrealBuildTool.UnrealTargetPlatform InPlatform, UnrealBuildTool.UnrealTargetConfiguration InConfiguration, string InUprojectPath = null, string InAddArgs = "")
 			{
 				// Is this platform a compilable target?
 				if (!Platform.Platforms[InPlatform].CanBeCompiled())
@@ -708,9 +691,6 @@ namespace AutomationTool
 						Platform = InPlatform,
 						Config = InConfiguration,
 						UprojectPath = InUprojectPath,
-						ForceMonolithic = bForceMonolithic,
-						ForceNonUnity = bForceNonUnity,
-						ForceDebugInfo = bForceDebugInfo,
 						UBTArgs = InAddArgs,
 					});
 				}
