UAT: fix Linux packaging.

- Added check for a missing architecture parameter together with fallback to a hardcoded default.
- Ugly but a replacement is in the works.

diff --git a/Engine/Source/Programs/AutomationTool/AutomationUtils/ProjectParams.cs b/Engine/Source/Programs/AutomationTool/AutomationUtils/ProjectParams.cs
--- a/Engine/Source/Programs/AutomationTool/AutomationUtils/ProjectParams.cs
+++ b/Engine/Source/Programs/AutomationTool/AutomationUtils/ProjectParams.cs
@@ -589,6 +589,48 @@ namespace AutomationTool
 			this.DevicePassword = ParseParamValueIfNotSpecified(Command, DevicePassword, "devicepass", String.Empty);
 			this.CrashReporter = GetParamValueIfNotSpecified(Command, CrashReporter, this.CrashReporter, "crashreporter");
 			this.SpecifiedArchitecture = ParseParamValueIfNotSpecified(Command, SpecifiedArchitecture, "specifiedarchitecture", String.Empty);
+			if (String.IsNullOrEmpty(this.SpecifiedArchitecture))
+			{
+				// check if Linux is among of any platforms - it needs to know the default architecture
+				string OffendingConfigs = "";
+
+				if (this.ClientTargetPlatforms != null)
+				{
+					foreach (UnrealTargetPlatform ClientPlatform in this.ClientTargetPlatforms)
+					{
+						if (ClientPlatform == UnrealTargetPlatform.Linux)
+						{
+							OffendingConfigs = "client";
+							break;
+						}
+					}
+				}
+
+				if (this.ServerTargetPlatforms != null)
+				{
+					foreach (UnrealTargetPlatform ServerPlatform in this.ServerTargetPlatforms)
+					{
+						if (ServerPlatform == UnrealTargetPlatform.Linux)
+						{
+							OffendingConfigs += (OffendingConfigs == "") ? "server" : " and server";
+							break;
+						}
+					}
+				}
+
+				if (OffendingConfigs != "")
+				{
+					// [RCL] 2015-09-02 @FIXME: should get true DefaultArchitecture from the same place as UBT
+					const string DefaultLinuxArchitecture = "x86_64-unknown-linux-gnu";
+
+					CommandUtils.LogLog("Linux platform ({0}) requires a -specifiedarchitecture parameter, using default '{1}'",
+						OffendingConfigs, DefaultLinuxArchitecture
+						);
+
+					this.SpecifiedArchitecture = DefaultLinuxArchitecture;
+				}
+			}
+
 			if (ClientConfigsToBuild == null)
 			{
 				if (Command != null)
