Switch back to calling ubt on empty strings to fix linux packaging.

diff --git a/Engine/Source/Programs/AutomationTool/Scripts/CopyBuildToStagingDirectory.Automation.cs b/Engine/Source/Programs/AutomationTool/Scripts/CopyBuildToStagingDirectory.Automation.cs
--- a/Engine/Source/Programs/AutomationTool/Scripts/CopyBuildToStagingDirectory.Automation.cs
+++ b/Engine/Source/Programs/AutomationTool/Scripts/CopyBuildToStagingDirectory.Automation.cs
@@ -1544,7 +1544,7 @@ public partial class Project : CommandUtils
 					foreach (UnrealTargetPlatform ReceiptPlatform in SubPlatformsToStage)
 					{
                         string Architecture = Params.SpecifiedArchitecture;
-                        if (Architecture == null)
+                        if (string.IsNullOrEmpty(Architecture))
                         {
                             Architecture = "";
                             var BuildPlatform = UEBuildPlatform.GetBuildPlatform(ReceiptPlatform, true);
