Update access and modified dates of referenced manifests files during Rocket cleanup.

Adding a TouchFile method to CommandUtils as part of the solution.

diff --git a/Engine/Source/Programs/AutomationTool/AutomationUtils/CommandUtils.cs b/Engine/Source/Programs/AutomationTool/AutomationUtils/CommandUtils.cs
--- a/Engine/Source/Programs/AutomationTool/AutomationUtils/CommandUtils.cs
+++ b/Engine/Source/Programs/AutomationTool/AutomationUtils/CommandUtils.cs
@@ -924,6 +924,29 @@ namespace AutomationTool
 			return Result;
 		}
 
+		/// <summary>
+		/// Updates a file with the specified modified and access date, creating the file if it does not already exist.
+		/// An exception will be thrown if the directory does not already exist.
+		/// </summary>
+		/// <param name="Filename">The filename to touch, will be created if it does not exist.</param>
+		/// <param name="UtcDate">The accessed and modified date to set.  If not specified, defaults to the current date and time.</param>
+		public static void TouchFile(string Filename, DateTime? UtcDate = null)
+		{
+			var Date = UtcDate ?? DateTime.UtcNow;
+			Filename = ConvertSeparators(PathSeparator.Slash, Filename);
+			if (!File.Exists(Filename))
+			{
+				var Dir = GetDirectoryName(Filename);
+				if (!DirectoryExists_NoExceptions(Dir))
+				{
+					throw new AutomationException(new DirectoryNotFoundException("Directory not found: " + Dir), "Unable to create file {0} as directory does not exist.", Filename);
+				}
+				File.Create(Filename).Dispose();
+			}
+			File.SetLastAccessTimeUtc(Filename, Date);
+			File.SetLastWriteTimeUtc(Filename, Date);
+		}
+
 		/// <summary>
 		/// Sets file attributes. Will not change attributes that have not been specified.
 		/// </summary>
