Add the Tools node as a dependency of the BlueprintDocumentation node, because it may need to generate shaders.

diff --git a/Engine/Source/Programs/AutomationTool/Scripts/BuildDocumentation.Automation.cs b/Engine/Source/Programs/AutomationTool/Scripts/BuildDocumentation.Automation.cs
--- a/Engine/Source/Programs/AutomationTool/Scripts/BuildDocumentation.Automation.cs
+++ b/Engine/Source/Programs/AutomationTool/Scripts/BuildDocumentation.Automation.cs
@@ -114,6 +114,7 @@ class BlueprintDocumentationNode : DocumentationNode
 		: base(InHostPlatform)
 	{
 		AddDependency(GUBP.RootEditorNode.StaticGetFullName(InHostPlatform));
+		AddDependency(GUBP.ToolsNode.StaticGetFullName(InHostPlatform));
 	}
 
 	public override void DoBuild(GUBP bp)
