Fix UAT discovery of build folders in other projects.

diff --git a/Engine/Source/Programs/AutomationTool/AutomationUtils/ScriptCompiler.cs b/Engine/Source/Programs/AutomationTool/AutomationUtils/ScriptCompiler.cs
--- a/Engine/Source/Programs/AutomationTool/AutomationUtils/ScriptCompiler.cs
+++ b/Engine/Source/Programs/AutomationTool/AutomationUtils/ScriptCompiler.cs
@@ -129,19 +129,20 @@ namespace AutomationTool
 			{
 				AllGameFolders = new List<DirectoryReference>{ new DirectoryReference(Path.GetDirectoryName(ScriptsForProjectFileName)) };
 			}
-			var BuildFolders = new List<DirectoryReference>(AllGameFolders.Count);
+
+			var AllAdditionalScriptFolders = new List<DirectoryReference>(AdditionalScriptsFolders.Select(x => new DirectoryReference(x)));
 			foreach (var Folder in AllGameFolders)
 			{
 				var GameBuildFolder = DirectoryReference.Combine(Folder, "Build");
 				if (GameBuildFolder.Exists())
 				{
-					BuildFolders.Add(GameBuildFolder);
+					AllAdditionalScriptFolders.Add(GameBuildFolder);
 				}
 			}
 
 			Log.TraceVerbose("Discovering game folders.");
 
-			var DiscoveredModules = UnrealBuildTool.RulesCompiler.FindAllRulesSourceFiles(UnrealBuildTool.RulesCompiler.RulesFileType.AutomationModule, GameFolders: AllGameFolders, ForeignPlugins: null, AdditionalSearchPaths: AdditionalScriptsFolders.Select(x => new DirectoryReference(x)).ToList());
+			var DiscoveredModules = UnrealBuildTool.RulesCompiler.FindAllRulesSourceFiles(UnrealBuildTool.RulesCompiler.RulesFileType.AutomationModule, GameFolders: AllGameFolders, ForeignPlugins: null, AdditionalSearchPaths: AllAdditionalScriptFolders);
 			var ModulesToCompile = new List<string>(DiscoveredModules.Count);
 			foreach (var ModuleFilename in DiscoveredModules)
 			{
