All GUBP nodes for a job will use the same temp storage root determined by the branch it is in.
* This should keep storage created by a branch isolated to one location, and allow game-specific branches to control their own storage independently.
* Removed fallback code that looks for temp storage in UE4 folder if not found, since this doesn't really make sense anymore.
* Removed WasLocal out param from RetrieveFromTempStorage since no one was using it.
* Added a new branch hacker Dictionary: GameNameToBuildShareMapping, which essentially specifies, per-game, which build share to use. If one is not specified for that game, it falls back to the UE4 build share.
* Moved branch hacker options creation earlier in GUBP so the JobInfo can use it to determine the location for temp storage.

diff --git a/Engine/Source/Programs/AutomationTool/GUBP/ElectricCommander.cs b/Engine/Source/Programs/AutomationTool/GUBP/ElectricCommander.cs
--- a/Engine/Source/Programs/AutomationTool/GUBP/ElectricCommander.cs
+++ b/Engine/Source/Programs/AutomationTool/GUBP/ElectricCommander.cs
@@ -211,7 +211,7 @@ namespace AutomationTool
 			}
 		}
 
-        public void SaveStatus(TempStorageNodeInfo TempStorageNodeInfo, bool bSaveSharedTempStorage, string GameNameIfAny, string JobStepIDForFailure = null)
+        public void SaveStatus(TempStorageNodeInfo TempStorageNodeInfo, bool bSaveSharedTempStorage, string JobStepIDForFailure = null)
 		{
 			string Contents = "Just a status record: " + TempStorageNodeInfo.NodeStorageName;
 			if (!String.IsNullOrEmpty(JobStepIDForFailure) && CommandUtils.IsBuildMachine)
@@ -229,7 +229,7 @@ namespace AutomationTool
 			string RecordOfSuccess = CommandUtils.CombinePaths(CommandUtils.CmdEnv.LocalRoot, "Engine", "Saved", "Logs", TempStorageNodeInfo.NodeStorageName +".log");
 			CommandUtils.CreateDirectory(Path.GetDirectoryName(RecordOfSuccess));
 			CommandUtils.WriteAllText(RecordOfSuccess, Contents);
-            TempStorage.StoreToTempStorage(TempStorageNodeInfo, new List<string> { RecordOfSuccess }, !bSaveSharedTempStorage, GameNameIfAny, CommandUtils.CmdEnv.LocalRoot);		
+            TempStorage.StoreToTempStorage(TempStorageNodeInfo, new List<string> { RecordOfSuccess }, !bSaveSharedTempStorage, CommandUtils.CmdEnv.LocalRoot);		
 		}
 
 		public void UpdateEmailProperties(BuildNode NodeToDo, int LastSucceededCL, string FailedString, string FailCauserEMails, string EMailNote, bool SendSuccessForGreenAfterRed)
diff --git a/Engine/Source/Programs/AutomationTool/GUBP/GUBP.cs b/Engine/Source/Programs/AutomationTool/GUBP/GUBP.cs
--- a/Engine/Source/Programs/AutomationTool/GUBP/GUBP.cs
+++ b/Engine/Source/Programs/AutomationTool/GUBP/GUBP.cs
@@ -95,6 +95,10 @@ public partial class GUBP : BuildCommand
 
         bool bSaveSharedTempStorage = P4Enabled && (IsBuildMachine || GlobalCommandLine.UseLocalBuildStorage);
 
+        // Get the branch hacker options.
+        string BranchName = P4Enabled ? P4Env.BuildRootP4 : ParseParamValue("BranchName", "");
+        GUBPBranchHacker.BranchOptions BranchOptions = GetBranchOptions(BranchName);
+
         // encapsulate the logic and temp vars to determine the job info.
         // Ensures the temp vars don't escape the scope to be used later.
         Func<JobInfo> GetJobInfo = () =>
@@ -114,7 +118,9 @@ public partial class GUBP : BuildCommand
                 PreflightUID = ParseParamInt("PreflightUID", 0);
             }
             return new JobInfo(
+                BranchName,
                 P4Enabled ? P4Env.BuildRootEscaped : "NoP4",
+                BranchOptions.RootNameForTempStorage ?? "UE4",
                 P4Enabled ? ParseParamInt("CL", P4Env.Changelist) : 0,
                 PreflightShelveCL, 
                 PreflightUID);
@@ -147,7 +153,8 @@ public partial class GUBP : BuildCommand
 
 		int TimeQuantum = 20;
 		BuildGraphTemplate GraphTemplate = new BuildGraphTemplate();
-		AddNodesForBranch(HostPlatforms, JobInfo, out GraphTemplate.BuildNodeTemplates, out GraphTemplate.AggregateNodeTemplates, ref TimeQuantum);
+
+        AddNodesForBranch(HostPlatforms, JobInfo, BranchOptions, out GraphTemplate.BuildNodeTemplates, out GraphTemplate.AggregateNodeTemplates, ref TimeQuantum);
 
 		BuildGraph Graph = new BuildGraph(GraphTemplate);
 		FindCompletionState(Graph.BuildNodes, JobInfo, LocalOnly);
@@ -287,29 +294,23 @@ public partial class GUBP : BuildCommand
             // Construct the full temp storage node info
             var TempStorageNodeInfo = new TempStorageNodeInfo(JobInfo, NodeToDo.Name);
 
-			string GameNameIfAny = NodeToDo.GameNameIfAnyForTempStorage;
-
 			if (LocalOnly)
 			{
                 NodeToDo.IsComplete = TempStorage.LocalTempStorageManifestExists(TempStorageNodeInfo, bQuiet: true);
 			}
 			else
 			{
-                NodeToDo.IsComplete = TempStorage.TempStorageExists(TempStorageNodeInfo, GameNameIfAny, bLocalOnly: false, bQuiet: true);
-				if(GameNameIfAny != "" && !NodeToDo.IsComplete)
-				{
-                    NodeToDo.IsComplete = TempStorage.TempStorageExists(TempStorageNodeInfo, "UE4", bLocalOnly: false, bQuiet: true);
-				}
+                NodeToDo.IsComplete = TempStorage.TempStorageExists(TempStorageNodeInfo, bLocalOnly: false, bQuiet: true);
 			}
 
             LogVerbose("** {0}", NodeToDo.Name);
 			if (!NodeToDo.IsComplete)
 			{
-                LogVerbose("***** GUBP Trigger Node was already triggered {0} -> {1} : {2}", NodeToDo.Name, GameNameIfAny, TempStorageNodeInfo.GetRelativeDirectory());
+                LogVerbose("***** GUBP Trigger Node was already triggered {0} -> {1} : {2}", NodeToDo.Name, JobInfo.RootNameForTempStorage, TempStorageNodeInfo.GetRelativeDirectory());
 			}
 			else
 			{
-                LogVerbose("***** GUBP Trigger Node was NOT yet triggered {0} -> {1} : {2}", NodeToDo.Name, GameNameIfAny, TempStorageNodeInfo.GetRelativeDirectory());
+                LogVerbose("***** GUBP Trigger Node was NOT yet triggered {0} -> {1} : {2}", NodeToDo.Name, JobInfo.RootNameForTempStorage, TempStorageNodeInfo.GetRelativeDirectory());
 			}
 		}
     }
@@ -761,13 +762,11 @@ public partial class GUBP : BuildCommand
         // Don't get node history on nodes that don't have a valid CL.
         if (!(NodeToDo is TriggerNode) && JobInfo.Changelist > 0)
         {
-            string GameNameIfAny = NodeToDo.GameNameIfAnyForTempStorage;
-
             History = new NodeHistory
             {
-                AllStarted = TempStorage.FindMatchingSharedTempStorageNodeCLs(new TempStorageNodeInfo(JobInfo, NodeToDo.Name + StartedTempStorageSuffix), GameNameIfAny),
-                AllSucceeded = TempStorage.FindMatchingSharedTempStorageNodeCLs(new TempStorageNodeInfo(JobInfo, NodeToDo.Name + SucceededTempStorageSuffix), GameNameIfAny),
-                AllFailed = TempStorage.FindMatchingSharedTempStorageNodeCLs(new TempStorageNodeInfo(JobInfo, NodeToDo.Name + FailedTempStorageSuffix), GameNameIfAny)
+                AllStarted = TempStorage.FindMatchingSharedTempStorageNodeCLs(new TempStorageNodeInfo(JobInfo, NodeToDo.Name + StartedTempStorageSuffix)),
+                AllSucceeded = TempStorage.FindMatchingSharedTempStorageNodeCLs(new TempStorageNodeInfo(JobInfo, NodeToDo.Name + SucceededTempStorageSuffix)),
+                AllFailed = TempStorage.FindMatchingSharedTempStorageNodeCLs(new TempStorageNodeInfo(JobInfo, NodeToDo.Name + FailedTempStorageSuffix))
             };
 
             if (History.AllFailed.Count > 0)
@@ -908,7 +907,6 @@ public partial class GUBP : BuildCommand
     {
         int Result = P4Env.Changelist;
 
-        string GameNameIfAny = NodeToDo.GameNameIfAnyForTempStorage;
         var TempStorageNodeInfo = new TempStorageNodeInfo(JobInfo, NodeToDo.Name + FailedTempStorageSuffix);
 
         List<int> BackwardsFails = new List<int>(History.AllFailed);
@@ -933,8 +931,7 @@ public partial class GUBP : BuildCommand
             List<string> Files = null;
             try
             {
-                bool WasLocal;
-                Files = TempStorage.RetrieveFromTempStorage(ThisTempStorageNodeInfo, out WasLocal, GameNameIfAny, CmdEnv.LocalRoot); // this will fail on our CL if we didn't fail or we are just setting up the branch
+                Files = TempStorage.RetrieveFromTempStorage(ThisTempStorageNodeInfo, CmdEnv.LocalRoot); // this will fail on our CL if we didn't fail or we are just setting up the branch
             }
             catch (Exception)
             {
@@ -1276,7 +1273,6 @@ public partial class GUBP : BuildCommand
 
             var TempStorageNodeInfo = new TempStorageNodeInfo(JobInfo, NodeToDo.Name);
             
-			string GameNameIfAny = NodeToDo.GameNameIfAnyForTempStorage;
 			string StorageRootIfAny = NodeToDo.RootIfAnyForTempStorage;
 
 			if (bFake)
@@ -1291,8 +1287,8 @@ public partial class GUBP : BuildCommand
             if (NodeToDo.IsComplete)
             {
 				// Just fetch the build products from temp storage
-				Log("***** Retrieving GUBP Node {0} -> {1} : {2}", NodeToDo.Name, GameNameIfAny, TempStorageNodeInfo.GetRelativeDirectory());
-				NodeToDo.RetrieveBuildProducts(GameNameIfAny, StorageRootIfAny, TempStorageNodeInfo);
+				Log("***** Retrieving GUBP Node {0} -> {1} : {2}", NodeToDo.Name, JobInfo.RootNameForTempStorage, TempStorageNodeInfo.GetRelativeDirectory());
+				NodeToDo.RetrieveBuildProducts(StorageRootIfAny, TempStorageNodeInfo);
             }
             else
             {
@@ -1303,7 +1299,7 @@ public partial class GUBP : BuildCommand
                 // Record that the node has started. We save our status to a new temp storage location specifically named with a suffix so we can find it later.
                 if (SaveSuccessRecords) 
                 {
-                    EC.SaveStatus(new TempStorageNodeInfo(JobInfo, NodeToDo.Name + StartedTempStorageSuffix), bSaveSharedTempStorage, GameNameIfAny);
+                    EC.SaveStatus(new TempStorageNodeInfo(JobInfo, NodeToDo.Name + StartedTempStorageSuffix), bSaveSharedTempStorage);
                 }
 
 				// Execute the node
@@ -1314,7 +1310,7 @@ public partial class GUBP : BuildCommand
 				// Archive the build products if the node succeeded
 				if(bResult)
 				{
-					NodeToDo.ArchiveBuildProducts(GameNameIfAny, StorageRootIfAny, TempStorageNodeInfo, !bSaveSharedTempStorage);
+					NodeToDo.ArchiveBuildProducts(StorageRootIfAny, TempStorageNodeInfo, !bSaveSharedTempStorage);
 				}
 
 				// Record that the node has finished
@@ -1327,7 +1323,7 @@ public partial class GUBP : BuildCommand
 					}
 					using(TelemetryStopwatch SaveNodeStatusStopwatch = new TelemetryStopwatch("SaveNodeStatus"))
 					{
-                        EC.SaveStatus(new TempStorageNodeInfo(JobInfo, NodeToDo.Name + (bResult? SucceededTempStorageSuffix : FailedTempStorageSuffix)), bSaveSharedTempStorage, GameNameIfAny, ParseParamValue("MyJobStepId"));
+                        EC.SaveStatus(new TempStorageNodeInfo(JobInfo, NodeToDo.Name + (bResult? SucceededTempStorageSuffix : FailedTempStorageSuffix)), bSaveSharedTempStorage, ParseParamValue("MyJobStepId"));
 					}
 					using(TelemetryStopwatch UpdateECPropsStopwatch = new TelemetryStopwatch("UpdateECProps"))
 					{
diff --git a/Engine/Source/Programs/AutomationTool/GUBP/JobInfo.cs b/Engine/Source/Programs/AutomationTool/GUBP/JobInfo.cs
--- a/Engine/Source/Programs/AutomationTool/GUBP/JobInfo.cs
+++ b/Engine/Source/Programs/AutomationTool/GUBP/JobInfo.cs
@@ -17,6 +17,16 @@ namespace AutomationTool
     /// </remarks>
     public class JobInfo
     {
+        /// <summary>
+        /// BranchName for this job. P4Env.BuildRootP4 if P4Enabled. Otherwise, uses the -BranchName= parameter. Can be empty string.
+        /// </summary>
+        public string BranchName { get; private set; }
+
+        /// <summary>
+        /// The root name beneath <see cref="CommandUtils.RootBuildStorageDirectory"/> where all temp storage will placed for this job.
+        /// </summary>
+        public string RootNameForTempStorage { get; private set; }
+
         /// <summary>
         /// Used as the prefix for temp storage for the job.
         /// 
@@ -70,19 +80,23 @@ namespace AutomationTool
         /// <returns>A new JobInfo instance with all the same attributes as the current instance, but a different CL.</returns>
         public JobInfo CreateWithNewChangelist(int NewCL)
         {
-            return new JobInfo(BranchNameForTempStorage, NewCL, PreflightShelveChangelist, PreflightUID);
+            return new JobInfo(BranchName, RootNameForTempStorage, BranchNameForTempStorage, NewCL, PreflightShelveChangelist, PreflightUID);
         }
 
         /// <summary>
         /// Public ctor. This is an immutable class, so all attributes must be set a construction time.
         /// </summary>
+        /// <param name="BranchName">See similarly named attribute.</param>
+        /// <param name="RootNameForTempStorage">See similarly named attribute.</param>
         /// <param name="BranchNameForTempStorage">See similarly named attribute.</param>
         /// <param name="Changelist">See similarly named attribute.</param>
         /// <param name="PreflightShelveChangelist">See similarly named attribute.</param>
         /// <param name="PreflightUID">See similarly named attribute.</param>
-        public JobInfo(string BranchNameForTempStorage, int Changelist, int PreflightShelveChangelist, int PreflightUID)
+        public JobInfo(string BranchName, string BranchNameForTempStorage, string RootNameForTempStorage, int Changelist, int PreflightShelveChangelist, int PreflightUID)
         {
+            this.BranchName = BranchName;
             this.BranchNameForTempStorage = BranchNameForTempStorage;
+            this.RootNameForTempStorage = RootNameForTempStorage;
             this.Changelist = Changelist;
             this.PreflightShelveChangelist = PreflightShelveChangelist;
             this.PreflightUID = PreflightUID;
diff --git a/Engine/Source/Programs/AutomationTool/GUBP/LegacyBranchSetup.cs b/Engine/Source/Programs/AutomationTool/GUBP/LegacyBranchSetup.cs
--- a/Engine/Source/Programs/AutomationTool/GUBP/LegacyBranchSetup.cs
+++ b/Engine/Source/Programs/AutomationTool/GUBP/LegacyBranchSetup.cs
@@ -14,7 +14,7 @@ partial class GUBP
 	public class GUBPBranchConfig
 	{
 		public List<UnrealTargetPlatform> HostPlatforms;
-		public string BranchName;
+        public string BranchName { get { return JobInfo.BranchName; } }
 		public BranchInfo Branch;
 		public GUBPBranchHacker.BranchOptions BranchOptions;
 		public bool bForceIncrementalCompile;
@@ -23,10 +23,9 @@ partial class GUBP
 		public Dictionary<string, GUBPNode> GUBPNodes = new Dictionary<string, GUBPNode>();
 		public Dictionary<string, GUBPAggregateNode> GUBPAggregates = new Dictionary<string, GUBPAggregateNode>();
 
-        public GUBPBranchConfig(IEnumerable<UnrealTargetPlatform> InHostPlatforms, string InBranchName, BranchInfo InBranch, GUBPBranchHacker.BranchOptions InBranchOptions, bool bInForceIncrementalCompile, JobInfo JobInfo)
+        public GUBPBranchConfig(IEnumerable<UnrealTargetPlatform> InHostPlatforms, BranchInfo InBranch, GUBPBranchHacker.BranchOptions InBranchOptions, bool bInForceIncrementalCompile, JobInfo JobInfo)
 		{
 			HostPlatforms = new List<UnrealTargetPlatform>(InHostPlatforms);
-			BranchName = InBranchName;
 			Branch = InBranch;
 			BranchOptions = InBranchOptions;
 			bForceIncrementalCompile = bInForceIncrementalCompile;
@@ -160,6 +159,17 @@ partial class GUBP
 			public bool bNoEditorDependenciesForTools = false;
 			public Dictionary<string, sbyte> FrequencyBarriers = new Dictionary<string,sbyte>();
 			public int QuantumOverride = 0;
+            /// <summary>
+            /// Allows a branch to override the build share it uses for temp storage for all nodes. This will override the default value for <see cref="JobInfo.RootNameForTempStorage"/>.
+            /// </summary>
+            public string RootNameForTempStorage;
+            /// <summary>
+            /// Stores a mapping from a GameName (determined by the <see cref="BranchInfo.BranchUProject.GameName"/>) to a share name underneath <see cref="CommandUtils.RootBuildStorageDirectory"/>.
+            /// This is where formal builds and such will be stored for each game. Instead of tying formal build outputs to the same location as temp storage,
+            /// this allows a project to have it's own share for formal build products, but use, say, the UE4 default storage location for temp storage.
+            /// If the mapping does not exist for a game, it uses the UE4 folder.
+            /// </summary>
+            public readonly Dictionary<string, string> GameNameToBuildShareMapping = new Dictionary<string,string>();
         }
         public virtual void ModifyOptions(GUBP bp, ref BranchOptions Options, string Branch)
         {
@@ -345,20 +355,8 @@ partial class GUBP
         return AltHostPlatform;
     }
 
-    void AddNodesForBranch(List<UnrealTargetPlatform> InitialHostPlatforms, JobInfo JobInfo, out List<BuildNodeTemplate> AllNodeDefinitions, out List<AggregateNodeTemplate> AllAggregateDefinitions, ref int TimeQuantum)
+    void AddNodesForBranch(List<UnrealTargetPlatform> InitialHostPlatforms, JobInfo JobInfo, GUBPBranchHacker.BranchOptions BranchOptions, out List<BuildNodeTemplate> AllNodeDefinitions, out List<AggregateNodeTemplate> AllAggregateDefinitions, ref int TimeQuantum)
 	{
-		string BranchName;
-		if (P4Enabled)
-		{
-			BranchName = P4Env.BuildRootP4;
-		}
-		else
-		{ 
-			BranchName = ParseParamValue("BranchName", "");
-		}
-
-		GUBPBranchHacker.BranchOptions BranchOptions = GetBranchOptions(BranchName);
-
 		if (BranchOptions.QuantumOverride != 0)
 		{
 			TimeQuantum = BranchOptions.QuantumOverride;
@@ -453,7 +451,7 @@ partial class GUBP
             }
         }
 
-		GUBPBranchConfig BranchConfig = new GUBPBranchConfig(HostPlatforms, BranchName, Branch, BranchOptions, bForceIncrementalCompile, JobInfo);
+		GUBPBranchConfig BranchConfig = new GUBPBranchConfig(HostPlatforms, Branch, BranchOptions, bForceIncrementalCompile, JobInfo);
 
         BranchConfig.AddNode(new VersionFilesNode());
 
@@ -1173,7 +1171,7 @@ partial class GUBP
 			List<string> NodeNames = new List<string>();
 			foreach(GUBP.GUBPNode Node in BranchConfig.GUBPNodes.Values)
 			{
-				if (Node.GameNameIfAnyForTempStorage() == GameProj.GameName)
+				if (Node.GameNameIfAnyForFullGameAggregateNode() == GameProj.GameName)
 				{
 					NodeNames.Add(Node.GetFullName());
 				}
@@ -1201,7 +1199,7 @@ partial class GUBP
 		FindFrequenciesForNodes(BranchConfig, LegacyToNewNodes, FrequencyOverrides);
 
 		// Get the email list for each node
-		FindEmailsForNodes(BranchName, AllNodeDefinitions);
+		FindEmailsForNodes(BranchConfig.BranchName, AllNodeDefinitions);
 	}
 
 	private static Dictionary<string, int> ApplyFrequencyBarriers(Dictionary<string, GUBPNode> GUBPNodes, Dictionary<string, GUBPAggregateNode> GUBPAggregates, Dictionary<string, sbyte> FrequencyBarriers)
diff --git a/Engine/Source/Programs/AutomationTool/GUBP/LegacyNodes.cs b/Engine/Source/Programs/AutomationTool/GUBP/LegacyNodes.cs
--- a/Engine/Source/Programs/AutomationTool/GUBP/LegacyNodes.cs
+++ b/Engine/Source/Programs/AutomationTool/GUBP/LegacyNodes.cs
@@ -28,7 +28,12 @@ partial class GUBP
 		{
 			return GetFullName();
 		}
-        public virtual string GameNameIfAnyForTempStorage()
+        /// <summary>
+        /// This function is legacy and deprecated for new usage. It purely exists to maintain an
+        /// older hack to determine which <see cref="FullGameAggregateNode"/>s to add.
+        /// </summary>
+        /// <returns></returns>
+        public virtual string GameNameIfAnyForFullGameAggregateNode()
         {
             return "";
         }
@@ -926,10 +931,10 @@ partial class GUBP
 			}
 			return Template;
 		}
-		public override string GameNameIfAnyForTempStorage()
+		public override string GameNameIfAnyForFullGameAggregateNode()
         {
             return GameProjects[0].Options(HostPlatform).GroupName ?? GameProjects[0].GameName;
-            }                    
+        }                    
         public override UE4Build.BuildAgenda GetAgenda(GUBP bp)
         {
             var Agenda = new UE4Build.BuildAgenda();
@@ -1172,7 +1177,7 @@ partial class GUBP
         {
             return StaticGetFullName(HostPlatform, GameProj, TargetPlatform, WithXp, Precompiled);
         }
-        public override string GameNameIfAnyForTempStorage()
+        public override string GameNameIfAnyForFullGameAggregateNode()
         {
             return GameProj.GameName;
         }
@@ -1410,7 +1415,7 @@ partial class GUBP
 
 		public abstract string GetFullName();
 
-		public virtual string GameNameIfAnyForTempStorage()
+        public virtual string GameNameIfAnyForFullGameAggregateNode()
 		{
 			return "";
 		}
@@ -1654,7 +1659,7 @@ partial class GUBP
         {
             return StaticGetFullName(HostPlatform, GameProj, CookPlatform);
         }
-        public override string GameNameIfAnyForTempStorage()
+        public override string GameNameIfAnyForFullGameAggregateNode()
         {
             return GameProj.GameName;
         }
@@ -1769,7 +1774,7 @@ partial class GUBP
         {
             return StaticGetFullName(HostPlatform, GameProj, TargetPlatform);
         }
-        public override string GameNameIfAnyForTempStorage()
+        public override string GameNameIfAnyForFullGameAggregateNode()
         {
             return GameProj.GameName;
         }
@@ -1944,7 +1949,7 @@ partial class GUBP
         {
             return StaticGetFullName(GameProj, HostPlatform, ClientTargetPlatforms, ClientConfigs, ServerTargetPlatforms, ServerConfigs, ClientNotGame);
         }
-        public override string GameNameIfAnyForTempStorage()
+        public override string GameNameIfAnyForFullGameAggregateNode()
         {
             return GameProj.GameName;
         }
@@ -1980,7 +1985,13 @@ partial class GUBP
 		}
 		public static string GetArchiveDirectory(GUBP.GUBPBranchConfig BranchConfig, BranchInfo.BranchUProject InGameProj, UnrealTargetPlatform InHostPlatform, List<UnrealTargetPlatform> InClientTargetPlatforms = null, List<UnrealTargetConfiguration> InClientConfigs = null, List<UnrealTargetPlatform> InServerTargetPlatforms = null, List<UnrealTargetConfiguration> InServerConfigs = null, bool InClientNotGame = false)
         {
-            string BaseDir = TempStorage.ResolveSharedBuildDirectory(InGameProj.GameName);
+            // Find the build share where formal builds will be placed for this game.
+            string BuildShareName;
+            if (!BranchConfig.BranchOptions.GameNameToBuildShareMapping.TryGetValue(InGameProj.GameName, out BuildShareName))
+            {
+                BuildShareName = "UE4";
+            }
+            string BaseDir = CommandUtils.CombinePaths(CommandUtils.RootBuildStorageDirectory(), BuildShareName);
             string NodeName = StaticGetFullName(InGameProj, InHostPlatform, InClientTargetPlatforms, InClientConfigs, InServerTargetPlatforms, InServerConfigs, InClientNotGame);
             string Inner = P4Env.BuildRootEscaped + "-CL-" + P4Env.ChangelistString + BranchConfig.JobInfo.GetPreflightSuffix();
             string ArchiveDirectory = CombinePaths(BaseDir, NodeName, Inner);
@@ -2212,7 +2223,7 @@ partial class GUBP
         {
             return StaticGetFullName(GameProj, HostPlatform, ClientTargetPlatform, ClientConfig);
         }
-        public override string GameNameIfAnyForTempStorage()
+        public override string GameNameIfAnyForFullGameAggregateNode()
         {
             return GameProj.GameName;
         }
@@ -2586,7 +2597,7 @@ partial class GUBP
         {
             return StaticGetFullName(HostPlatform, GameProj, TestName);
         }
-        public override string GameNameIfAnyForTempStorage()
+        public override string GameNameIfAnyForFullGameAggregateNode()
         {
             return GameProj.GameName;
         }
@@ -2666,7 +2677,7 @@ partial class GUBP
         {
             return StaticGetFullName(HostPlatform, GameProj, AggregateName);
         }
-        public override string GameNameIfAnyForTempStorage()
+        public override string GameNameIfAnyForFullGameAggregateNode()
         {
             return GameProj.GameName;
         }
@@ -2674,12 +2685,12 @@ partial class GUBP
 
     public class CleanSharedTempStorageNode : GUBPNode
     {
-		HashSet<string> GameNames;
+        string RootNameForTempStorage;
 
         public CleanSharedTempStorageNode(GUBP bp, GUBPBranchConfig BranchConfig)
         {
             var ToolsNode = BranchConfig.FindNode(ToolsForCompileNode.StaticGetFullName(UnrealTargetPlatform.Win64));
-			GameNames = new HashSet<string>(BranchConfig.GUBPNodes.Values.Select(x => x.GameNameIfAnyForTempStorage()));
+            RootNameForTempStorage = BranchConfig.JobInfo.RootNameForTempStorage;
             AgentSharingGroup = ToolsNode.AgentSharingGroup;
         }
         public override float Priority()
@@ -2699,7 +2710,7 @@ partial class GUBP
         {
             {
                 var StartTime = DateTime.UtcNow;
-                TempStorage.CleanSharedTempStorageDirectory(GameNames);
+                TempStorage.CleanSharedTempStorageDirectory(RootNameForTempStorage);
                 var BuildDuration = (DateTime.UtcNow - StartTime).TotalMilliseconds;
                 Log("Took {0}s to clear temp storage of old files.", BuildDuration / 1000);
             }
diff --git a/Engine/Source/Programs/AutomationTool/GUBP/Nodes/BuildNode.cs b/Engine/Source/Programs/AutomationTool/GUBP/Nodes/BuildNode.cs
--- a/Engine/Source/Programs/AutomationTool/GUBP/Nodes/BuildNode.cs
+++ b/Engine/Source/Programs/AutomationTool/GUBP/Nodes/BuildNode.cs
@@ -27,7 +27,7 @@ namespace AutomationTool
 		public bool IsSticky;
 		public bool IsTest;
 		public string DisplayGroupName;
-		public string GameNameIfAnyForTempStorage;
+        public string GameNameIfAnyForFullGameAggregateNode;
 		public string RootIfAnyForTempStorage;
 
 		public abstract BuildNode Instantiate();
@@ -56,7 +56,7 @@ namespace AutomationTool
 		public bool IsSticky;
 		public bool IsTest;
 		public string DisplayGroupName;
-		public string GameNameIfAnyForTempStorage;
+        public string GameNameIfAnyForFullGameAggregateNode;
 		public string RootIfAnyForTempStorage;
 
 		public List<string> BuildProducts;
@@ -78,34 +78,25 @@ namespace AutomationTool
 			IsSticky = Template.IsSticky;
 			IsTest = Template.IsTest;
 			DisplayGroupName = Template.DisplayGroupName;
-			GameNameIfAnyForTempStorage = Template.GameNameIfAnyForTempStorage;
+            GameNameIfAnyForFullGameAggregateNode = Template.GameNameIfAnyForFullGameAggregateNode;
 			RootIfAnyForTempStorage = Template.RootIfAnyForTempStorage;
 		}
 
-		public virtual void ArchiveBuildProducts(string GameNameIfAny, string StorageRootIfAny, TempStorageNodeInfo TempStorageNodeInfo, bool bLocalOnly)
+		public virtual void ArchiveBuildProducts(string StorageRootIfAny, TempStorageNodeInfo TempStorageNodeInfo, bool bLocalOnly)
 		{
-			TempStorage.StoreToTempStorage(TempStorageNodeInfo, BuildProducts, bLocalOnly, GameNameIfAny, StorageRootIfAny);
+			TempStorage.StoreToTempStorage(TempStorageNodeInfo, BuildProducts, bLocalOnly, StorageRootIfAny);
 		}
 
-		public virtual void RetrieveBuildProducts(string GameNameIfAny, string StorageRootIfAny, TempStorageNodeInfo TempStorageNodeInfo)
+		public virtual void RetrieveBuildProducts(string StorageRootIfAny, TempStorageNodeInfo TempStorageNodeInfo)
 		{
 			CommandUtils.Log("***** Retrieving GUBP Node {0} from {1}", Name, TempStorageNodeInfo.GetRelativeDirectory());
 			try
 			{
-				bool WasLocal;
-				BuildProducts = TempStorage.RetrieveFromTempStorage(TempStorageNodeInfo, out WasLocal, GameNameIfAny, StorageRootIfAny);
+				BuildProducts = TempStorage.RetrieveFromTempStorage(TempStorageNodeInfo, StorageRootIfAny);
 			}
 			catch (Exception Ex)
 			{
-				if (GameNameIfAny != "")
-				{
-					bool WasLocal;
-					BuildProducts = TempStorage.RetrieveFromTempStorage(TempStorageNodeInfo, out WasLocal, "", StorageRootIfAny);
-				}
-				else
-				{
-					throw new AutomationException(Ex, "Build Products cannot be found for node {0}", Name);
-				}
+				throw new AutomationException(Ex, "Build Products cannot be found for node {0}", Name);
 			}
 		}
 
diff --git a/Engine/Source/Programs/AutomationTool/GUBP/Nodes/LegacyNode.cs b/Engine/Source/Programs/AutomationTool/GUBP/Nodes/LegacyNode.cs
--- a/Engine/Source/Programs/AutomationTool/GUBP/Nodes/LegacyNode.cs
+++ b/Engine/Source/Programs/AutomationTool/GUBP/Nodes/LegacyNode.cs
@@ -34,7 +34,7 @@ namespace AutomationTool
 			Definition.OrderDependencyNames = String.Join(";", Node.FullNamesOfPseudodependencies);
 			Definition.IsTest = Node.IsTest();
 			Definition.DisplayGroupName = Node.GetDisplayGroupName();
-			Definition.GameNameIfAnyForTempStorage = Node.GameNameIfAnyForTempStorage();
+            Definition.GameNameIfAnyForFullGameAggregateNode = Node.GameNameIfAnyForFullGameAggregateNode();
 			Definition.RootIfAnyForTempStorage = Node.RootIfAnyForTempStorage();
 		}
 
@@ -56,9 +56,9 @@ namespace AutomationTool
 			Node = Template.Node;
 		}
 
-		public override void RetrieveBuildProducts(string GameNameIfAny, string StorageRootIfAny, TempStorageNodeInfo TempStorageNodeInfo)
+		public override void RetrieveBuildProducts(string StorageRootIfAny, TempStorageNodeInfo TempStorageNodeInfo)
 		{
-			base.RetrieveBuildProducts(GameNameIfAny, StorageRootIfAny, TempStorageNodeInfo);
+			base.RetrieveBuildProducts(StorageRootIfAny, TempStorageNodeInfo);
 
 			Node.BuildProducts = BuildProducts;
 		}
diff --git a/Engine/Source/Programs/AutomationTool/GUBP/TempStorage.cs b/Engine/Source/Programs/AutomationTool/GUBP/TempStorage.cs
--- a/Engine/Source/Programs/AutomationTool/GUBP/TempStorage.cs
+++ b/Engine/Source/Programs/AutomationTool/GUBP/TempStorage.cs
@@ -494,6 +494,9 @@ namespace AutomationTool
 
         /// <summary>
         /// Cleans all temp storage data from the given directory that is older than a certain threshold (defined directly in the code).
+        /// This will clean up builds that are NOT part of the the job for this branch. This is intentional because once a branch is shut down,
+        /// there would be nothing left to clean it. This will only clean the supplied folder, but it is assumed this folder is an actual
+        /// Temp Storage folder with .TempManifest files in it.
         /// </summary>
         /// <param name="TopDirectory">Fully qualified path of the folder that is the root of a bunch of temp storage folders. Should be a string returned by <see cref="ResolveSharedTempStorageDirectory"/> </param>
         private static void CleanSharedTempStorage(string TopDirectory)
@@ -524,7 +527,6 @@ namespace AutomationTool
             {
                 const double MaximumDaysToKeepTempStorage = 3;
                 var Now = DateTime.UtcNow;
-                // This will search legacy folders as well, but those should go away within a few days.
                 foreach (var OldManifestInfo in
                     // First subdirectory is the branch name
                     from BranchDir in Directory.EnumerateDirectories(TopDirectory)
@@ -566,27 +568,24 @@ namespace AutomationTool
         }
 
         /// <summary>
-        /// Cleans the shared temp storage folders for all the games listed. This code will ensure that any duplicate games
+        /// Cleans the shared temp storage folder for the given root name. This code will ensure that any duplicate games
         /// or games that resolve to use the same temp folder are not cleaned twice.
         /// </summary>
-        /// <param name="GameNames">List of game names to clean temp folders for.</param>
-        public static void CleanSharedTempStorageDirectory(IEnumerable<string> GameNames)
+        /// <param name="RootNameForTempStorage">List of game names to clean temp folders for.</param>
+        public static void CleanSharedTempStorageDirectory(string RootNameForTempStorage)
         {
             if (!CommandUtils.IsBuildMachine || Utils.IsRunningOnMono)  // saw a hang on this, anyway it isn't necessary to clean with macs, they are slow anyway
             {
                 return;
             }
-            // Generate a unique set of folder names to clean so we don't clean folders twice.
-            foreach (var TempStorageFolder in new HashSet<string>(GameNames.Select(ResolveSharedTempStorageDirectory)))
+            var TempStorageFolder = ResolveSharedTempStorageDirectory(RootNameForTempStorage);
+            try
             {
-                try
-                {
-                    CleanSharedTempStorage(TempStorageFolder);
-                }
-                catch (Exception Ex)
-                {
-                    CommandUtils.LogWarning("Unable to Clean Temp Directory {0}. Exception: {1}", TempStorageFolder, Ex);
-                }
+                CleanSharedTempStorage(TempStorageFolder);
+            }
+            catch (Exception Ex)
+            {
+                CommandUtils.LogWarning("Unable to Clean Temp Directory {0}. Exception: {1}", TempStorageFolder, Ex);
             }
         }
 
@@ -634,51 +633,36 @@ namespace AutomationTool
         /// If the shared build folder does not already exist, will try to use UE4's folder instead.
         /// This is because the root build folders are independently owned by each game team and managed by IT, so cannot be created on the fly.
         /// </summary>
-        /// <param name="GameName">GameName to look for the build folder for. If empty or the GameName folder cannot be found, uses the UE4 folder (GameName = UE4)</param>
+        /// The root name beneath <see cref="CommandUtils.RootBuildStorageDirectory"/> where all temp storage will placed for this job.
         /// <returns>The full path of the shared build directory name, or throws an exception if none is found and the UE4 folder is not found either.</returns>
-        public static string ResolveSharedBuildDirectory(string GameName)
+        public static string ResolveSharedBuildDirectory(string RootNameForTempStorage)
         {
-            if (GameName == null) throw new ArgumentNullException("GameName");
+            if (RootNameForTempStorage == null) throw new ArgumentNullException("RootNameForTempStorage");
 
-            if (SharedBuildDirectoryResolveCache.ContainsKey(GameName))
+            if (SharedBuildDirectoryResolveCache.ContainsKey(RootNameForTempStorage))
             {
-                return SharedBuildDirectoryResolveCache[GameName];
+                return SharedBuildDirectoryResolveCache[RootNameForTempStorage];
             }
             string Root = CommandUtils.RootBuildStorageDirectory();
-            string Result = CommandUtils.CombinePaths(Root, GameName);
-            if (GameName == "" || string.Equals(GameName, "ShooterGame", StringComparison.InvariantCultureIgnoreCase) || !InternalUtils.Robust_DirectoryExistsAndIsWritable_NoExceptions(Result))
+            string Result = CommandUtils.CombinePaths(Root, RootNameForTempStorage);
+            if (!InternalUtils.Robust_DirectoryExistsAndIsWritable_NoExceptions(Result))
             {
-                string GameStr = "Game";
-                bool HadGame = false;
-                if (GameName.EndsWith(GameStr, StringComparison.InvariantCultureIgnoreCase))
-                {
-                    string ShortFolder = GameName.Substring(0, GameName.Length - GameStr.Length);
-                    Result = CommandUtils.CombinePaths(Root, ShortFolder);
-                    HadGame = true;
-                }
-                if (!HadGame || !InternalUtils.Robust_DirectoryExistsAndIsWritable_NoExceptions(Result))
-                {
-                    Result = CommandUtils.CombinePaths(Root, "UE4");
-                    if (!InternalUtils.Robust_DirectoryExistsAndIsWritable_NoExceptions(Result))
-                    {
-                        throw new AutomationException("Could not find an appropriate shared temp folder {0}", Result);
-                    }
-                }
+                throw new AutomationException("Could not find an appropriate shared temp folder {0}", Result);
             }
-            SharedBuildDirectoryResolveCache.Add(GameName, Result);
+            SharedBuildDirectoryResolveCache.Add(RootNameForTempStorage, Result);
             return Result;
         }
 
         /// <summary>
-        /// Returns the full path of the temp storage directory for the given game name (ie, P:\Builds\GameName\TmpStore).
+        /// Returns the full path of the temp storage directory for the given game name (ie, P:\Builds\RootNameForTempStorage\TmpStore).
         /// If the game's build folder does not exist, will use the temp storage folder in the UE4 build folder (ie, P:\Builds\UE4\TmpStore).
         /// If the directory does not exist, we will try to create it.
         /// </summary>
-        /// <param name="GameName">game name to determine the temp storage folder for. Empty is equivalent to "UE4".</param>
+        /// <param name="RootNameForTempStorage">The root name beneath <see cref="CommandUtils.RootBuildStorageDirectory"/> where all temp storage will placed for this job.</param>
         /// <returns>Essentially adds the TmpStore/ subdirectory to the SharedBuildDirectory for the game.</returns>
-        public static string ResolveSharedTempStorageDirectory(string GameName)
+        public static string ResolveSharedTempStorageDirectory(string RootNameForTempStorage)
         {
-            string Result = CommandUtils.CombinePaths(ResolveSharedBuildDirectory(GameName), TempStorageSubdirectoryName);
+            string Result = CommandUtils.CombinePaths(ResolveSharedBuildDirectory(RootNameForTempStorage), TempStorageSubdirectoryName);
 
             if (!InternalUtils.Robust_DirectoryExists_NoExceptions(Result, "Could not find {0}"))
             {
@@ -686,20 +670,19 @@ namespace AutomationTool
             }
             if (!InternalUtils.Robust_DirectoryExists_NoExceptions(Result, "Could not find {0}"))
             {
-                throw new AutomationException("Could not create an appropriate shared temp folder {0} for game {1}", Result, GameName);
+                throw new AutomationException("Could not create an appropriate shared temp folder {0} for {1}", Result, RootNameForTempStorage);
             }
             return Result;
         }
 
         /// <summary>
-        /// Returns the full path to the temp storage directory for the given temp storage node and game (ie, P:\Builds\GameName\TmpStore\NodeInfoDirectory)
+        /// Returns the full path to the temp storage directory for the given temp storage node and game (ie, P:\Builds\RootNameForTempStorage\TmpStore\NodeInfoDirectory)
         /// </summary>
         /// <param name="TempStorageNodeInfo">Node info descibing the block of temp storage (essentially used to identify a subdirectory insides the game's temp storage folder).</param>
-        /// <param name="GameName">game name to determine the temp storage folder for. Empty is equivalent to "UE4".</param>
         /// <returns>The full path to the temp storage directory for the given storage block name for the given game.</returns>
-        private static string SharedTempStorageDirectory(TempStorageNodeInfo TempStorageNodeInfo, string GameName)
+        private static string SharedTempStorageDirectory(TempStorageNodeInfo TempStorageNodeInfo)
         {
-            return CommandUtils.CombinePaths(ResolveSharedTempStorageDirectory(GameName), TempStorageNodeInfo.GetRelativeDirectory());
+            return CommandUtils.CombinePaths(ResolveSharedTempStorageDirectory(TempStorageNodeInfo.JobInfo.RootNameForTempStorage), TempStorageNodeInfo.GetRelativeDirectory());
         }
 
         /// <summary>
@@ -729,14 +712,13 @@ namespace AutomationTool
         }
 
         /// <summary>
-        /// Gets the name of the temp storage manifest file for the given temp storage node and game (ie, P:\Builds\GameName\TmpStore\NodeInfoDirectory\NodeInfoFilename.TempManifest)
+        /// Gets the name of the temp storage manifest file for the given temp storage node and game (ie, P:\Builds\RootNameForTempStorage\TmpStore\NodeInfoDirectory\NodeInfoFilename.TempManifest)
         /// </summary>
         /// <param name="TempStorageNodeInfo">Node info descibing the block of temp storage (essentially used to identify a subdirectory insides the game's temp storage folder).</param>
-        /// <param name="GameName">game name to determine the temp storage folder for. Empty is equivalent to "UE4".</param>
         /// <returns>The name of the temp storage manifest file for the given storage block name for the given game.</returns>
-        private static string SharedTempStorageManifestFilename(TempStorageNodeInfo TempStorageNodeInfo, string GameName)
+        private static string SharedTempStorageManifestFilename(TempStorageNodeInfo TempStorageNodeInfo)
         {
-            return CommandUtils.CombinePaths(SharedTempStorageDirectory(TempStorageNodeInfo, GameName), TempStorageNodeInfo.GetManifestFilename());
+            return CommandUtils.CombinePaths(SharedTempStorageDirectory(TempStorageNodeInfo), TempStorageNodeInfo.GetManifestFilename());
         }
 
         /// <summary>
@@ -753,10 +735,9 @@ namespace AutomationTool
         /// Deletes all shared temp storage (file and manifest) for the given temp storage node and game.
         /// </summary>
         /// <param name="TempStorageNodeInfo">Node info descibing the block of temp storage (essentially used to identify a subdirectory insides the game's temp storage folder).</param>
-        /// <param name="GameName">game name to determine the temp storage folder for. Empty is equivalent to "UE4".</param>
-        private static void DeleteSharedTempStorage(TempStorageNodeInfo TempStorageNodeInfo, string GameName)
+        private static void DeleteSharedTempStorage(TempStorageNodeInfo TempStorageNodeInfo)
         {
-            CommandUtils.DeleteDirectory(true, SharedTempStorageDirectory(TempStorageNodeInfo, GameName));
+            CommandUtils.DeleteDirectory(true, SharedTempStorageDirectory(TempStorageNodeInfo));
         }
 
         /// <summary>
@@ -786,12 +767,11 @@ namespace AutomationTool
         /// Checks if a shared temp storage manifest exists for the given temp storage node and game.
         /// </summary>
         /// <param name="TempStorageNodeInfo">Node info descibing the block of temp storage (essentially used to identify a subdirectory insides the game's temp storage folder).</param>
-        /// <param name="GameName">game name to determine the temp storage folder for. Empty is equivalent to "UE4".</param>
         /// <param name="bQuiet">True to suppress logging during the operation.</param>
         /// <returns>true if the shared temp storage manifest file exists</returns>
-        private static bool SharedTempStorageManifestExists(TempStorageNodeInfo TempStorageNodeInfo, string GameName, bool bQuiet)
+        private static bool SharedTempStorageManifestExists(TempStorageNodeInfo TempStorageNodeInfo, bool bQuiet)
         {
-            var SharedManifest = SharedTempStorageManifestFilename(TempStorageNodeInfo, GameName);
+            var SharedManifest = SharedTempStorageManifestFilename(TempStorageNodeInfo);
             if (CommandUtils.FileExists_NoExceptions(bQuiet, SharedManifest))
             {
                 return true;
@@ -803,13 +783,12 @@ namespace AutomationTool
         /// Checks if a temp storage manifest exists, either locally or in the shared folder (depending on the value of bLocalOnly) for the given temp storage node and game.
         /// </summary>
         /// <param name="TempStorageNodeInfo">Node info descibing the block of temp storage (essentially used to identify a subdirectory insides the game's temp storage folder).</param>
-        /// <param name="GameName">game name to determine the temp storage folder for. Empty is equivalent to "UE4".</param>
         /// <param name="bLocalOnly">If true, only ensures that the local temp storage manifest exists. Otherwise, checks if the manifest exists in either local or shared storage.</param>
         /// <param name="bQuiet">True to suppress logging during the operation.</param>
         /// <returns>true if the shared temp storage manifest file exists</returns>
-        public static bool TempStorageExists(TempStorageNodeInfo TempStorageNodeInfo, string GameName, bool bLocalOnly, bool bQuiet)
+        public static bool TempStorageExists(TempStorageNodeInfo TempStorageNodeInfo, bool bLocalOnly, bool bQuiet)
         {
-            return LocalTempStorageManifestExists(TempStorageNodeInfo, bQuiet) || (!bLocalOnly && SharedTempStorageManifestExists(TempStorageNodeInfo, GameName, bQuiet));
+            return LocalTempStorageManifestExists(TempStorageNodeInfo, bQuiet) || (!bLocalOnly && SharedTempStorageManifestExists(TempStorageNodeInfo, bQuiet));
         }
 
         /// <summary>
@@ -1016,10 +995,9 @@ namespace AutomationTool
         /// <param name="TempStorageNodeInfo">Node info descibing the block of temp storage (essentially used to identify a subdirectory insides the game's temp storage folder).</param>
         /// <param name="Files">Fully qualified names of files to reference in the manifest file.</param>
         /// <param name="bLocalOnly">If true, only ensures that the local temp storage manifest exists. Otherwise, checks if the manifest exists in either local or shared storage.</param>
-        /// <param name="GameName">game name to determine the temp storage folder for. Empty is equivalent to "UE4".</param>
         /// <param name="RootDir">Folder that all the given files are rooted from. If null or empty, CmdEnv.LocalRoot is used.</param>
         /// <returns>The created manifest instance (which has already been saved to disk).</returns>
-        public static void StoreToTempStorage(TempStorageNodeInfo TempStorageNodeInfo, List<string> Files, bool bLocalOnly, string GameName, string RootDir)
+        public static void StoreToTempStorage(TempStorageNodeInfo TempStorageNodeInfo, List<string> Files, bool bLocalOnly, string RootDir)
         {
             using (var TelemetryStopwatch = new TelemetryStopwatch("StoreToTempStorage"))
             {
@@ -1038,7 +1016,7 @@ namespace AutomationTool
                 }
                 else
                 {
-                    var SharedStorageNodeDir = SharedTempStorageDirectory(TempStorageNodeInfo, GameName);
+                    var SharedStorageNodeDir = SharedTempStorageDirectory(TempStorageNodeInfo);
                     CommandUtils.Log("Storing to {0}", SharedStorageNodeDir);
                     // this folder should not already exist, else we have concurrency or job duplication problems.
                     if (CommandUtils.DirectoryExists_NoExceptions(SharedStorageNodeDir))
@@ -1048,7 +1026,7 @@ namespace AutomationTool
                     CommandUtils.CreateDirectory(true, SharedStorageNodeDir);
 
                     var LocalManifestFilename = LocalTempStorageManifestFilename(TempStorageNodeInfo);
-                    var SharedManifestFilename = SharedTempStorageManifestFilename(TempStorageNodeInfo, GameName);
+                    var SharedManifestFilename = SharedTempStorageManifestFilename(TempStorageNodeInfo);
                     var StagingDir = Path.GetDirectoryName(LocalManifestFilename);
                     var ZipBasename = Path.GetFileNameWithoutExtension(LocalManifestFilename);
                     // initiate the parallel zip operation.
@@ -1067,11 +1045,9 @@ namespace AutomationTool
         /// If the temp manifest for this block is found locally, the copy is skipped, as we assume this is the same machine that created the temp storage and the files are still there.
         /// </summary>
         /// <param name="TempStorageNodeInfo">Node info descibing the block of temp storage (essentially used to identify a subdirectory insides the game's temp storage folder).</param>
-        /// <param name="WasLocal">upon return, this parameter is set to true if the temp manifest was found locally and the copy was avoided.</param>
-        /// <param name="GameName">game name to determine the temp storage folder for. Empty is equivalent to "UE4".</param>
         /// <param name="RootDir">Folder that all the retrieved files should be rooted from. If null or empty, CmdEnv.LocalRoot is used.</param>
         /// <returns>List of fully qualified paths to all the files that were retrieved. This is returned even if we skip the copy (set WasLocal = true) .</returns>
-        public static List<string> RetrieveFromTempStorage(TempStorageNodeInfo TempStorageNodeInfo, out bool WasLocal, string GameName, string RootDir)
+        public static List<string> RetrieveFromTempStorage(TempStorageNodeInfo TempStorageNodeInfo, string RootDir)
         {
             using (var TelemetryStopwatch = new TelemetryStopwatch("RetrieveFromTempStorage"))
             {
@@ -1096,21 +1072,19 @@ namespace AutomationTool
                     {
                         throw new AutomationException("Local files in manifest {0} were tampered with.", LocalManifest);
                     }
-                    WasLocal = true;
                     TelemetryStopwatch.Finish(string.Format("RetrieveFromTempStorage.{0}.{1}.{2}.Local.{3}.{4}.{5}", Files.Count, Local.GetTotalSize(), 0L, 0L, 0L, TempStorageNodeInfo.NodeStorageName));
                     return Files;
                 }
-                WasLocal = false;
 
                 // We couldn't find the node storage locally, so get it from the shared location.
-                var SharedStorageNodeDir = SharedTempStorageDirectory(TempStorageNodeInfo, GameName);
+                var SharedStorageNodeDir = SharedTempStorageDirectory(TempStorageNodeInfo);
 
                 CommandUtils.Log("Attempting to retrieve from {0}", SharedStorageNodeDir);
                 if (!CommandUtils.DirectoryExists_NoExceptions(SharedStorageNodeDir))
                 {
                     throw new AutomationException("Storage Block Does Not Exists! {0}", SharedStorageNodeDir);
                 }
-                var SharedManifest = SharedTempStorageManifestFilename(TempStorageNodeInfo, GameName);
+                var SharedManifest = SharedTempStorageManifestFilename(TempStorageNodeInfo);
                 InternalUtils.Robust_FileExists(SharedManifest, "Storage Block Manifest Does Not Exists! {0}");
 
                 var Shared = TempStorageManifest.Load(SharedManifest);
@@ -1161,12 +1135,11 @@ namespace AutomationTool
         /// @todo: Someday Make FindNodeHistory query the build database directly to get a true history of the node.
         /// </summary>
         /// <param name="TempStorageNodeInfo">Node info descibing the block of temp storage (essentially used to identify a subdirectory insides the game's temp storage folder).</param>
-        /// <param name="GameName">game name to determine the temp storage folder for. Empty is equivalent to "UE4".</param>
         /// <returns></returns>
-        public static List<int> FindMatchingSharedTempStorageNodeCLs(TempStorageNodeInfo TempStorageNodeInfo, string GameName)
+        public static List<int> FindMatchingSharedTempStorageNodeCLs(TempStorageNodeInfo TempStorageNodeInfo)
         {
             // Find the shared temp storage folder for this game and branch
-            var SharedTempStorageDirectoryForGameAndBranch = CommandUtils.CombinePaths(ResolveSharedTempStorageDirectory(GameName), TempStorageNodeInfo.JobInfo.BranchNameForTempStorage);
+            var SharedTempStorageDirectoryForGameAndBranch = CommandUtils.CombinePaths(ResolveSharedTempStorageDirectory(TempStorageNodeInfo.JobInfo.RootNameForTempStorage), TempStorageNodeInfo.JobInfo.BranchNameForTempStorage);
             int dummy;
             return (
                 // Look for all folders underneath this, it should be a CL, or CL-PreflightInfo
@@ -1188,12 +1161,12 @@ namespace AutomationTool
         static internal void TestTempStorage(CommandEnvironment CmdEnv)
         {
             // We are not a real GUBP job, so fake the values.
-            var TestTempStorageInfo = new TempStorageNodeInfo(new JobInfo("Test", 0, 0, 0), "Test");
+            var TestTempStorageInfo = new TempStorageNodeInfo(new JobInfo("", "Test", "UE4", 0, 0, 0), "Test");
 
             // Delete any local and shared temp storage that may exist.
             DeleteLocalTempStorage();
-            DeleteSharedTempStorage(TestTempStorageInfo, "UE4");
-            if (TempStorageExists(TestTempStorageInfo, "UE4", false, false))
+            DeleteSharedTempStorage(TestTempStorageInfo);
+            if (TempStorageExists(TestTempStorageInfo, false, false))
             {
                 throw new AutomationException("storage should not exist");
             }
@@ -1224,14 +1197,14 @@ namespace AutomationTool
             try
             {
                 // Store the test file to temp storage.
-                StoreToTempStorage(TestTempStorageInfo, TestFiles.Select(TestFile => TestFile.FileName).ToList(), false, "UE4", CmdEnv.LocalRoot);
+                StoreToTempStorage(TestTempStorageInfo, TestFiles.Select(TestFile => TestFile.FileName).ToList(), false, CmdEnv.LocalRoot);
                 // The manifest should exist locally.
                 if (!LocalTempStorageManifestExists(TestTempStorageInfo))
                 {
                     throw new AutomationException("local storage should exist");
                 }
                 // The manifest should exist on the shared drive.
-                if (!SharedTempStorageManifestExists(TestTempStorageInfo, "UE4", false))
+                if (!SharedTempStorageManifestExists(TestTempStorageInfo, false))
                 {
                     throw new AutomationException("shared storage should exist");
                 }
@@ -1243,14 +1216,13 @@ namespace AutomationTool
                     throw new AutomationException("local storage should not exist");
                 }
                 // But the shared storage should still exist.
-                if (!TempStorageExists(TestTempStorageInfo, "UE4", false, false))
+                if (!TempStorageExists(TestTempStorageInfo, false, false))
                 {
                     throw new AutomationException("some storage should exist");
                 }
 
                 // Now we should be able to retrieve the test files from shared storage, and it should overwrite our read-only files.
-                bool WasLocal;
-                RetrieveFromTempStorage(TestTempStorageInfo, out WasLocal, "UE4", CmdEnv.LocalRoot);
+                RetrieveFromTempStorage(TestTempStorageInfo, CmdEnv.LocalRoot);
                 // Now delete the local manifest so we can try again with no files present (the usual case for restoring from temp storage).
                 DeleteLocalTempStorage();
 
@@ -1261,14 +1233,14 @@ namespace AutomationTool
                 }
 
                 // Now we should be able to retrieve the test files from shared storage.
-                RetrieveFromTempStorage(TestTempStorageInfo, out WasLocal, "UE4", CmdEnv.LocalRoot);
+                RetrieveFromTempStorage(TestTempStorageInfo, CmdEnv.LocalRoot);
                 // the local manifest should be there, since we just retrieved from shared storage.
                 if (!LocalTempStorageManifestExists(TestTempStorageInfo))
                 {
                     throw new AutomationException("local storage should exist");
                 }
                 // The shared manifest should also still be there.
-                if (!SharedTempStorageManifestExists(TestTempStorageInfo, "UE4", false))
+                if (!SharedTempStorageManifestExists(TestTempStorageInfo, false))
                 {
                     throw new AutomationException("shared storage should exist");
                 }
@@ -1281,15 +1253,15 @@ namespace AutomationTool
                     }
                 }
                 // Now delete the shared temp storage
-                DeleteSharedTempStorage(TestTempStorageInfo, "UE4");
+                DeleteSharedTempStorage(TestTempStorageInfo);
                 // Shared temp storage manifest should no longer exist.
-                if (SharedTempStorageManifestExists(TestTempStorageInfo, "UE4", false))
+                if (SharedTempStorageManifestExists(TestTempStorageInfo, false))
                 {
                     throw new AutomationException("shared storage should not exist");
                 }
                 // Retrieving temp storage should now just retrieve from local
-                RetrieveFromTempStorage(TestTempStorageInfo, out WasLocal, "UE4", CmdEnv.LocalRoot);
-                if (!WasLocal || !LocalTempStorageManifestExists(TestTempStorageInfo))
+                RetrieveFromTempStorage(TestTempStorageInfo, CmdEnv.LocalRoot);
+                if (!LocalTempStorageManifestExists(TestTempStorageInfo))
                 {
                     throw new AutomationException("local storage should exist");
                 }
@@ -1303,7 +1275,7 @@ namespace AutomationTool
                     MissingFile.Add(CommandUtils.CombinePaths(CmdEnv.LocalRoot, "Engine", "SomeFileThatDoesntExist.txt"));
                     try
                     {
-                        StoreToTempStorage(TestTempStorageInfo, MissingFile, false, "UE4", CmdEnv.LocalRoot);
+                        StoreToTempStorage(TestTempStorageInfo, MissingFile, false, CmdEnv.LocalRoot);
                     }
                     catch (AutomationException)
                     {
@@ -1316,9 +1288,9 @@ namespace AutomationTool
                 }
 
                 // clear the shared temp storage again.
-                DeleteSharedTempStorage(TestTempStorageInfo, "UE4");
+                DeleteSharedTempStorage(TestTempStorageInfo);
                 // store the test files to shared temp storage again.
-                StoreToTempStorage(TestTempStorageInfo, TestFiles.Select(TestFile => TestFile.FileName).ToList(), false, "UE4", CmdEnv.LocalRoot);
+                StoreToTempStorage(TestTempStorageInfo, TestFiles.Select(TestFile => TestFile.FileName).ToList(), false, CmdEnv.LocalRoot);
                 // force a load from shared by deleting the local manifest
                 DeleteLocalTempStorage();
                 // delete our test files locally.
@@ -1330,14 +1302,14 @@ namespace AutomationTool
                 // now test that retrieving from shared temp storage properly balks that a file is missing.
                 {
                     // tamper with the shared files.
-                    var RandomSharedZipFile = Directory.EnumerateFiles(SharedTempStorageDirectory(TestTempStorageInfo, "UE4"), "*.zip").First();
+                    var RandomSharedZipFile = Directory.EnumerateFiles(SharedTempStorageDirectory(TestTempStorageInfo), "*.zip").First();
                     // delete the shared file.
                     CommandUtils.DeleteFile(RandomSharedZipFile);
 
                     bool bFailedProperly = false;
                     try
                     {
-                        RetrieveFromTempStorage(TestTempStorageInfo, out WasLocal, "UE4", CmdEnv.LocalRoot);
+                        RetrieveFromTempStorage(TestTempStorageInfo, CmdEnv.LocalRoot);
                     }
                     catch (AutomationException)
                     {
@@ -1355,9 +1327,9 @@ namespace AutomationTool
                 }
 
                 // clear the shared temp storage again.
-                DeleteSharedTempStorage(TestTempStorageInfo, "UE4");
+                DeleteSharedTempStorage(TestTempStorageInfo);
                 // Copy the files to temp storage.
-                StoreToTempStorage(TestTempStorageInfo, TestFiles.Select(TestFile => TestFile.FileName).ToList(), false, "UE4", CmdEnv.LocalRoot);
+                StoreToTempStorage(TestTempStorageInfo, TestFiles.Select(TestFile => TestFile.FileName).ToList(), false, CmdEnv.LocalRoot);
                 // Delete a local file.
                 CommandUtils.DeleteFile(TestFiles[0].FileName);
                 // retrieving from temp storage should use WasLocal, but should balk because a local file was deleted.
@@ -1365,7 +1337,7 @@ namespace AutomationTool
                     bool bFailedProperly = false;
                     try
                     {
-                        RetrieveFromTempStorage(TestTempStorageInfo, out WasLocal, "UE4", CmdEnv.LocalRoot);
+                        RetrieveFromTempStorage(TestTempStorageInfo, CmdEnv.LocalRoot);
                     }
                     catch (AutomationException)
                     {
@@ -1380,7 +1352,7 @@ namespace AutomationTool
             finally
             {
                 // Do a final cleanup.
-                DeleteSharedTempStorage(TestTempStorageInfo, "UE4");
+                DeleteSharedTempStorage(TestTempStorageInfo);
                 DeleteLocalTempStorage();
                 foreach (var TestFile in TestFiles)
                 {
