UE-19835 - fix for being able to stage the same file to multiple places
#uat

diff --git a/Engine/Source/Programs/AutomationTool/AutomationUtils/DeploymentContext.cs b/Engine/Source/Programs/AutomationTool/AutomationUtils/DeploymentContext.cs
--- a/Engine/Source/Programs/AutomationTool/AutomationUtils/DeploymentContext.cs
+++ b/Engine/Source/Programs/AutomationTool/AutomationUtils/DeploymentContext.cs
@@ -151,12 +151,12 @@ public class DeploymentContext //: ProjectParams
 	///  After staging, this is a map from source file to relative file in the stage
 	/// These file are binaries, etc and can't go into a pak file
 	/// </summary>
-	public Dictionary<string, string> NonUFSStagingFiles = new Dictionary<string, string>();
+	public Dictionary<string, List<string>> NonUFSStagingFiles = new Dictionary<string, List<string>>();
 	/// <summary>
 	///  After staging, this is a map from source file to relative file in the stage
 	/// These file are debug, and can't go into a pak file
 	/// </summary>
-	public Dictionary<string, string> NonUFSStagingFilesDebug = new Dictionary<string, string>();
+	public Dictionary<string, List<string>> NonUFSStagingFilesDebug = new Dictionary<string, List<string>>();
 	/// <summary>
 	///  After staging, this is a map from source file to relative file in the stage
 	/// These file are content, and can go into a pak file
@@ -383,11 +383,11 @@ public class DeploymentContext //: ProjectParams
 		}
 		else if (FileType == StagedFileType.NonUFS)
 		{
-			AddUniqueStagingFile(NonUFSStagingFiles, InputPath, OutputPath);
+			AddStagingFile(NonUFSStagingFiles, InputPath, OutputPath);
 		}
 		else if (FileType == StagedFileType.DebugNonUFS)
 		{
-			AddUniqueStagingFile(NonUFSStagingFilesDebug, InputPath, OutputPath);
+			AddStagingFile(NonUFSStagingFilesDebug, InputPath, OutputPath);
 		}
 	}
 
@@ -582,11 +582,11 @@ public class DeploymentContext //: ProjectParams
 				}
 				else if (FileType == StagedFileType.NonUFS)
 				{
-					AddUniqueStagingFile(NonUFSStagingFiles, FileToCopy, Dest);
+					AddStagingFile(NonUFSStagingFiles, FileToCopy, Dest);
 				}
 				else if (FileType == StagedFileType.DebugNonUFS)
 				{
-					AddUniqueStagingFile(NonUFSStagingFilesDebug, FileToCopy, Dest);
+					AddStagingFile(NonUFSStagingFilesDebug, FileToCopy, Dest);
 				}
 				FilesAdded++;
 			}
@@ -618,6 +618,22 @@ public class DeploymentContext //: ProjectParams
 		}
 	}
 
+	private void AddStagingFile(Dictionary<string, List<string>> FilesToStage, string FileToCopy, string Dest)
+	{
+		List<string> ExistingDest;
+		if (FilesToStage.TryGetValue(FileToCopy, out ExistingDest))
+		{
+			if (!FilesToStage[FileToCopy].Contains(Dest))
+			{
+				FilesToStage[FileToCopy].Add(Dest);
+			}
+		}
+		else
+		{
+			FilesToStage.Add(FileToCopy, new List<string>(){Dest});
+		}
+	}
+
 	public int ArchiveFiles(string InPath, string Wildcard = "*", bool bRecursive = true, string[] ExcludeWildcard = null)
 	{
 		int FilesAdded = 0;
diff --git a/Engine/Source/Programs/AutomationTool/Mac/MacPlatform.Automation.cs b/Engine/Source/Programs/AutomationTool/Mac/MacPlatform.Automation.cs
--- a/Engine/Source/Programs/AutomationTool/Mac/MacPlatform.Automation.cs
+++ b/Engine/Source/Programs/AutomationTool/Mac/MacPlatform.Automation.cs
@@ -203,8 +203,12 @@ public class MacPlatform : Platform
 						}
 
 						string AppPath = Executable.Path.Substring(0, Executable.Path.LastIndexOf(".app/") + 4);
-						string AppRelativePath = SC.NonUFSStagingFiles[Executable.Path].Substring(0, SC.NonUFSStagingFiles[Executable.Path].LastIndexOf(".app/") + 4);
-						StageBootstrapExecutable(SC, BootstrapExeName, AppPath, AppRelativePath, BootstrapArguments);
+						object Dest = SC.NonUFSStagingFiles[Executable.Path];
+						foreach (var DestPath in SC.NonUFSStagingFiles[Executable.Path])
+						{
+							string AppRelativePath = DestPath.Substring(0, DestPath.LastIndexOf(".app/") + 4);
+							StageBootstrapExecutable(SC, BootstrapExeName, AppPath, AppRelativePath, BootstrapArguments);
+						}
 					}
 				}
 			}
diff --git a/Engine/Source/Programs/AutomationTool/Scripts/CopyBuildToStagingDirectory.Automation.cs b/Engine/Source/Programs/AutomationTool/Scripts/CopyBuildToStagingDirectory.Automation.cs
--- a/Engine/Source/Programs/AutomationTool/Scripts/CopyBuildToStagingDirectory.Automation.cs
+++ b/Engine/Source/Programs/AutomationTool/Scripts/CopyBuildToStagingDirectory.Automation.cs
@@ -165,6 +165,21 @@ public partial class Project : CommandUtils
 		}
 		return Result;
 	}
+	
+	public static Dictionary<string, List<string>> ConvertToLower(Dictionary<string, List<string>> Mapping)
+	{
+		var Result = new Dictionary<string, List<string>>();
+		foreach (var Pair in Mapping)
+		{
+			List<string> NewList = new List<string>();
+			foreach(string s in Pair.Value)
+			{
+				NewList.Add(s.ToLowerInvariant());
+			}
+			Result.Add(Pair.Key, NewList);
+		}
+		return Result;
+	}
 
 	public static void MaybeConvertToLowerCase(ProjectParams Params, DeploymentContext SC)
 	{
@@ -557,6 +572,33 @@ public partial class Project : CommandUtils
 		}
 	}
 
+	public static void DumpTargetManifest(Dictionary<string, List<string>> Mapping, string Filename, string StageDir, List<string> CRCFiles)
+	{
+		// const string Iso8601DateTimeFormat = "yyyy-MM-ddTHH:mm:ssZ"; // probably should work
+		// const string Iso8601DateTimeFormat = "o"; // predefined universal Iso standard format (has too many millisecond spaces for our read code in FDateTime.ParseISO8601
+		const string Iso8601DateTimeFormat = "yyyy'-'MM'-'dd'T'HH':'mm':'ss'.'fffZ";
+
+		if (Mapping.Count > 0)
+		{
+			var Lines = new List<string>();
+			foreach (var Pair in Mapping)
+			{
+				foreach (var DestFile in Pair.Value)
+				{
+					string TimeStamp = File.GetLastWriteTimeUtc(Pair.Key).ToString(Iso8601DateTimeFormat);
+					if (CRCFiles.Contains(DestFile))
+					{
+						byte[] FileData = File.ReadAllBytes(StageDir + "/" + DestFile);
+						TimeStamp = BitConverter.ToString(System.Security.Cryptography.MD5.Create().ComputeHash(FileData)).Replace("-", string.Empty);
+					}
+					string DestData = DestFile + "\t" + TimeStamp;
+					Lines.Add(DestData);
+				}
+			}
+			WriteAllLines(Filename, Lines.ToArray());
+		}
+	}
+
 	public static void CopyManifestFilesToStageDir(Dictionary<string, string> Mapping, string StageDir, string ManifestName, List<string> CRCFiles)
 	{
 		Log("Copying {0} to staging directory: {1}", ManifestName, StageDir);
@@ -588,6 +630,44 @@ public partial class Project : CommandUtils
 		}
 	}
 
+	public static void CopyManifestFilesToStageDir(Dictionary<string, List<string>> Mapping, string StageDir, string ManifestName, List<string> CRCFiles)
+	{
+		Log("Copying {0} to staging directory: {1}", ManifestName, StageDir);
+		string ManifestPath = "";
+		string ManifestFile = "";
+		if (!String.IsNullOrEmpty(ManifestName))
+		{
+			ManifestFile = "Manifest_" + ManifestName + ".txt";
+			ManifestPath = CombinePaths(StageDir, ManifestFile);
+			DeleteFile(ManifestPath);
+		}
+		foreach (var Pair in Mapping)
+		{
+			if (Pair.Value.Count > 1)
+			{
+				Log("More than 1 file being deployed.");
+			}
+			string Src = Pair.Key;
+			foreach (var DestPath in Pair.Value)
+			{
+				string DestFile = CombinePaths(StageDir, DestPath);
+				if (Src != DestFile)  // special case for things created in the staging directory, like the pak file
+				{
+					CopyFileIncremental(Src, DestFile, bFilterSpecialLinesFromIniFiles: true);
+				}
+			}
+		}
+		if (!String.IsNullOrEmpty(ManifestPath) && Mapping.Count > 0)
+		{
+			DumpTargetManifest(Mapping, ManifestPath, StageDir, CRCFiles);
+			if (!FileExists(ManifestPath))
+			{
+				throw new AutomationException("Failed to write manifest {0}", ManifestPath);
+			}
+			CopyFile(ManifestPath, CombinePaths(CmdEnv.LogFolder, ManifestFile));
+		}
+	}
+
 	public static void DumpManifest(Dictionary<string, string> Mapping, string Filename)
 	{
 		if (Mapping.Count > 0)
@@ -604,6 +684,23 @@ public partial class Project : CommandUtils
 		}
 	}
 
+	public static void DumpManifest(Dictionary<string, List<string>> Mapping, string Filename)
+	{
+		if (Mapping.Count > 0)
+		{
+			var Lines = new List<string>();
+			foreach (var Pair in Mapping)
+			{
+				string Src = Pair.Key;
+				foreach (var DestPath in Pair.Value)
+				{
+					Lines.Add("\"" + Src + "\" \"" + DestPath + "\"");
+				}
+			}
+			WriteAllLines(Filename, Lines.ToArray());
+		}
+	}
+
 	public static void DumpManifest(DeploymentContext SC, string BaseFilename, bool DumpUFSFiles = true)
 	{
 		DumpManifest(SC.NonUFSStagingFiles, BaseFilename + "_NonUFSFiles.txt");
diff --git a/Engine/Source/Programs/AutomationTool/Win/WinPlatform.Automation.cs b/Engine/Source/Programs/AutomationTool/Win/WinPlatform.Automation.cs
--- a/Engine/Source/Programs/AutomationTool/Win/WinPlatform.Automation.cs
+++ b/Engine/Source/Programs/AutomationTool/Win/WinPlatform.Automation.cs
@@ -75,7 +75,10 @@ public abstract class BaseWinPlatform : Platform
 							BootstrapExeName = SC.ShortProjectName + ".exe";
 						}
 
-						StageBootstrapExecutable(SC, BootstrapExeName, Executable.Path, SC.NonUFSStagingFiles[Executable.Path], BootstrapArguments);
+						foreach (string StagePath in SC.NonUFSStagingFiles[Executable.Path])
+						{
+							StageBootstrapExecutable(SC, BootstrapExeName, Executable.Path, StagePath, BootstrapArguments);
+						}
 					}
 				}
 			}
