remove some debug code

diff --git a/Engine/Source/Programs/AutomationTool/Scripts/CopyBuildToStagingDirectory.Automation.cs b/Engine/Source/Programs/AutomationTool/Scripts/CopyBuildToStagingDirectory.Automation.cs
--- a/Engine/Source/Programs/AutomationTool/Scripts/CopyBuildToStagingDirectory.Automation.cs
+++ b/Engine/Source/Programs/AutomationTool/Scripts/CopyBuildToStagingDirectory.Automation.cs
@@ -643,10 +643,6 @@ public partial class Project : CommandUtils
 		}
 		foreach (var Pair in Mapping)
 		{
-			if (Pair.Value.Count > 1)
-			{
-				Log("More than 1 file being deployed.");
-			}
 			string Src = Pair.Key;
 			foreach (var DestPath in Pair.Value)
 			{
