Fix deploy to Android ADB connected over WiFi
#ue4
#android

diff --git a/Engine/Source/Programs/AutomationTool/Android/AndroidPlatform.Automation.cs b/Engine/Source/Programs/AutomationTool/Android/AndroidPlatform.Automation.cs
--- a/Engine/Source/Programs/AutomationTool/Android/AndroidPlatform.Automation.cs
+++ b/Engine/Source/Programs/AutomationTool/Android/AndroidPlatform.Automation.cs
@@ -462,9 +462,11 @@ public class AndroidPlatform : Platform
 		string RemoteDir = StorageLocation + "/UE4Game/" + Params.ShortProjectName;
 
 		// Note: appends the device name to make the filename unique; these files will be deleted later during delta manifest generation
+		// Replace colon with underscore for legal filename (colon may be present for wifi connected devices)
+		string DeviceName = Params.Device.Replace(":", "_");
 
 		// Try retrieving the UFS files manifest files from the device
-		string UFSManifestFileName = CombinePaths(SC.StageDirectory, DeploymentContext.UFSDeployedManifestFileName + "_" + Params.Device);
+		string UFSManifestFileName = CombinePaths(SC.StageDirectory, DeploymentContext.UFSDeployedManifestFileName + "_" + DeviceName);
 		ProcessResult UFSResult = RunAdbCommand(Params, " pull " + RemoteDir + "/" + DeploymentContext.UFSDeployedManifestFileName + " \"" + UFSManifestFileName + "\"", null, ERunOptions.AppMustExist);
 		if (!UFSResult.Output.Contains("bytes"))
 		{
@@ -472,7 +474,7 @@ public class AndroidPlatform : Platform
 		}
 
 		// Try retrieving the non UFS files manifest files from the device
-		string NonUFSManifestFileName = CombinePaths(SC.StageDirectory, DeploymentContext.NonUFSDeployedManifestFileName + "_" + Params.Device);
+		string NonUFSManifestFileName = CombinePaths(SC.StageDirectory, DeploymentContext.NonUFSDeployedManifestFileName + "_" + DeviceName);
 		ProcessResult NonUFSResult = RunAdbCommand(Params, " pull " + RemoteDir + "/" + DeploymentContext.NonUFSDeployedManifestFileName + " \"" + NonUFSManifestFileName + "\"", null, ERunOptions.AppMustExist);
 		if (!NonUFSResult.Output.Contains("bytes"))
 		{
