Merge 2678656
UAT changes to support PS4 patch generation

diff --git a/Engine/Source/Programs/AutomationTool/AutomationUtils/Platform.cs b/Engine/Source/Programs/AutomationTool/AutomationUtils/Platform.cs
--- a/Engine/Source/Programs/AutomationTool/AutomationUtils/Platform.cs
+++ b/Engine/Source/Programs/AutomationTool/AutomationUtils/Platform.cs
@@ -447,6 +447,16 @@ namespace AutomationTool
             return "";
         }
 
+		/// <summary>
+		/// Returns true if the platform wants patches to generate a small .pak file containing the difference
+		/// of current data against a shipped pak file.
+		/// </summary>
+		/// <returns></returns>
+		public virtual bool GetPlatformPatchesWithDiffPak()
+		{
+			return true;
+		}
+
 		/// <summary>
 		///  Returns whether the platform requires a package to deploy to a device
 		/// </summary>
diff --git a/Engine/Source/Programs/AutomationTool/AutomationUtils/ProjectParams.cs b/Engine/Source/Programs/AutomationTool/AutomationUtils/ProjectParams.cs
--- a/Engine/Source/Programs/AutomationTool/AutomationUtils/ProjectParams.cs
+++ b/Engine/Source/Programs/AutomationTool/AutomationUtils/ProjectParams.cs
@@ -460,6 +460,8 @@ namespace AutomationTool
 			this.CookFlavor = ParseParamValueIfNotSpecified(Command, CookFlavor, "cookflavor", String.Empty);
             this.NewCook = GetParamValueIfNotSpecified(Command, NewCook, this.NewCook, "NewCook");
             this.OldCook = GetParamValueIfNotSpecified(Command, OldCook, this.OldCook, "OldCook");
+			this.CreateReleaseVersionBasePath = ParseParamValueIfNotSpecified(Command, CreateReleaseVersionBasePath, "createreleaseversionroot", String.Empty);
+			this.BasedOnReleaseVersionBasePath = ParseParamValueIfNotSpecified(Command, BasedOnReleaseVersionBasePath, "basedonreleaseversionroot", String.Empty);
             this.CreateReleaseVersion = ParseParamValueIfNotSpecified(Command, CreateReleaseVersion, "createreleaseversion", String.Empty);
             this.BasedOnReleaseVersion = ParseParamValueIfNotSpecified(Command, BasedOnReleaseVersion, "basedonreleaseversion", String.Empty);
             this.GeneratePatch = GetParamValueIfNotSpecified(Command, GeneratePatch, this.GeneratePatch, "GeneratePatch");
@@ -1087,7 +1089,7 @@ namespace AutomationTool
         public string InternationalizationPreset;
 
         /// <summary>
-        /// Cook: Create a cooked release version
+        /// Cook: Create a cooked release version.  Also, the version. e.g. 1.0
         /// </summary>
         public string CreateReleaseVersion;
 
@@ -1106,6 +1108,16 @@ namespace AutomationTool
         /// </summary>
         public string BasedOnReleaseVersion;
 
+		/// <summary>
+        /// Cook: Path to the root of the directory where we store released versions of the game for a given version
+        /// </summary>
+		public string BasedOnReleaseVersionBasePath;
+
+		/// <summary>
+		/// Cook: Path to the root of the directory to create a new released version of the game.
+		/// </summary>
+		public string CreateReleaseVersionBasePath;
+
         /// <summary>
         /// Are we generating a patch, generate a patch from a previously released version of the game (use CreateReleaseVersion to create a release). 
         /// this requires BasedOnReleaseVersion
@@ -1859,6 +1871,45 @@ namespace AutomationTool
 		}
 		private string ProjectBinariesPath;
 
+
+		/// <summary>
+		/// Get the path to the directory of the version we are basing a diff or a patch on.  
+		/// </summary>				
+		public String GetBasedOnReleaseVersionPath(DeploymentContext SC)
+		{
+			String BasePath = BasedOnReleaseVersionBasePath;
+			String Platform = SC.StageTargetPlatform.GetCookPlatform(SC.DedicatedServer, false, CookFlavor);
+			if (String.IsNullOrEmpty(BasePath))
+			{
+				BasePath = CommandUtils.CombinePaths(SC.ProjectRoot, "Releases", BasedOnReleaseVersion, Platform);
+			}
+			else
+			{
+				BasePath = CommandUtils.CombinePaths(BasePath, BasedOnReleaseVersion, Platform);
+			}
+			return BasePath;
+		}
+
+		/// <summary>
+		/// Get the path to the target directory for creating a new release version
+		/// </summary>
+		/// <param name="SC"></param>
+		/// <returns></returns>
+		public String GetCreateReleaseVersionPath(DeploymentContext SC)
+		{
+			String BasePath = CreateReleaseVersionBasePath;
+			String Platform = SC.StageTargetPlatform.GetCookPlatform(SC.DedicatedServer, false, CookFlavor);
+			if (String.IsNullOrEmpty(BasePath))
+			{
+				BasePath = CommandUtils.CombinePaths(SC.ProjectRoot, "Releases", CreateReleaseVersion, Platform);
+			}
+			else
+			{
+				BasePath = CommandUtils.CombinePaths(BasePath, CreateReleaseVersion, Platform);
+			}
+			return BasePath;
+		}
+
         /// <summary>
         /// True if we are generating a patch
         /// </summary>
diff --git a/Engine/Source/Programs/AutomationTool/Scripts/CopyBuildToStagingDirectory.Automation.cs b/Engine/Source/Programs/AutomationTool/Scripts/CopyBuildToStagingDirectory.Automation.cs
--- a/Engine/Source/Programs/AutomationTool/Scripts/CopyBuildToStagingDirectory.Automation.cs
+++ b/Engine/Source/Programs/AutomationTool/Scripts/CopyBuildToStagingDirectory.Automation.cs
@@ -808,9 +808,9 @@ public partial class Project : CommandUtils
 	}
 
 
-    private static string GetReleasePakFilePath(DeploymentContext SC, ProjectParams Params, string ReleaseVersion, string PakName)
+    private static string GetReleasePakFilePath(DeploymentContext SC, ProjectParams Params, string PakName)
     {
-        string ReleaseVersionPath = CombinePaths(SC.ProjectRoot, "Releases", ReleaseVersion, SC.StageTargetPlatform.GetCookPlatform(SC.DedicatedServer, false, Params.CookFlavor), PakName);
+        string ReleaseVersionPath = CombinePaths(Params.GetBasedOnReleaseVersionPath(SC), PakName);
         return ReleaseVersionPath;
     }
 
@@ -825,7 +825,7 @@ public partial class Project : CommandUtils
 	private static void CreatePak(ProjectParams Params, DeploymentContext SC, Dictionary<string, string> UnrealPakResponseFile, string PakName)
 	{
         string PostFix = "";
-        if (Params.IsGeneratingPatch)
+		if (Params.IsGeneratingPatch && SC.StageTargetPlatform.GetPlatformPatchesWithDiffPak())
         {
             PostFix += "_P";
         }
@@ -882,11 +882,11 @@ public partial class Project : CommandUtils
 		}
 
         string PatchSourceContentPath = null;
-        if (Params.HasBasedOnReleaseVersion && Params.IsGeneratingPatch)
+		if (Params.HasBasedOnReleaseVersion && Params.IsGeneratingPatch && SC.StageTargetPlatform.GetPlatformPatchesWithDiffPak())
         {
             // don't include the post fix in this filename because we are looking for the source pak path
             string PakFilename = PakName + "-" + SC.FinalCookPlatform + ".pak";
-            PatchSourceContentPath = GetReleasePakFilePath(SC, Params, Params.BasedOnReleaseVersion, PakFilename);
+            PatchSourceContentPath = GetReleasePakFilePath(SC, Params, PakFilename);
             
 
             /*PatchSourceContentPath = Params.BasedOnReleaseVersion;
@@ -910,7 +910,7 @@ public partial class Project : CommandUtils
         {
             // copy the created pak to the release version directory we might need this later if we want to generate patches
             //string ReleaseVersionPath = CombinePaths( SC.ProjectRoot, "Releases", Params.CreateReleaseVersion, SC.StageTargetPlatform.GetCookPlatform(Params.DedicatedServer, false, Params.CookFlavor), Path.GetFileName(OutputLocation) );
-            string ReleaseVersionPath = GetReleasePakFilePath(SC, Params, Params.CreateReleaseVersion, Path.GetFileName(OutputLocation));
+            string ReleaseVersionPath = GetReleasePakFilePath(SC, Params, Path.GetFileName(OutputLocation));
 
 			InternalUtils.SafeCreateDirectory(Path.GetDirectoryName(ReleaseVersionPath));
 			InternalUtils.SafeCopyFile(OutputLocation, ReleaseVersionPath);
