Better way to find aapt (Android) by checking PkgRevision instead of using first subdirectory of build-tools
#ue4
#android

diff --git a/Engine/Source/Programs/AutomationTool/Android/AndroidPlatform.Automation.cs b/Engine/Source/Programs/AutomationTool/Android/AndroidPlatform.Automation.cs
--- a/Engine/Source/Programs/AutomationTool/Android/AndroidPlatform.Automation.cs
+++ b/Engine/Source/Programs/AutomationTool/Android/AndroidPlatform.Automation.cs
@@ -882,17 +882,89 @@ public class AndroidPlatform : Platform
 		}
 	}
 
+	static private string CachedAaptPath = null;
+	static private string LastAndroidHomePath = null;
+
+	private static uint GetRevisionValue(string VersionString)
+	{
+		// read up to 4 sections (ie. 20.0.3.5), first section most significant
+		// each section assumed to be 0 to 255 range
+		uint Value = 0;
+		try
+		{
+			string[] Sections= VersionString.Split(".".ToCharArray());
+			Value |= (Sections.Length > 0) ? (uint.Parse(Sections[0]) << 24) : 0;
+			Value |= (Sections.Length > 1) ? (uint.Parse(Sections[1]) << 16) : 0;
+			Value |= (Sections.Length > 2) ? (uint.Parse(Sections[2]) <<  8) : 0;
+			Value |= (Sections.Length > 3) ?  uint.Parse(Sections[3])        : 0;
+		}
+		catch (Exception)
+		{
+			// ignore poorly formed version
+		}
+		return Value;
+	}	
+
 	private static string GetAaptPath()
 	{
-		// there is a numbered directory in here, hunt it down
-        string path = Environment.ExpandEnvironmentVariables("%ANDROID_HOME%/build-tools/");
-		string[] Subdirs = Directory.GetDirectories(path);
+		// return cached path if ANDROID_HOME has not changed
+        string HomePath = Environment.ExpandEnvironmentVariables("%ANDROID_HOME%");
+		if (CachedAaptPath != null && LastAndroidHomePath == HomePath)
+		{
+			return CachedAaptPath;
+		}
+
+		// get a list of the directories in build-tools.. may be more than one set installed (or none which is bad)
+		string[] Subdirs = Directory.GetDirectories(Path.Combine(HomePath, "build-tools"));
         if (Subdirs.Length == 0)
         {
-            throw new AutomationException(ExitCode.Error_AndroidBuildToolsPathNotFound, "Failed to find %ANDROID_HOME%/build-tools subdirectory");
+            throw new AutomationException(ExitCode.Error_AndroidBuildToolsPathNotFound, "Failed to find %ANDROID_HOME%/build-tools subdirectory. Run SDK manager and install build-tools.");
         }
-		// we expect there to be one, so use the first one
-		return Path.Combine(Subdirs[0], Utils.IsRunningOnMono ? "aapt" : "aapt.exe");
+
+		// valid directories will have a source.properties with the Pkg.Revision (there is no guarantee we can use the directory name as revision)
+		string BestToolPath = null;
+		uint BestVersion = 0;
+		foreach (string CandidateDir in Subdirs)
+		{
+			string AaptFilename = Path.Combine(CandidateDir, Utils.IsRunningOnMono ? "aapt" : "aapt.exe");
+			uint RevisionValue = 0;
+
+			if (File.Exists(AaptFilename))
+			{
+				string SourcePropFilename = Path.Combine(CandidateDir, "source.properties");
+				if (File.Exists(SourcePropFilename))
+				{
+					string[] PropertyContents = File.ReadAllLines(SourcePropFilename);
+					foreach (string PropertyLine in PropertyContents)
+					{
+						if (PropertyLine.StartsWith("Pkg.Revision="))
+						{
+							RevisionValue = GetRevisionValue(PropertyLine.Substring(13));
+							break;
+						}
+					}
+				}
+			}
+
+			// remember it if newer version or haven't found one yet
+			if (RevisionValue > BestVersion || BestToolPath == null)
+			{
+				BestVersion = RevisionValue;
+				BestToolPath = AaptFilename;
+			}
+		}
+
+		if (BestToolPath == null)
+		{
+            throw new AutomationException(ExitCode.Error_AndroidBuildToolsPathNotFound, "Failed to find %ANDROID_HOME%/build-tools subdirectory with aapt. Run SDK manager and install build-tools.");
+		}
+
+		CachedAaptPath = BestToolPath;
+		LastAndroidHomePath = HomePath;
+
+		Log("Using this aapt: {0}", CachedAaptPath);
+
+		return CachedAaptPath;
 	}
 
 	private string GetBestDeviceArchitecture(ProjectParams Params)
