[INTEGRATE]

        Change 2684628 by [REDACTED_PERSON][REDACTED_HANDLE] on 2015/09/09 10:28 Submit latest shipping and debug BuildPatchTool Mac binaries
        Change 2684643 by [REDACTED_PERSON][REDACTED_HANDLE] on 2015/09/09 10:32 Submit shipping and debug BuildPatchTool binaries from Win64
        Change 2684648 by [REDACTED_PERSON][REDACTED_HANDLE] on 2015/09/09 10:38 Change the default BuildPatchTool ages called from automation build commands

diff --git a/Engine/Source/Programs/AutomationTool/AutomationUtils/MCPPublic.cs b/Engine/Source/Programs/AutomationTool/AutomationUtils/MCPPublic.cs
--- a/Engine/Source/Programs/AutomationTool/AutomationUtils/MCPPublic.cs
+++ b/Engine/Source/Programs/AutomationTool/AutomationUtils/MCPPublic.cs
@@ -309,9 +309,9 @@ namespace EpicGames.MCP.Automation
         public class PatchGenerationOptions
         {
             /// <summary>
-            /// By default, we will only consider data modified within 12 days to be reusable
+            /// By default, we will only consider data referenced from manifests modified within five days to be reusable.
             /// </summary>
-            public const int DEFAULT_DATA_AGE_THRESHOLD = 12;
+            private const int DEFAULT_DATA_AGE_THRESHOLD = 5;
 
             public PatchGenerationOptions()
             {
@@ -352,10 +352,10 @@ namespace EpicGames.MCP.Automation
             public MCPPlatform Platform;
             /// <summary>
             /// When identifying existing patch data to reuse in this build, only
-            /// files modified within this number of days will be considered for reuse.
-            /// IMPORTANT: This should always be smaller than the data age threshold for any compactify process which will run on the directory, to ensure
+            /// files referenced from a manifest file modified within this number of days will be considered for reuse.
+            /// IMPORTANT: This should always be smaller than the minimum age at which manifest files can be deleted by any cleanup process, to ensure
             /// that we do not reuse any files which could be deleted by a concurrently running compactify. It is recommended that this number be at least
-            /// two days less than the compactify data age threshold.
+            /// two days less than the cleanup data age threshold.
             /// </summary>
             public int DataAgeThreshold;
 			/// <summary>
@@ -377,11 +377,11 @@ namespace EpicGames.MCP.Automation
 		/// </summary>
 		public class CompactifyOptions
 		{
-			private static readonly int DefaultDataAgeThreshold = PatchGenerationOptions.DEFAULT_DATA_AGE_THRESHOLD + 2;
+			private const int DEFAULT_DATA_AGE_THRESHOLD = 2;
 
 			public CompactifyOptions()
 			{
-				DataAgeThreshold = DefaultDataAgeThreshold;
+				DataAgeThreshold = DEFAULT_DATA_AGE_THRESHOLD;
 			}
 
 			/// <summary>
@@ -402,10 +402,9 @@ namespace EpicGames.MCP.Automation
 			/// </summary>
 			public string ManifestsToKeepFile;
 			/// <summary>
-			/// Path data files modified within this number of days will *not* be deleted, allowing them to be reused by patch generation processes.
-			/// IMPORTANT: This should always be larger than the data age threshold for any build processes which will run on the directory, to ensure
-			/// that we do not delete any files which could be reused by a concurrently running build. It is recommended that this number be at least
-			/// two days greater than the build data age threshold.
+			/// Patch data files modified within this number of days will *not* be deleted, to ensure that any patch files being written out by a.
+			/// patch generation process are not deleted before their corresponding manifest file(s) can be written out.
+			/// NOTE: this should be set to a value larger than the expected maximum time that a build could take.
 			/// </summary>
 			public int DataAgeThreshold;
 		}
diff --git a/Engine/Source/Programs/AutomationTool/Scripts/CopyBuildToStagingDirectory.Automation.cs b/Engine/Source/Programs/AutomationTool/Scripts/CopyBuildToStagingDirectory.Automation.cs
--- a/Engine/Source/Programs/AutomationTool/Scripts/CopyBuildToStagingDirectory.Automation.cs
+++ b/Engine/Source/Programs/AutomationTool/Scripts/CopyBuildToStagingDirectory.Automation.cs
@@ -971,7 +971,7 @@ public partial class Project : CommandUtils
 				string DestManifestPath = CombinePaths(ManifestDir, ManifestFilename);
 				InternalUtils.SafeCreateDirectory(ManifestDir, true);
 
-				string CmdLine = String.Format("-BuildRoot=\"{0}\" -CloudDir=\"{1}\" -AppID={2} -AppName=\"{3}\" -BuildVersion=\"{4}\" -AppLaunch=\"{5}\" -DataAgeThreshold=12", BuildRoot, CloudDir, AppID, AppName, VersionString, AppLaunch);
+				string CmdLine = String.Format("-BuildRoot=\"{0}\" -CloudDir=\"{1}\" -AppID={2} -AppName=\"{3}\" -BuildVersion=\"{4}\" -AppLaunch=\"{5}\"", BuildRoot, CloudDir, AppID, AppName, VersionString, AppLaunch);
 				CmdLine += " -AppArgs=\"\"";
 				CmdLine += " -custom=\"bIsPatch=false\"";
 				CmdLine += String.Format(" -customint=\"ChunkID={0}\"", ChunkID);
