Merging a couple of UAT changes back from game branch.

diff --git a/Engine/Source/Programs/AutomationTool/AutomationUtils/LocalP4Environment.cs b/Engine/Source/Programs/AutomationTool/AutomationUtils/LocalP4Environment.cs
--- a/Engine/Source/Programs/AutomationTool/AutomationUtils/LocalP4Environment.cs
+++ b/Engine/Source/Programs/AutomationTool/AutomationUtils/LocalP4Environment.cs
@@ -147,33 +147,33 @@ namespace AutomationTool
 			}
 			else
 			{
-				// Figure out the build root
-				string KnownFilePathFromRoot = CommandEnvironment.KnownFileRelativeToRoot;
-				string KnownLocalPath = CommandUtils.MakePathSafeToUseWithCommandLine(CommandUtils.CombinePaths(PathSeparator.Slash, LocalRootPath, KnownFilePathFromRoot));
-				ProcessResult P4Result = Connection.P4(String.Format("files -m 1 {0}", KnownLocalPath), AllowSpew: false);
+			// Figure out the build root
+			string KnownFilePathFromRoot = CommandEnvironment.KnownFileRelativeToRoot;
+			string KnownLocalPath = CommandUtils.MakePathSafeToUseWithCommandLine(CommandUtils.CombinePaths(PathSeparator.Slash, LocalRootPath, KnownFilePathFromRoot));
+			ProcessResult P4Result = Connection.P4(String.Format("files -m 1 {0}", KnownLocalPath), AllowSpew: false);
 
-				string KnownFileDepotMapping = P4Result.Output;
+			string KnownFileDepotMapping = P4Result.Output;
 
-				// Get the build root
-				Log.TraceVerbose("Looking for {0} in {1}", KnownFilePathFromRoot, KnownFileDepotMapping);
-				int EndIdx = KnownFileDepotMapping.IndexOf(KnownFilePathFromRoot, StringComparison.CurrentCultureIgnoreCase);
-				if (EndIdx < 0)
-				{
-					EndIdx = KnownFileDepotMapping.IndexOf(CommandUtils.ConvertSeparators(PathSeparator.Slash, KnownFilePathFromRoot), StringComparison.CurrentCultureIgnoreCase);
-				}
-				// Get the root path without the trailing path separator
-				BuildRootPath = KnownFileDepotMapping.Substring(0, EndIdx - 1);
+			// Get the build root
+			Log.TraceVerbose("Looking for {0} in {1}", KnownFilePathFromRoot, KnownFileDepotMapping);
+			int EndIdx = KnownFileDepotMapping.IndexOf(KnownFilePathFromRoot, StringComparison.CurrentCultureIgnoreCase);
+			if (EndIdx < 0)
+			{
+				EndIdx = KnownFileDepotMapping.IndexOf(CommandUtils.ConvertSeparators(PathSeparator.Slash, KnownFilePathFromRoot), StringComparison.CurrentCultureIgnoreCase);
+			}
+			// Get the root path without the trailing path separator
+			BuildRootPath = KnownFileDepotMapping.Substring(0, EndIdx - 1);
 
-				// Get the client root
-				if (LocalRootPath.StartsWith(CommandUtils.CombinePaths(PathSeparator.Slash, ThisClient.RootPath, "/"), StringComparison.InvariantCultureIgnoreCase) || LocalRootPath == ThisClient.RootPath)
-				{
-					ClientRootPath = CommandUtils.CombinePaths(PathSeparator.Depot, String.Format("//{0}/", ThisClient.Name), LocalRootPath.Substring(ThisClient.RootPath.Length));
-				}
-				else
-				{
-					throw new AutomationException("LocalRootPath ({0}) does not start with the client root path ({1})", LocalRootPath, ThisClient.RootPath);
-				}
+			// Get the client root
+			if (LocalRootPath.StartsWith(CommandUtils.CombinePaths(PathSeparator.Slash, ThisClient.RootPath, "/"), StringComparison.InvariantCultureIgnoreCase) || LocalRootPath == CommandUtils.CombinePaths(PathSeparator.Slash, ThisClient.RootPath))
+			{
+				ClientRootPath = CommandUtils.CombinePaths(PathSeparator.Depot, String.Format("//{0}/", ThisClient.Name), LocalRootPath.Substring(ThisClient.RootPath.Length));
 			}
+			else
+			{
+				throw new AutomationException("LocalRootPath ({0}) does not start with the client root path ({1})", LocalRootPath, ThisClient.RootPath);
+			}
+		}
 		}
 
 		/// <summary>
diff --git a/Engine/Source/Programs/AutomationTool/AutomationUtils/P4Utils.cs b/Engine/Source/Programs/AutomationTool/AutomationUtils/P4Utils.cs
--- a/Engine/Source/Programs/AutomationTool/AutomationUtils/P4Utils.cs
+++ b/Engine/Source/Programs/AutomationTool/AutomationUtils/P4Utils.cs
@@ -2637,6 +2637,25 @@ namespace AutomationTool
 			return Output;
 		}
 
+		/// <summary>
+		/// Gets the contents of a particular file in the depot and writes it to a local file without syncing it
+		/// </summary>
+		/// <param name="DepotPath">Depot path to the file (with revision/range if necessary)</param>
+		/// <param name="OutputFileName">Output file to write to</param>
+		public void PrintToFile(string DepotPath, string FileName, bool AllowSpew = true)
+		{
+			string Output;
+			if(!P4Output(out Output, "print -q -o \"" + FileName + "\" " + DepotPath, AllowSpew: AllowSpew, WithClient: false))
+			{
+				throw new AutomationException("p4 print {0} failed", DepotPath);
+			}
+			if(!Output.Trim().Contains("\n") && Output.Contains("no such file(s)"))
+			{
+				throw new AutomationException("p4 print {0} failed", DepotPath);
+			}
+		}
+
+
 		#region Utilities
 
 		private static object[] OldStyleBinaryFlags = new object[]
