Put GitHub promotion zips on the network.

diff --git a/Engine/Source/Programs/AutomationTool/Scripts/RocketBuild.Automation.cs b/Engine/Source/Programs/AutomationTool/Scripts/RocketBuild.Automation.cs
--- a/Engine/Source/Programs/AutomationTool/Scripts/RocketBuild.Automation.cs
+++ b/Engine/Source/Programs/AutomationTool/Scripts/RocketBuild.Automation.cs
@@ -133,7 +133,16 @@ namespace Rocket
 					string GitConfigRelativePath = "Engine/Build/Git/UnrealBot.ini";
 					if(CommandUtils.FileExists(CommandUtils.CombinePaths(CommandUtils.CmdEnv.LocalRoot, GitConfigRelativePath)))
 					{
-						BranchConfig.AddNode(new BuildGitPromotable(HostPlatform, BranchConfig.HostPlatforms, GitConfigRelativePath));
+						string GitPromotableDirectory;
+						if(ShouldDoSeriousThingsLikeP4CheckinAndPostToMCP(BranchConfig))
+						{
+							GitPromotableDirectory = CommandUtils.CombinePaths(CommandUtils.RootBuildStorageDirectory(), "UE4", "GitHub", CommandUtils.P4Env.BuildRootEscaped);
+						}
+						else
+						{
+							GitPromotableDirectory = CommandUtils.CombinePaths(CommandUtils.CmdEnv.LocalRoot, "Engine", "Saved", "GitPromotable");
+						}
+						BranchConfig.AddNode(new BuildGitPromotable(HostPlatform, BranchConfig.HostPlatforms, GitConfigRelativePath, GitPromotableDirectory));
 					}
 				}
 
@@ -310,10 +319,12 @@ namespace Rocket
 	public class BuildGitPromotable : GUBP.HostPlatformNode
 	{
 		string ConfigRelativePath;
+		string PromotableDirectory;
 
-		public BuildGitPromotable(UnrealTargetPlatform HostPlatform, List<UnrealTargetPlatform> ForHostPlatforms, string InConfigRelativePath) : base(HostPlatform)
+		public BuildGitPromotable(UnrealTargetPlatform HostPlatform, List<UnrealTargetPlatform> ForHostPlatforms, string InConfigRelativePath, string InPromotableDirectory) : base(HostPlatform)
 		{
 			ConfigRelativePath = InConfigRelativePath;
+			PromotableDirectory = InPromotableDirectory;
 
 			foreach(UnrealTargetPlatform ForHostPlatform in ForHostPlatforms)
 			{
@@ -340,11 +351,16 @@ namespace Rocket
 			PromotableFilter.AddRuleForFiles(AllDependencyBuildProducts, CommandUtils.CmdEnv.LocalRoot, FileFilterType.Include);
 			PromotableFilter.ReadRulesFromFile(CommandUtils.CombinePaths(CommandUtils.CmdEnv.LocalRoot, ConfigRelativePath), "promotable");
 			PromotableFilter.ExcludeConfidentialFolders();
-			
-			// Copy everything that matches the filter to the promotion folder
-			string PromotableFolder = CommandUtils.CombinePaths(CommandUtils.CmdEnv.LocalRoot, "Engine", "Saved", "GitPromotable");
-			CommandUtils.DeleteDirectoryContents(PromotableFolder);
-			BuildProducts = CommandUtils.ThreadedCopyFiles(CommandUtils.CmdEnv.LocalRoot, PromotableFolder, PromotableFilter, bIgnoreSymlinks: true);
+
+			// Create a zip file containing the promoted binaries
+			CommandUtils.CreateDirectory(CommandUtils.GetDirectoryName(PromotableDirectory));
+			string PromotableZip = CommandUtils.CombinePaths(PromotableDirectory, String.Format("{0}.zip", CommandUtils.P4Env.Changelist));
+			CommandUtils.ZipFiles(PromotableZip + ".incoming", CommandUtils.CmdEnv.LocalRoot, PromotableFilter);
+			CommandUtils.RenameFile(PromotableZip + ".incoming", PromotableZip);
+
+			// Create a dummy build product
+			BuildProducts = new List<string>();
+			SaveRecordOfSuccessAndAddToBuildProducts();
 		}
 	}
 
