Only execute input dependencies when building a specific node. Order dependencies are only needed for setting up the EC graph.

diff --git a/Engine/Source/Programs/AutomationTool/GUBP/GUBP.cs b/Engine/Source/Programs/AutomationTool/GUBP/GUBP.cs
--- a/Engine/Source/Programs/AutomationTool/GUBP/GUBP.cs
+++ b/Engine/Source/Programs/AutomationTool/GUBP/GUBP.cs
@@ -1022,7 +1022,7 @@ public partial class GUBP : BuildCommand
 		HashSet<BuildNode> RecursiveNodesToDo = new HashSet<BuildNode>(NodesToDo);
 		foreach(BuildNode NodeToDo in NodesToDo)
 		{
-			RecursiveNodesToDo.UnionWith(NodeToDo.OrderDependencies);
+			RecursiveNodesToDo.UnionWith(NodeToDo.InputDependencies);
 		}
 
 		return RecursiveNodesToDo;
