Try to add labels for GitHub promotions every 24h or so.

diff --git a/Engine/Source/Programs/AutomationTool/Scripts/RocketBuild.Automation.cs b/Engine/Source/Programs/AutomationTool/Scripts/RocketBuild.Automation.cs
--- a/Engine/Source/Programs/AutomationTool/Scripts/RocketBuild.Automation.cs
+++ b/Engine/Source/Programs/AutomationTool/Scripts/RocketBuild.Automation.cs
@@ -133,16 +133,7 @@ namespace Rocket
 					string GitConfigRelativePath = "Engine/Build/Git/UnrealBot.ini";
 					if(CommandUtils.FileExists(CommandUtils.CombinePaths(CommandUtils.CmdEnv.LocalRoot, GitConfigRelativePath)))
 					{
-						string GitPromotableDirectory;
-						if(ShouldDoSeriousThingsLikeP4CheckinAndPostToMCP(BranchConfig))
-						{
-							GitPromotableDirectory = CommandUtils.CombinePaths(CommandUtils.RootBuildStorageDirectory(), "UE4", "GitHub", CommandUtils.P4Env.BuildRootEscaped);
-						}
-						else
-						{
-							GitPromotableDirectory = CommandUtils.CombinePaths(CommandUtils.CmdEnv.LocalRoot, "Engine", "Saved", "GitPromotable");
-						}
-						BranchConfig.AddNode(new BuildGitPromotable(HostPlatform, BranchConfig.HostPlatforms, GitConfigRelativePath, GitPromotableDirectory));
+						BranchConfig.AddNode(new LabelGitHubPromotion(HostPlatform, BranchConfig.HostPlatforms));
 					}
 				}
 
@@ -316,27 +307,21 @@ namespace Rocket
         }
     }
 
-	public class BuildGitPromotable : GUBP.HostPlatformNode
+	public class LabelGitHubPromotion : GUBP.HostPlatformNode
 	{
-		string ConfigRelativePath;
-		string PromotableDirectory;
-
-		public BuildGitPromotable(UnrealTargetPlatform HostPlatform, List<UnrealTargetPlatform> ForHostPlatforms, string InConfigRelativePath, string InPromotableDirectory) : base(HostPlatform)
+		public LabelGitHubPromotion(UnrealTargetPlatform HostPlatform, List<UnrealTargetPlatform> ForHostPlatforms) : base(HostPlatform)
 		{
-			ConfigRelativePath = InConfigRelativePath;
-			PromotableDirectory = InPromotableDirectory;
-
 			foreach(UnrealTargetPlatform ForHostPlatform in ForHostPlatforms)
 			{
-				AddDependency(GUBP.RootEditorNode.StaticGetFullName(ForHostPlatform));
-				AddDependency(GUBP.ToolsNode.StaticGetFullName(ForHostPlatform));
-				AddDependency(GUBP.InternalToolsNode.StaticGetFullName(ForHostPlatform));
+				AddPseudodependency(GUBP.RootEditorNode.StaticGetFullName(ForHostPlatform));
+				AddPseudodependency(GUBP.ToolsNode.StaticGetFullName(ForHostPlatform));
+				AddPseudodependency(GUBP.InternalToolsNode.StaticGetFullName(ForHostPlatform));
 			}
 		}
 
 		public static string StaticGetFullName(UnrealTargetPlatform HostPlatform)
 		{
-			return "BuildGitPromotable" + StaticGetHostPlatformSuffix(HostPlatform);
+			return "LabelGitHubPromotion" + StaticGetHostPlatformSuffix(HostPlatform);
 		}
 
 		public override string GetFullName()
@@ -344,19 +329,18 @@ namespace Rocket
 			return StaticGetFullName(HostPlatform);
 		}
 
+		public override int CISFrequencyQuantumShift(GUBP.GUBPBranchConfig BranchConfig)
+		{
+			return 6;
+		}
+
 		public override void DoBuild(GUBP bp)
 		{
-			// Create a filter for all the promoted binaries
-			FileFilter PromotableFilter = new FileFilter();
-			PromotableFilter.AddRuleForFiles(AllDependencyBuildProducts, CommandUtils.CmdEnv.LocalRoot, FileFilterType.Include);
-			PromotableFilter.ReadRulesFromFile(CommandUtils.CombinePaths(CommandUtils.CmdEnv.LocalRoot, ConfigRelativePath), "promotable");
-			PromotableFilter.ExcludeConfidentialFolders();
-
-			// Create a zip file containing the promoted binaries
-			CommandUtils.CreateDirectory(CommandUtils.GetDirectoryName(PromotableDirectory));
-			string PromotableZip = CommandUtils.CombinePaths(PromotableDirectory, String.Format("{0}.zip", CommandUtils.P4Env.Changelist));
-			CommandUtils.ZipFiles(PromotableZip + ".incoming", CommandUtils.CmdEnv.LocalRoot, PromotableFilter);
-			CommandUtils.RenameFile(PromotableZip + ".incoming", PromotableZip);
+			// Label everything in the branch at this changelist
+			if(CommandUtils.AllowSubmit)
+			{
+				CommandUtils.P4.MakeDownstreamLabel(CommandUtils.P4Env, "GitHub-Promotion", null);
+			}
 
 			// Create a dummy build product
 			BuildProducts = new List<string>();
