Tweaking HTML5 package compression.
When packaged HTML5 can now handle being hosted on servers that don't support static compression - done by removing explicit requests for compressed files.
Compression requires correct server setup. htaccess file included in package deploy to make this easier.
Changed HTML5LaunchHelper to switch between compressed/uncompressed files when uncompressed files do/don't exist.

diff --git a/Engine/Build/HTML5/GameX.html.template b/Engine/Build/HTML5/GameX.html.template
--- a/Engine/Build/HTML5/GameX.html.template
+++ b/Engine/Build/HTML5/GameX.html.template
@@ -140,7 +140,12 @@ addLog( "Starting download " + urlName );
 return $.ajax({
 	url: urlName,
 	dataType: "text",
-	success: function(scriptText){
+	success: function(scriptText, dataType, xhrobject) {
+		var encoding = xhrobject.getResponseHeader('Content-Encoding');
+		var type = xhrobject.getResponseHeader('Content-Type');
+		if (encoding != 'gzip' && (type == 'application/x-javascript' || type == 'application/javascript')) {
+			addLog("downloaded uncompressed javascript, please host compressed javascript when in production");
+		}
 		// insert into dom - this bit is tricky, standard way makes asm.js complain about caching. 
 		// use a blob object. 
 		var scriptblob = new Blob([scriptText], { type: 'text/javascript' });
@@ -163,9 +168,6 @@ return $.ajax({
 
 				if (!this.totaladded) { 
 					totalDownloadSize += e.total; 
-					if ( e.total > 20*1024*1024){
-						addLog("possibly downloading uncompressed javascript, please host compressed javascript when in production");
-					}
 					this.totaladded = true; 
 				}
 				  
@@ -291,7 +293,7 @@ window.onerror  = function(msg, url, line, column, error) {
         addLog( " downloading symbols file ");
         // get non minified symbols and demangle the callstack. 
 		$.ajax({
-			  url : "%CONFIG%.js.symbols.gz",
+			  url : "%CONFIG%.js.symbols",
 			  dataType: "text",
 			  success : function (data) {	
                   addLog(" symbols file downloaded, demangling..");
@@ -378,7 +380,7 @@ var Module = {
 
 		// file location. 
 		locateFile : function (name){ 
-		      return name + ".gz";
+		      return name;
 		},
 
 
@@ -435,9 +437,9 @@ $(document).ready(function() {
     addLog( "Starting downloads" );
     // start downloads in parallel, print when done.  
     var promises = [
-			'Utility.js.gz',
-			'%GAME%.data.js.gz',
-			'%CONFIG%.js.gz'   
+			'Utility.js',
+			'%GAME%.data.js',
+			'%CONFIG%.js'   
 		].map(download); 
 
 	// add the data file promise. 
diff --git a/Engine/Source/Programs/AutomationTool/HTML5/HTML5Platform.Automation.cs b/Engine/Source/Programs/AutomationTool/HTML5/HTML5Platform.Automation.cs
--- a/Engine/Source/Programs/AutomationTool/HTML5/HTML5Platform.Automation.cs
+++ b/Engine/Source/Programs/AutomationTool/HTML5/HTML5Platform.Automation.cs
@@ -148,6 +148,9 @@ public class HTML5Platform : Platform
 		string MonoPath = Path.Combine(CombinePaths(CmdEnv.LocalRoot, "Engine"), "Build", "BatchFiles", "Mac", "SetupMono.sh");
 		GenerateMacCommandFromTemplate(MacBashTemplateFile, MacBashOutputFile, MonoPath);
 
+		string htaccessTemplate = Path.Combine(CombinePaths(CmdEnv.LocalRoot, "Engine"), "Build", "HTML5", "htaccess.template");
+		File.Copy(htaccessTemplate, Path.Combine(PackagePath, ".htaccess"), true);
+
 		string JSDir = Path.Combine(CombinePaths(CmdEnv.LocalRoot, "Engine"), "Build", "HTML5");
 		string OutDir = PackagePath;
 
@@ -164,17 +167,17 @@ public class HTML5Platform : Platform
 		// Compress all files. These are independent tasks which can be threaded. 
 		Task[] CompressionTasks = new Task[6];
 		//data file.
-		CompressionTasks[0] = Task.Factory.StartNew( () => CompressFile(FinalDataLocation, FinalDataLocation + ".gz"));		
+		CompressionTasks[0] = Task.Factory.StartNew( () => CompressFile(FinalDataLocation, FinalDataLocation + "gz"));		
 		// data file .js driver.
-		CompressionTasks[1] = Task.Factory.StartNew( () => CompressFile(FinalDataLocation + ".js" , FinalDataLocation + ".js.gz"));
+		CompressionTasks[1] = Task.Factory.StartNew( () => CompressFile(FinalDataLocation + ".js" , FinalDataLocation + ".jsgz"));
 		// main js.
-		CompressionTasks[2] = Task.Factory.StartNew(() => CompressFile(Path.Combine(PackagePath, GameExe), Path.Combine(PackagePath, GameExe) + ".gz"));
+		CompressionTasks[2] = Task.Factory.StartNew(() => CompressFile(Path.Combine(PackagePath, GameExe), Path.Combine(PackagePath, GameExe) + "gz"));
 		// mem init file.
-		CompressionTasks[3] = Task.Factory.StartNew(() => CompressFile(Path.Combine(PackagePath, GameExe) + ".mem", Path.Combine(PackagePath, GameExe) + ".mem.gz"));
+		CompressionTasks[3] = Task.Factory.StartNew(() => CompressFile(Path.Combine(PackagePath, GameExe) + ".mem", Path.Combine(PackagePath, GameExe) + ".memgz"));
 		// symbols file.
-		CompressionTasks[4] = Task.Factory.StartNew(() => CompressFile(Path.Combine(PackagePath, GameExe) + ".symbols", Path.Combine(PackagePath, GameExe) + ".symbols.gz"));
+		CompressionTasks[4] = Task.Factory.StartNew(() => CompressFile(Path.Combine(PackagePath, GameExe) + ".symbols", Path.Combine(PackagePath, GameExe) + ".symbolsgz"));
 		// Utility 
-		CompressionTasks[5] = Task.Factory.StartNew(() => CompressFile(OutDir + "/Utility.js", OutDir + "/Utility.js.gz"));
+		CompressionTasks[5] = Task.Factory.StartNew(() => CompressFile(OutDir + "/Utility.js", OutDir + "/Utility.jsgz"));
 
 		File.Copy(CombinePaths(CmdEnv.LocalRoot, "Engine/Binaries/DotNET/HTML5LaunchHelper.exe"),CombinePaths(OutDir, "HTML5LaunchHelper.exe"),true);
 		Task.WaitAll(CompressionTasks);
@@ -365,17 +368,23 @@ public class HTML5Platform : Platform
 		string OutputFile = Path.Combine(PackagePath, (Params.ClientConfigsToBuild[0].ToString() != "Development" ? (Params.ShortProjectName + "-HTML5-" + Params.ClientConfigsToBuild[0].ToString()) : Params.ShortProjectName)) + ".html";
 
 		// data file 
-		SC.ArchiveFiles(PackagePath, Path.GetFileName(FinalDataLocation + ".gz"));
+		SC.ArchiveFiles(PackagePath, Path.GetFileName(FinalDataLocation));
+		SC.ArchiveFiles(PackagePath, Path.GetFileName(FinalDataLocation + "gz"));
 		// data file js driver 
-		SC.ArchiveFiles(PackagePath, Path.GetFileName(FinalDataLocation + ".js.gz"));
+		SC.ArchiveFiles(PackagePath, Path.GetFileName(FinalDataLocation + ".js"));
+		SC.ArchiveFiles(PackagePath, Path.GetFileName(FinalDataLocation + ".jsgz"));
 		// main js file
-		SC.ArchiveFiles(PackagePath, Path.GetFileName(GameExe + ".gz"));
+		SC.ArchiveFiles(PackagePath, Path.GetFileName(GameExe));
+		SC.ArchiveFiles(PackagePath, Path.GetFileName(GameExe + "gz"));
 		// memory init file
-		SC.ArchiveFiles(PackagePath, Path.GetFileName(GameExe + ".mem.gz"));
+		SC.ArchiveFiles(PackagePath, Path.GetFileName(GameExe + ".mem"));
+		SC.ArchiveFiles(PackagePath, Path.GetFileName(GameExe + ".memgz"));
 		// symbols file
-		SC.ArchiveFiles(PackagePath, Path.GetFileName(GameExe + ".symbols.gz"));
+		SC.ArchiveFiles(PackagePath, Path.GetFileName(GameExe + ".symbols"));
+		SC.ArchiveFiles(PackagePath, Path.GetFileName(GameExe + ".symbolsgz"));
 		// utilities
-		SC.ArchiveFiles(PackagePath, Path.GetFileName("Utility.js.gz"));
+		SC.ArchiveFiles(PackagePath, Path.GetFileName("Utility.js"));
+		SC.ArchiveFiles(PackagePath, Path.GetFileName("Utility.jsgz"));
 		// landing page.
 		SC.ArchiveFiles(PackagePath, Path.GetFileName(OutputFile));
 
@@ -384,6 +393,7 @@ public class HTML5Platform : Platform
 		SC.ArchiveFiles(LaunchHelperPath, "HTML5LaunchHelper.exe");
 		SC.ArchiveFiles(Path.Combine(CombinePaths(CmdEnv.LocalRoot, "Engine"), "Build", "HTML5"), "Readme.txt");
 		SC.ArchiveFiles(PackagePath, Path.GetFileName(Path.Combine(PackagePath, "RunMacHTML5LaunchHelper.command")));
+		SC.ArchiveFiles(PackagePath, Path.GetFileName(Path.Combine(PackagePath, ".htaccess")));
 
 		if (HTMLPakAutomation.CanCreateMapPaks(Params))
 		{
@@ -570,8 +580,8 @@ public class HTML5Platform : Platform
 	private static IDictionary<string, string> MimeTypeMapping = new Dictionary<string, string>(StringComparer.InvariantCultureIgnoreCase) 
 		{
 			{ ".html", "text/html"},
-			{ ".js.gz", "application/x-javascript" },  // upload compressed javascript. 
-			{ ".data.gz", "appication/octet-stream"}
+			{ ".jsgz", "application/x-javascript" },  // upload compressed javascript. 
+			{ ".datagz", "appication/octet-stream"}
 		}; 
 
 	public void UploadToS3Worker(FileInfo Info, string KeyId, string AccessKey, string BucketName, string FolderName )
@@ -603,7 +613,7 @@ public class HTML5Platform : Platform
 			Request.Headers["x-amz-acl"] = "public-read"; // we probably want to make public read by default. 
 
 			// set correct content encoding for compressed javascript. 
-			if ( Info.Extension == ".gz")
+			if ( Info.Extension.EndsWith("gz") )
 			{
 				Request.Headers["Content-Encoding"] = "gzip";
 			}
diff --git a/Engine/Source/Programs/HTML5/HTML5LaunchHelper/Program.cs b/Engine/Source/Programs/HTML5/HTML5LaunchHelper/Program.cs
--- a/Engine/Source/Programs/HTML5/HTML5LaunchHelper/Program.cs
+++ b/Engine/Source/Programs/HTML5/HTML5LaunchHelper/Program.cs
@@ -164,6 +164,11 @@ namespace HTML5LaunchHelper
 			else if (File.Exists(Root + Context.Request.Url.LocalPath))
 			{
 				string RequestedFile = Root + Context.Request.Url.LocalPath;
+				string RequestedFileCompressed = Root + Context.Request.Url.LocalPath + "gz";
+				if (File.Exists(RequestedFileCompressed)) 
+				{
+					RequestedFile = RequestedFileCompressed;
+				}
 				System.Console.WriteLine("Serving " + RequestedFile);
 				using (Stream source = File.OpenRead(RequestedFile))
 				{
@@ -175,7 +180,7 @@ namespace HTML5LaunchHelper
 						MimeType = MimeTypeMapping[Extention];
 					}
 					// This is the crux of serving pre-compressed files. 
-					if (Extention == ".gz")
+					if (Extention.EndsWith("gz"))
 					{
 						Context.Response.AddHeader("Content-Encoding", "gzip");
 					}
