Fix numerous path slash issues for HTML5 packaging on Mac.
Fixes UE-20141

diff --git a/Engine/Source/Programs/AutomationTool/HTML5/HTML5Platform.PakFiles.Automation.cs b/Engine/Source/Programs/AutomationTool/HTML5/HTML5Platform.PakFiles.Automation.cs
--- a/Engine/Source/Programs/AutomationTool/HTML5/HTML5Platform.PakFiles.Automation.cs
+++ b/Engine/Source/Programs/AutomationTool/HTML5/HTML5Platform.PakFiles.Automation.cs
@@ -73,7 +73,7 @@ public class HTMLPakAutomation
 			// we need to operate in the root
 			using (new PushedDirectory(Path.Combine(PackagePath)))
 			{
-				string CmdLine = string.Format("\"{0}\" \"{1}\" --preload {2}\\ --js-output=\"{1}.js\"", PackagerPath, FinalDataLocation, SC.ShortProjectName);
+				string CmdLine = string.Format("\"{0}\" \"{1}\" --preload \"{2}\" --js-output=\"{1}.js\"", PackagerPath, FinalDataLocation, SC.ShortProjectName);
 				CommandUtils.RunAndLog(CommandUtils.CmdEnv, PythonPath, CmdLine);
 			}
 		}
@@ -93,7 +93,7 @@ public class HTMLPakAutomation
 		 }
 
 		 // find list of files in the Engine Dir. 
-		 string[] files = Directory.GetFiles( StagedDir + "\\Engine\\", "*", SearchOption.AllDirectories);
+		 string[] files = Directory.GetFiles( Path.Combine(StagedDir, "Engine"), "*", SearchOption.AllDirectories);
 
 		 Dictionary<string,string> UnrealPakResponseFile = new Dictionary<string,string>();
 
@@ -126,7 +126,7 @@ public class HTMLPakAutomation
 		List<string> AllFiles = new List<string>();
 		try
 		{
-			string[] SlateFiles = Directory.GetFiles(StagedDir + "\\" + Params.ShortProjectName + "\\Content\\Slate", "*", SearchOption.AllDirectories);
+			string[] SlateFiles = Directory.GetFiles(Path.Combine(StagedDir, Params.ShortProjectName, "Content", "Slate"), "*", SearchOption.AllDirectories);
 			AllFiles.AddRange(SlateFiles);
 		}
 		catch ( Exception /*Ex*/)
@@ -136,7 +136,7 @@ public class HTMLPakAutomation
 
 		try
 		{
-			string[] Resources = Directory.GetFiles(StagedDir + "\\" + Params.ShortProjectName + "\\Content\\Resources", "*", SearchOption.AllDirectories);
+			string[] Resources = Directory.GetFiles(Path.Combine(StagedDir, Params.ShortProjectName, "Content", "Resources"), "*", SearchOption.AllDirectories);
 			AllFiles.AddRange(Resources);
 		}
 		catch (Exception /*Ex*/)
@@ -144,11 +144,11 @@ public class HTMLPakAutomation
 			// not found. 
 		}
 
-		string[] ConfigFiles = Directory.GetFiles(StagedDir + "\\" + Params.ShortProjectName + "\\Config", "*", SearchOption.AllDirectories);
+		string[] ConfigFiles = Directory.GetFiles(Path.Combine(StagedDir, Params.ShortProjectName, "Config"), "*", SearchOption.AllDirectories);
 
 		AllFiles.AddRange(ConfigFiles);
 
-		string[] TopLevelFiles = Directory.GetFiles(StagedDir + "\\" + Params.ShortProjectName + "\\", "*", SearchOption.TopDirectoryOnly);
+		string[] TopLevelFiles = Directory.GetFiles(Path.Combine(StagedDir, Params.ShortProjectName), "*", SearchOption.TopDirectoryOnly);
 
 		AllFiles.AddRange(TopLevelFiles);
 
@@ -250,6 +250,7 @@ public class HTMLPakAutomation
 	/// <param name="ResponseFile"></param>
 	private static void WritePakResponseFile(string Filename, Dictionary<string, string> ResponseFile, bool Compressed)
 	{
+		Directory.CreateDirectory(Path.GetDirectoryName(Filename));
 		using (var Writer = new StreamWriter(Filename, false, new System.Text.UTF8Encoding(true)))
 		{
 			foreach (var Entry in ResponseFile)
