Don't write zip when storing to temp storage unless you actually have a file to add in your thread.

diff --git a/Engine/Source/Programs/AutomationTool/GUBP/TempStorage.cs b/Engine/Source/Programs/AutomationTool/GUBP/TempStorage.cs
--- a/Engine/Source/Programs/AutomationTool/GUBP/TempStorage.cs
+++ b/Engine/Source/Programs/AutomationTool/GUBP/TempStorage.cs
@@ -861,27 +861,30 @@ namespace AutomationTool
                     let ZipFileName = Path.Combine(StagingDir ?? OutputDir, string.Format("{0}{1}.zip", ZipBasename, bZipInParallel ? "-" + CoreNum.ToString("00") : ""))
                     select new Thread(() =>
                     {
-                        // Create one zip per thread using the given basename
-                        using (var ZipArchive = new Ionic.Zip.ZipFile(ZipFileName) { CompressionLevel = Ionic.Zlib.CompressionLevel.BestSpeed })
+                        // don't create the zip unless we have at least one file to add
+                        string File;
+                        if (FilesToZip.TryDequeue(out File))
                         {
+                            // Create one zip per thread using the given basename
+                            using (var ZipArchive = new Ionic.Zip.ZipFile(ZipFileName) { CompressionLevel = Ionic.Zlib.CompressionLevel.BestSpeed })
+                            {
 
-                            // pull from the queue until we are out of files.
-                            string File;
-                            while (FilesToZip.TryDequeue(out File))
+                                // pull from the queue until we are out of files.
+                                do
+                                {
+                                    // use fastest compression. In our best case we are CPU bound, so this is a good tradeoff,
+                                    // cutting overall time by 2/3 while only modestly increasing the compression ratio (22.7% -> 23.8% for RootEditor PDBs).
+                                    // This is in cases of a super hot cache, so the operation was largely CPU bound.
+                                    ZipArchive.AddFile(File, CommandUtils.ConvertSeparators(PathSeparator.Slash, Path.GetDirectoryName(CommandUtils.StripBaseDirectory(File, RootDir))));
+                                } while (FilesToZip.TryDequeue(out File));
+                            }
+                            Interlocked.Add(ref ZipFilesTotalSize, new FileInfo(ZipFileName).Length);
+                            // if we are using a staging dir, copy to the final location and delete the staged copy.
+                            if (StagingDir != null)
                             {
-                                // use fastest compression. In our best case we are CPU bound, so this is a good tradeoff,
-                                // cutting overall time by 2/3 while only modestly increasing the compression ratio (22.7% -> 23.8% for RootEditor PDBs).
-                                // This is in cases of a super hot cache, so the operation was largely CPU bound.
-                                ZipArchive.AddFile(File, CommandUtils.ConvertSeparators(PathSeparator.Slash, Path.GetDirectoryName(CommandUtils.StripBaseDirectory(File, RootDir))));
+                                CommandUtils.CopyFile(ZipFileName, CommandUtils.MakeRerootedFilePath(ZipFileName, StagingDir, OutputDir));
+                                CommandUtils.DeleteFile(true, ZipFileName);
                             }
-                            ZipArchive.Save();
-                        }
-                        Interlocked.Add(ref ZipFilesTotalSize, new FileInfo(ZipFileName).Length);
-                        // if we are using a staging dir, copy to the final location and delete the staged copy.
-                        if (StagingDir != null)
-                        {
-                            CommandUtils.CopyFile(ZipFileName, CommandUtils.MakeRerootedFilePath(ZipFileName, StagingDir, OutputDir));
-                            CommandUtils.DeleteFile(true, ZipFileName);
                         }
                     })).ToList();
             }
@@ -892,29 +895,33 @@ namespace AutomationTool
                     let ZipFileName = Path.Combine(StagingDir ?? OutputDir, string.Format("{0}{1}.zip", ZipBasename, bZipInParallel ? "-" + CoreNum.ToString("00") : ""))
                     select new Thread(() =>
                     {
-                        // Create one zip per thread using the given basename
-                        using (var ZipArchive = System.IO.Compression.ZipFile.Open(ZipFileName, System.IO.Compression.ZipArchiveMode.Create))
+                        // don't create the zip unless we have at least one file to add
+                        string File;
+                        if (FilesToZip.TryDequeue(out File))
                         {
+                            // Create one zip per thread using the given basename
+                            using (var ZipArchive = System.IO.Compression.ZipFile.Open(ZipFileName, System.IO.Compression.ZipArchiveMode.Create))
+                            {
 
-                            // pull from the queue until we are out of files.
-                            string File;
-                            while (FilesToZip.TryDequeue(out File))
+                                // pull from the queue until we are out of files.
+                                do
+                                {
+                                    // use fastest compression. In our best case we are CPU bound, so this is a good tradeoff,
+                                    // cutting overall time by 2/3 while only modestly increasing the compression ratio (22.7% -> 23.8% for RootEditor PDBs).
+                                    // This is in cases of a super hot cache, so the operation was largely CPU bound.
+                                    // Also, sadly, mono appears to have a bug where nothing you can do will properly set the LastWriteTime on the created entry,
+                                    // so we have to ignore timestamps on files extracted from a zip, since it may have been created on a Mac.
+                                    ZipArchive.CreateEntryFromFile(File, CommandUtils.ConvertSeparators(PathSeparator.Slash, CommandUtils.StripBaseDirectory(File, RootDir)), System.IO.Compression.CompressionLevel.Fastest);
+                                } while (FilesToZip.TryDequeue(out File));
+                            }
+                            Interlocked.Add(ref ZipFilesTotalSize, new FileInfo(ZipFileName).Length);
+                            // if we are using a staging dir, copy to the final location and delete the staged copy.
+                            if (StagingDir != null)
                             {
-                                // use fastest compression. In our best case we are CPU bound, so this is a good tradeoff,
-                                // cutting overall time by 2/3 while only modestly increasing the compression ratio (22.7% -> 23.8% for RootEditor PDBs).
-                                // This is in cases of a super hot cache, so the operation was largely CPU bound.
-                                // Also, sadly, mono appears to have a bug where nothing you can do will properly set the LastWriteTime on the created entry,
-                                // so we have to ignore timestamps on files extracted from a zip, since it may have been created on a Mac.
-                                ZipArchive.CreateEntryFromFile(File, CommandUtils.ConvertSeparators(PathSeparator.Slash, CommandUtils.StripBaseDirectory(File, RootDir)), System.IO.Compression.CompressionLevel.Fastest);
+                                CommandUtils.CopyFile(ZipFileName, CommandUtils.MakeRerootedFilePath(ZipFileName, StagingDir, OutputDir));
+                                CommandUtils.DeleteFile(true, ZipFileName);
                             }
                         }
-                        Interlocked.Add(ref ZipFilesTotalSize, new FileInfo(ZipFileName).Length);
-                        // if we are using a staging dir, copy to the final location and delete the staged copy.
-                        if (StagingDir != null)
-                        {
-                            CommandUtils.CopyFile(ZipFileName, CommandUtils.MakeRerootedFilePath(ZipFileName, StagingDir, OutputDir));
-                            CommandUtils.DeleteFile(true, ZipFileName);
-                        }
                     })).ToList();
             }
             ZipThreads.ForEach(thread => thread.Start());
