Fix temp storage on mac. I accidentally removed the line that saved the archive on Mono. .NET API doesn't have an explicit "save" call so I forgot to look for it.

diff --git a/Engine/Source/Programs/AutomationTool/GUBP/TempStorage.cs b/Engine/Source/Programs/AutomationTool/GUBP/TempStorage.cs
--- a/Engine/Source/Programs/AutomationTool/GUBP/TempStorage.cs
+++ b/Engine/Source/Programs/AutomationTool/GUBP/TempStorage.cs
@@ -877,6 +877,7 @@ namespace AutomationTool
                                     // This is in cases of a super hot cache, so the operation was largely CPU bound.
                                     ZipArchive.AddFile(File, CommandUtils.ConvertSeparators(PathSeparator.Slash, Path.GetDirectoryName(CommandUtils.StripBaseDirectory(File, RootDir))));
                                 } while (FilesToZip.TryDequeue(out File));
+                                ZipArchive.Save();
                             }
                             Interlocked.Add(ref ZipFilesTotalSize, new FileInfo(ZipFileName).Length);
                             // if we are using a staging dir, copy to the final location and delete the staged copy.
