fix for single quotes not being converted to double quotes in the additional command line parameter in UAT

diff --git a/Engine/Source/Programs/AutomationTool/AutomationUtils/ProjectParams.cs b/Engine/Source/Programs/AutomationTool/AutomationUtils/ProjectParams.cs
--- a/Engine/Source/Programs/AutomationTool/AutomationUtils/ProjectParams.cs
+++ b/Engine/Source/Programs/AutomationTool/AutomationUtils/ProjectParams.cs
@@ -559,6 +559,7 @@ namespace AutomationTool
 			this.StageCommandline = ParseParamValueIfNotSpecified(Command, StageCommandline, "cmdline");
 			this.BundleName = ParseParamValueIfNotSpecified(Command, BundleName, "bundlename");
 			this.RunCommandline = ParseParamValueIfNotSpecified(Command, RunCommandline, "addcmdline");
+			this.RunCommandline = this.RunCommandline.Replace('\'', '\"'); // replace any single quotes with double quotes
 			this.Package = GetParamValueIfNotSpecified(Command, Package, this.Package, "package");
 			this.Deploy = GetParamValueIfNotSpecified(Command, Deploy, this.Deploy, "deploy");
 			this.IterativeDeploy = GetParamValueIfNotSpecified(Command, IterativeDeploy, this.IterativeDeploy, new string[] {"iterativedeploy", "iterate" } );
