Fixed packaging on Mac

diff --git a/Engine/Source/Programs/AutomationTool/Mac/MacPlatform.Automation.cs b/Engine/Source/Programs/AutomationTool/Mac/MacPlatform.Automation.cs
--- a/Engine/Source/Programs/AutomationTool/Mac/MacPlatform.Automation.cs
+++ b/Engine/Source/Programs/AutomationTool/Mac/MacPlatform.Automation.cs
@@ -49,7 +49,7 @@ public class MacPlatform : Platform
 			{
 				string ExtensionWildcard = "*" + DebugExtention;
 				DebugExtentionWildCards.Add(ExtensionWildcard);
-				SC.StageFiles(StagedFileType.DebugNonUFS, InPath, ExtensionWildcard, true, null, NewName, false, true, null);
+				SC.StageFiles(StagedFileType.DebugNonUFS, InPath, ExtensionWildcard, true, null, NewName, true, true, null);
 			}
 
 			// Also stage the non-debug files, excluding the debug ones staged above
@@ -367,7 +367,7 @@ public class MacPlatform : Platform
 				File.Delete(GameExePath);
 				File.Move(UE4GamePath, GameExePath);
 
-				string DefaultIconPath = CombinePaths(BundlePath, "Contents", "Resources", "UE4.icns");
+				string DefaultIconPath = CombinePaths(BundlePath, "Contents", "Resources", "UE4Game.icns");
 				string CustomIconSrcPath = CombinePaths(BundlePath, "Contents", "Resources", "Application.icns");
 				string CustomIconDestPath = CombinePaths(BundlePath, "Contents", "Resources", SC.ShortProjectName + ".icns");
 				if (File.Exists(CustomIconSrcPath))
