fix for attempting to sign a list of files with an empty string

diff --git a/Engine/Source/Programs/AutomationTool/AutomationUtils/CommandUtils.cs b/Engine/Source/Programs/AutomationTool/AutomationUtils/CommandUtils.cs
--- a/Engine/Source/Programs/AutomationTool/AutomationUtils/CommandUtils.cs
+++ b/Engine/Source/Programs/AutomationTool/AutomationUtils/CommandUtils.cs
@@ -2479,7 +2479,12 @@ namespace AutomationTool
 				throw new AutomationException("SignTool not found at '{0}' (are you missing the Windows SDK?)", SignToolName);
 			}
 
-			// Code sign the executable
+			// nothing to sign
+			if (String.IsNullOrEmpty(FilesToSign))
+			{
+				return;
+			}
+
 			// Code sign the executable
 			string[] TimestampServer = { "http://timestamp.verisign.com/scripts/timestamp.dll",
 									     "http://timestamp.globalsign.com/scripts/timstamp.dll",
