Add a more resilient check for whether a target platform is available; check that the actual nodes exist. Prevents auto-configuration problems in other branches when bNoInstalledEngine is not set correctly.

diff --git a/Engine/Source/Programs/AutomationTool/Scripts/RocketBuild.Automation.cs b/Engine/Source/Programs/AutomationTool/Scripts/RocketBuild.Automation.cs
--- a/Engine/Source/Programs/AutomationTool/Scripts/RocketBuild.Automation.cs
+++ b/Engine/Source/Programs/AutomationTool/Scripts/RocketBuild.Automation.cs
@@ -86,8 +86,8 @@ namespace Rocket
 				// Find all the target platforms for this host platform.
 				List<UnrealTargetPlatform> TargetPlatforms = GetTargetPlatforms(bp, HostPlatform);
 
-				// Remove any platforms that aren't available on this machine
-				TargetPlatforms.RemoveAll(x => !ActivePlatforms.Contains(x));
+				// Remove any target platforms that aren't available
+				TargetPlatforms.RemoveAll(x => !IsTargetPlatformAvailable(BranchConfig, HostPlatform, x));
 
 				// Get the temp directory for stripped files for this host
 				string StrippedDir = Path.GetFullPath(CommandUtils.CombinePaths(CommandUtils.CmdEnv.LocalRoot, "Engine", "Saved", "Rocket", HostPlatform.ToString()));
@@ -159,6 +159,20 @@ namespace Rocket
 			return FEngineVersionSupport.FromVersionFile(CommandUtils.CombinePaths(CommandUtils.CmdEnv.LocalRoot, @"Engine\Source\Runtime\Launch\Resources\Version.h")).ToString();
 		}
 
+		public static bool IsTargetPlatformAvailable(GUBP.GUBPBranchConfig BranchConfig, UnrealTargetPlatform HostPlatform, UnrealTargetPlatform TargetPlatform)
+		{
+			UnrealTargetPlatform SourceHostPlatform = GetSourceHostPlatform(BranchConfig.HostPlatforms, HostPlatform, TargetPlatform);
+
+			bool bIsCodeTargetPlatform = IsCodeTargetPlatform(HostPlatform, TargetPlatform);
+
+			if(!BranchConfig.HasNode(GUBP.GamePlatformMonolithicsNode.StaticGetFullName(SourceHostPlatform, BranchConfig.Branch.BaseEngineProject, TargetPlatform, false, bIsCodeTargetPlatform)))
+			{
+				return false;
+			}
+
+			return true;
+		}
+
 		public static List<UnrealTargetPlatform> GetTargetPlatforms(BuildCommand Command, UnrealTargetPlatform HostPlatform)
 		{
 			List<UnrealTargetPlatform> TargetPlatforms = new List<UnrealTargetPlatform>();
