Enabled the Localise commandlet for the UE4 editor on the Main branch

diff --git a/Engine/Source/Programs/AutomationTool/GUBP/LegacyBranchSetup.cs b/Engine/Source/Programs/AutomationTool/GUBP/LegacyBranchSetup.cs
--- a/Engine/Source/Programs/AutomationTool/GUBP/LegacyBranchSetup.cs
+++ b/Engine/Source/Programs/AutomationTool/GUBP/LegacyBranchSetup.cs
@@ -170,6 +170,11 @@ partial class GUBP
             /// If the mapping does not exist for a game, it uses the UE4 folder.
             /// </summary>
             public readonly Dictionary<string, string> GameNameToBuildShareMapping = new Dictionary<string,string>();
+			/// <summary>
+			/// Allows a branch to override the "Localise" commandlet settings.
+			/// </summary>
+			public bool bShouldBuildLocalization = false;
+			public string LocalizationBranchSuffix = "";
         }
         public virtual void ModifyOptions(GUBP bp, ref BranchOptions Options, string Branch)
         {
@@ -1114,9 +1119,9 @@ partial class GUBP
 			}
 		}
 
-		if(HostPlatforms.Contains(UnrealTargetPlatform.Win64) && BranchConfig.HasNode(RootEditorNode.StaticGetFullName(UnrealTargetPlatform.Win64)))
+		if (HostPlatforms.Contains(UnrealTargetPlatform.Win64) && BranchConfig.HasNode(RootEditorNode.StaticGetFullName(UnrealTargetPlatform.Win64)) && BranchOptions.bShouldBuildLocalization)
 		{
-			BranchConfig.AddNode(new BuildLocalization());
+			BranchConfig.AddNode(new BuildLocalization(BranchOptions.LocalizationBranchSuffix));
 		}
 
         BranchConfig.AddNode(new WaitForTestShared(this));
diff --git a/Engine/Source/Programs/AutomationTool/GUBP/LegacyNodes.cs b/Engine/Source/Programs/AutomationTool/GUBP/LegacyNodes.cs
--- a/Engine/Source/Programs/AutomationTool/GUBP/LegacyNodes.cs
+++ b/Engine/Source/Programs/AutomationTool/GUBP/LegacyNodes.cs
@@ -2620,16 +2620,19 @@ partial class GUBP
 
 	public class BuildLocalization : HostPlatformNode
 	{
-		public BuildLocalization()
+		public BuildLocalization(string InLocalizationBranchSuffix)
 			: base(UnrealTargetPlatform.Win64)
 		{
+			LocalizationBranchSuffix = InLocalizationBranchSuffix;
+
 			AddDependency(RootEditorNode.StaticGetFullName(HostPlatform));
 		}
 
 		public override void DoBuild(GUBP bp)
 		{
 			// Run the localise script.
-			//CommandUtils.RunUAT(CommandUtils.CmdEnv, "Localise");
+			// todo: Will need to pass along LocalizationBranchSuffix to the commandlet (once this is set up to work for branches other than Main)
+			CommandUtils.RunUAT(CommandUtils.CmdEnv, "Localise");
 
 			// Don't pass on any build products to other build nodes at the moment.
 			BuildProducts = new List<string>();
@@ -2650,6 +2653,8 @@ partial class GUBP
 		{
 			return StaticGetFullName();
 		}
+
+		private string LocalizationBranchSuffix;
 	}
 
     public class GameAggregateNode : HostPlatformAggregateNode
