Commenting out LocalizationBranchSuffix for now in order to fix a Mono warning on Mac

diff --git a/Engine/Source/Programs/AutomationTool/GUBP/LegacyNodes.cs b/Engine/Source/Programs/AutomationTool/GUBP/LegacyNodes.cs
--- a/Engine/Source/Programs/AutomationTool/GUBP/LegacyNodes.cs
+++ b/Engine/Source/Programs/AutomationTool/GUBP/LegacyNodes.cs
@@ -2623,7 +2623,7 @@ partial class GUBP
 		public BuildLocalization(string InLocalizationBranchSuffix)
 			: base(UnrealTargetPlatform.Win64)
 		{
-			LocalizationBranchSuffix = InLocalizationBranchSuffix;
+			//LocalizationBranchSuffix = InLocalizationBranchSuffix;
 
 			AddDependency(RootEditorNode.StaticGetFullName(HostPlatform));
 		}
@@ -2654,7 +2654,7 @@ partial class GUBP
 			return StaticGetFullName();
 		}
 
-		private string LocalizationBranchSuffix;
+		//private string LocalizationBranchSuffix;
 	}
 
     public class GameAggregateNode : HostPlatformAggregateNode
