Add a utility function to get the describe record for a single changelist.

diff --git a/Engine/Source/Programs/AutomationTool/AutomationUtils/P4Utils.cs b/Engine/Source/Programs/AutomationTool/AutomationUtils/P4Utils.cs
--- a/Engine/Source/Programs/AutomationTool/AutomationUtils/P4Utils.cs
+++ b/Engine/Source/Programs/AutomationTool/AutomationUtils/P4Utils.cs
@@ -922,6 +922,33 @@ namespace AutomationTool
             }
         }
 
+		/// <summary>
+		/// Wraps P4 describe
+		/// </summary>
+		/// <param name="Changelist">Changelist numbers to query full descriptions for</param>
+		/// <param name="DescribeRecord">Describe record for the given changelist.</param>
+		/// <param name="AllowSpew"></param>
+		/// <returns>True if everything went okay</returns>
+        public bool DescribeChangelist(int Changelist, out DescribeRecord DescribeRecord, bool AllowSpew = true)
+        {
+			List<DescribeRecord> DescribeRecords;
+			if(!DescribeChangelists(new List<int>{ Changelist }, out DescribeRecords, AllowSpew))
+			{
+				DescribeRecord = null;
+				return false;
+			}
+			else if(DescribeRecords.Count != 1)
+			{
+				DescribeRecord = null;
+				return false;
+			}
+			else 
+			{
+				DescribeRecord = DescribeRecords[0];
+				return true;
+			}
+		}
+
 		/// <summary>
 		/// Wraps P4 describe
 		/// </summary>
