Additional functionality for Perfore/Jira integration: now fills in the NewBranchFixedIn field for every changelist marked fixed, and updates MainCL field for every change integrated into //UE4/Main.

diff --git a/Engine/Source/Programs/AutomationTool/AutomationUtils/P4Utils.cs b/Engine/Source/Programs/AutomationTool/AutomationUtils/P4Utils.cs
--- a/Engine/Source/Programs/AutomationTool/AutomationUtils/P4Utils.cs
+++ b/Engine/Source/Programs/AutomationTool/AutomationUtils/P4Utils.cs
@@ -2682,6 +2682,67 @@ namespace AutomationTool
 			}
 		}
 
+		/// <summary>
+		/// Runs the 'interchanges' command on a stream, to determine a list of changelists that have not been integrated to its parent (or vice-versa, if bReverse is set).
+		/// </summary>
+		/// <param name="StreamName">The name of the stream, eg. //UE4/Dev-Animation</param>
+		/// <param name="bReverse">If true, returns changes that have not been merged from the parent stream into this one.</param>
+		/// <returns>List of changelist numbers that are pending integration</returns>
+		public List<int> StreamInterchanges(string StreamName, bool bReverse)
+		{
+			string Output;
+			if(!P4Output(out Output, String.Format("interchanges {0}-S {1}", bReverse? "-r " : "", StreamName), Input:null, AllowSpew:false))
+			{
+				throw new AutomationException("Couldn't get unintegrated stream changes from {0}", StreamName);
+			}
+
+			List<int> Changelists = new List<int>();
+			if(!Output.StartsWith("All revision(s) already integrated"))
+			{
+				foreach(string Line in Output.Split('\n'))
+				{
+					string[] Tokens = Line.Split(new char[]{ ' ' }, StringSplitOptions.RemoveEmptyEntries);
+					if(Tokens.Length > 0)
+					{
+						int Changelist;
+						if(Tokens[0] != "Change" || !int.TryParse(Tokens[1], out Changelist))
+						{
+							throw new AutomationException("Unexpected output from p4 interchanges: {0}", Line);
+						}
+						Changelists.Add(Changelist);
+					}
+				}
+			}
+			return Changelists;
+		}
+
+		/// <summary>
+		/// For a given file (and revision, potentially), returns where it was integrated from. Useful in conjunction with files in a P4DescribeRecord, with action = "integrate".
+		/// </summary>
+		/// <param name="DepotPath">The file to check. May have a revision specifier at the end (eg. //depot/UE4/foo.cpp#2) </param>
+		/// <returns>The file that it was integrated from, without a revision specifier</returns>
+		public string GetIntegrationSource(string DepotPath)
+		{
+			string Output;
+			if(P4Output(out Output, "filelog -m 1 \"" + DepotPath + "\"", Input:null, AllowSpew:false))
+			{
+				foreach(string Line in Output.Split('\n').Select(x => x.Trim()))
+				{
+					const string MergePrefix = "... ... merge from ";
+					if(Line.StartsWith(MergePrefix))
+					{
+						return Line.Substring(MergePrefix.Length, Line.LastIndexOf('#') - MergePrefix.Length);
+					}
+
+					const string CopyPrefix = "... ... copy from ";
+					if(Line.StartsWith(CopyPrefix))
+					{
+						return Line.Substring(CopyPrefix.Length, Line.LastIndexOf('#') - CopyPrefix.Length);
+					}
+				}
+			}
+			throw new AutomationException("Failed to get integration source for {0}", DepotPath);
+		}
 
 		#region Utilities
 
