Reenable documentation builds in main branch.

diff --git a/Engine/Source/Programs/AutomationTool/Scripts/BuildDocumentation.Automation.cs b/Engine/Source/Programs/AutomationTool/Scripts/BuildDocumentation.Automation.cs
--- a/Engine/Source/Programs/AutomationTool/Scripts/BuildDocumentation.Automation.cs
+++ b/Engine/Source/Programs/AutomationTool/Scripts/BuildDocumentation.Automation.cs
@@ -139,11 +139,11 @@ public class DocumentationNodeAdder : GUBP.GUBPNodeAdder
 {
 	public override void AddNodes(GUBP bp, GUBP.GUBPBranchConfig BranchConfig, UnrealTargetPlatform InHostPlatform, List<UnrealTargetPlatform> InActivePlatforms)
 	{
-/*		if(InHostPlatform == UnrealTargetPlatform.Win64 && !BranchConfig.BranchOptions.bNoDocumentation)
+		if(InHostPlatform == UnrealTargetPlatform.Win64 && !BranchConfig.BranchOptions.bNoDocumentation)
 		{
 			BranchConfig.AddNode(new ToolsForDocumentationNode(BranchConfig, InHostPlatform));
 			BranchConfig.AddNode(new CodeDocumentationNode(InHostPlatform));
 			BranchConfig.AddNode(new BlueprintDocumentationNode(InHostPlatform));
-		}*/
+		}
 	}
 }
