Reduce priority for Documentation builds, and decrease the frequency.

diff --git a/Engine/Source/Programs/AutomationTool/Scripts/BuildDocumentation.Automation.cs b/Engine/Source/Programs/AutomationTool/Scripts/BuildDocumentation.Automation.cs
--- a/Engine/Source/Programs/AutomationTool/Scripts/BuildDocumentation.Automation.cs
+++ b/Engine/Source/Programs/AutomationTool/Scripts/BuildDocumentation.Automation.cs
@@ -26,9 +26,14 @@ class ToolsForDocumentationNode : GUBP.CompileNode
         return StaticGetFullName(HostPlatform);
     }
 
+	public override float Priority()
+	{
+		return -100000.0f;
+	}
+
 	public override int CISFrequencyQuantumShift(GUBP.GUBPBranchConfig BranchConfig)
 	{
-		return base.CISFrequencyQuantumShift(BranchConfig) + 3;
+		return base.CISFrequencyQuantumShift(BranchConfig) + 6;
 	}
 
 	public override UE4Build.BuildAgenda GetAgenda(GUBP bp)
