Made the Localise automation script more data driven, and removed some hard-coded assumptions

The Localise automation script processing was failing as we have acquired the "zh-Hans-CN" culture in our OneSky project, but the UE4 commandlet doesn't generate that culture, only "zh-CN". This was causing the script to fail as it couldn't find the data for that culture in Perforce.

This change has the Localise automation script read the same config file that the UE4 commandlet will use, and then use that config file data to work out what cultures it will export/upload to OneSky.

This change should also make the script generic enough for it to be used for other projects (once we can pass in the arguments on the command line).

diff --git a/Engine/Source/Programs/AutomationTool/Scripts/Localisation.Automation.cs b/Engine/Source/Programs/AutomationTool/Scripts/Localisation.Automation.cs
--- a/Engine/Source/Programs/AutomationTool/Scripts/Localisation.Automation.cs
+++ b/Engine/Source/Programs/AutomationTool/Scripts/Localisation.Automation.cs
@@ -12,17 +12,32 @@ using EpicGames.OneSkyLocalization.Config;
 
 class Localise : BuildCommand
 {
-    private void ExportProjectToDirectory(OneSkyService oneSkyService, ProjectGroup projectGroup, string ProjectName)
-    {
-        var project = GetProject(oneSkyService, projectGroup.Name, ProjectName);
-        var file = project.UploadedFiles.FirstOrDefault(f => f.Filename == (ProjectName + ".po"));
+	struct ProjectInfo
+	{
+		/** The name of this project */
+		public string ProjectName;
 
-        //Export
-        if (file != null)
-        {
-            ExportFileToDirectory(file, new DirectoryInfo(CmdEnv.LocalRoot + "/Engine/Content/Localization/" + ProjectName), project.EnabledCultures);
-        }
-    }
+		/** Path to this projects localization config file (relative to the root working directory for the commandlet) */
+		public string LocalizationConfigFile;
+
+		/** The destination path for the files containing the gathered data for this project (extracted from its config file - relative to the root working directory for the commandlet) */
+		public string DestinationPath;
+
+		/** The name to use for this projects manifest file (extracted from its config file) */
+		public string ManifestName;
+
+		/** The name to use for this projects archive file (extracted from its config file) */
+		public string ArchiveName;
+
+		/** The name to use for this projects portable object file (extracted from its config file) */
+		public string PortableObjectName;
+
+		/** The native culture for this project (extracted from its config file) */
+		public string NativeCulture;
+
+		/** The cultures to generate for this project (extracted from its config file) */
+		public List<string> CulturesToGenerate;
+	};
 
     private int OneSkyDownloadedPOChangeList = 0;
 
@@ -30,9 +45,32 @@ class Localise : BuildCommand
 	{
 		var EditorExe = CombinePaths(CmdEnv.LocalRoot, @"Engine/Binaries/Win64/UE4Editor-Cmd.exe");
 
-        OneSkyConfigData OneSkyConfig = OneSkyConfigHelper.Find("OneSkyConfig_EpicGames");
-        var oneSkyService = new OneSkyService(OneSkyConfig.ApiKey, OneSkyConfig.ApiSecret);
-        var projectGroup = GetProjectGroup(oneSkyService, "Unreal Engine");
+		// todo: All of this should be passed in via the command line to allow this script to be generic
+		var UEProjectDirectory = "Engine";
+		var OneSkyConfigName = "OneSkyConfig_EpicGames";
+		var OneSkyProjectGroupName = "Unreal Engine";
+		var OneSkyProjectNames = new string[] {
+			"Engine",
+			"Editor",
+			"EditorTutorials",
+			"PropertyNames",
+			"ToolTips",
+			"Category",
+			"Keywords",
+		};
+
+		var RootWorkingDirectory = CombinePaths(CmdEnv.LocalRoot, UEProjectDirectory);
+
+		// Generate the info we need to gather for each project
+		var ProjectInfos = new List<ProjectInfo>();
+		foreach (var ProjectName in OneSkyProjectNames)
+		{
+			ProjectInfos.Add(GenerateProjectInfo(RootWorkingDirectory, ProjectName));
+		}
+
+		OneSkyConfigData OneSkyConfig = OneSkyConfigHelper.Find(OneSkyConfigName);
+        var OneSkyService = new OneSkyService(OneSkyConfig.ApiKey, OneSkyConfig.ApiSecret);
+		var OneSkyProjectGroup = GetOneSkyProjectGroup(OneSkyService, OneSkyProjectGroupName);
 
         // Create changelist for backed up POs from OneSky.
         if (P4Enabled)
@@ -41,13 +79,10 @@ class Localise : BuildCommand
         }
 
         // Export all text from OneSky
-        ExportProjectToDirectory(oneSkyService, projectGroup, "Engine");
-        ExportProjectToDirectory(oneSkyService, projectGroup, "Editor");
-        ExportProjectToDirectory(oneSkyService, projectGroup, "EditorTutorials");
-        ExportProjectToDirectory(oneSkyService, projectGroup, "PropertyNames");
-        ExportProjectToDirectory(oneSkyService, projectGroup, "ToolTips");
-        ExportProjectToDirectory(oneSkyService, projectGroup, "Category");
-        ExportProjectToDirectory(oneSkyService, projectGroup, "Keywords");
+		foreach (var ProjectInfo in ProjectInfos)
+		{
+			ExportOneSkyProjectToDirectory(RootWorkingDirectory, OneSkyService, OneSkyProjectGroup, ProjectInfo);
+		}
 
         // Submit changelist for backed up POs from OneSky.
         if (P4Enabled)
@@ -69,24 +104,14 @@ class Localise : BuildCommand
 
 		// Setup commandlet arguments for SCC.
 		string CommandletSCCArguments = String.Empty;
-		if (P4Enabled) { CommandletSCCArguments += (string.IsNullOrEmpty(CommandletSCCArguments) ? "" : " ") + "-EnableSCC"; }
-		if (!AllowSubmit) { CommandletSCCArguments += (string.IsNullOrEmpty(CommandletSCCArguments) ? "" : " ") + "-DisableSCCSubmit"; }
+		if (P4Enabled) { CommandletSCCArguments += (String.IsNullOrEmpty(CommandletSCCArguments) ? "" : " ") + "-EnableSCC"; }
+		if (!AllowSubmit) { CommandletSCCArguments += (String.IsNullOrEmpty(CommandletSCCArguments) ? "" : " ") + "-DisableSCCSubmit"; }
 
-		// Setup commandlet arguments with configurations.
-		var CommandletArgumentSets = new string[] 
-			{
-				String.Format("-config={0}", @"./Config/Localization/Engine.ini") + (string.IsNullOrEmpty(CommandletSCCArguments) ? "" : " " + CommandletSCCArguments),
-				String.Format("-config={0}", @"./Config/Localization/Editor.ini") + (string.IsNullOrEmpty(CommandletSCCArguments) ? "" : " " + CommandletSCCArguments),
-				String.Format("-config={0}", @"./Config/Localization/EditorTutorials.ini") + (string.IsNullOrEmpty(CommandletSCCArguments) ? "" : " " + CommandletSCCArguments),
-				String.Format("-config={0}", @"./Config/Localization/PropertyNames.ini") + (string.IsNullOrEmpty(CommandletSCCArguments) ? "" : " " + CommandletSCCArguments),
-                String.Format("-config={0}", @"./Config/Localization/ToolTips.ini") + (string.IsNullOrEmpty(CommandletSCCArguments) ? "" : " " + CommandletSCCArguments),
-                String.Format("-config={0}", @"./Config/Localization/Category.ini") + (string.IsNullOrEmpty(CommandletSCCArguments) ? "" : " " + CommandletSCCArguments),
-				String.Format("-config={0}", @"./Config/Localization/Keywords.ini") + (string.IsNullOrEmpty(CommandletSCCArguments) ? "" : " " + CommandletSCCArguments),
-			};
-
-		// Execute commandlet for each set of arguments.
-		foreach (var CommandletArguments in CommandletArgumentSets)
+		// Execute commandlet for each project.
+		foreach (var ProjectInfo in ProjectInfos)
 		{
+			var CommandletArguments = String.Format("-config={0}", ProjectInfo.LocalizationConfigFile) + (String.IsNullOrEmpty(CommandletSCCArguments) ? "" : " " + CommandletSCCArguments);
+
 			Log("Localization for {0} {1}", EditorArguments, CommandletArguments);
 
 			Log("Running UE4Editor to generate localization data");
@@ -101,137 +126,197 @@ class Localise : BuildCommand
 		}
 
         // Upload all text to OneSky
-        UploadDirectoryToProject(GetProject(oneSkyService, "Unreal Engine", "Engine"), new DirectoryInfo(CmdEnv.LocalRoot + "/Engine/Content/Localization/Engine"), "*.po");
-        UploadDirectoryToProject(GetProject(oneSkyService, "Unreal Engine", "Editor"), new DirectoryInfo(CmdEnv.LocalRoot + "/Engine/Content/Localization/Editor"), "*.po");
-        UploadDirectoryToProject(GetProject(oneSkyService, "Unreal Engine", "EditorTutorials"), new DirectoryInfo(CmdEnv.LocalRoot + "/Engine/Content/Localization/EditorTutorials"), "*.po");
-        UploadDirectoryToProject(GetProject(oneSkyService, "Unreal Engine", "PropertyNames"), new DirectoryInfo(CmdEnv.LocalRoot + "/Engine/Content/Localization/PropertyNames"), "*.po");
-        UploadDirectoryToProject(GetProject(oneSkyService, "Unreal Engine", "ToolTips"), new DirectoryInfo(CmdEnv.LocalRoot + "/Engine/Content/Localization/ToolTips"), "*.po");
-        UploadDirectoryToProject(GetProject(oneSkyService, "Unreal Engine", "Category"), new DirectoryInfo(CmdEnv.LocalRoot + "/Engine/Content/Localization/Category"), "*.po");
-        UploadDirectoryToProject(GetProject(oneSkyService, "Unreal Engine", "Keywords"), new DirectoryInfo(CmdEnv.LocalRoot + "/Engine/Content/Localization/Keywords"), "*.po");
+		foreach (var ProjectInfo in ProjectInfos)
+		{
+			UploadProjectToOneSky(RootWorkingDirectory, OneSkyService, OneSkyProjectGroup, ProjectInfo);
+		}
+	}
+
+	private ProjectInfo GenerateProjectInfo(string RootWorkingDirectory, string ProjectName)
+	{
+		var ProjectInfo = new ProjectInfo();
+
+		ProjectInfo.ProjectName = ProjectName;
+		ProjectInfo.LocalizationConfigFile = String.Format(@"Config/Localization/{0}.ini", ProjectName);
+
+		var LocalizationConfig = new ConfigCacheIni(new FileReference(CombinePaths(RootWorkingDirectory, ProjectInfo.LocalizationConfigFile)));
+
+		if (!LocalizationConfig.GetString("CommonSettings", "DestinationPath", out ProjectInfo.DestinationPath))
+		{
+			throw new AutomationException("Failed to find a required config key! Section: 'CommonSettings', Key: 'DestinationPath', File: '{0}'", ProjectInfo.LocalizationConfigFile);
+		}
+
+		if (!LocalizationConfig.GetString("CommonSettings", "ManifestName", out ProjectInfo.ManifestName))
+		{
+			throw new AutomationException("Failed to find a required config key! Section: 'CommonSettings', Key: 'ManifestName', File: '{0}'", ProjectInfo.LocalizationConfigFile);
+		}
+
+		if (!LocalizationConfig.GetString("CommonSettings", "ArchiveName", out ProjectInfo.ArchiveName))
+		{
+			throw new AutomationException("Failed to find a required config key! Section: 'CommonSettings', Key: 'ArchiveName', File: '{0}'", ProjectInfo.LocalizationConfigFile);
+		}
+
+		if (!LocalizationConfig.GetString("CommonSettings", "PortableObjectName", out ProjectInfo.PortableObjectName))
+		{
+			throw new AutomationException("Failed to find a required config key! Section: 'CommonSettings', Key: 'PortableObjectName', File: '{0}'", ProjectInfo.LocalizationConfigFile);
+		}
+
+		if (!LocalizationConfig.GetString("CommonSettings", "NativeCulture", out ProjectInfo.NativeCulture))
+		{
+			throw new AutomationException("Failed to find a required config key! Section: 'CommonSettings', Key: 'NativeCulture', File: '{0}'", ProjectInfo.LocalizationConfigFile);
+		}
+
+		if (!LocalizationConfig.GetArray("CommonSettings", "CulturesToGenerate", out ProjectInfo.CulturesToGenerate))
+		{
+			throw new AutomationException("Failed to find a required config key! Section: 'CommonSettings', Key: 'CulturesToGenerate', File: '{0}'", ProjectInfo.LocalizationConfigFile);
+		}
+
+		return ProjectInfo;
 	}
 
-    private static ProjectGroup GetProjectGroup(OneSkyService oneSkyService, string GroupName)
+	private static ProjectGroup GetOneSkyProjectGroup(OneSkyService OneSkyService, string ProjectGroupName)
     {
-        var launcherGroup = oneSkyService.ProjectGroups.FirstOrDefault(g => g.Name == GroupName);
+		var OneSkyProjectGroup = OneSkyService.ProjectGroups.FirstOrDefault(g => g.Name == ProjectGroupName);
 
-        if (launcherGroup == null)
+		if (OneSkyProjectGroup == null)
         {
-            launcherGroup = new ProjectGroup(GroupName, "en");
-            oneSkyService.ProjectGroups.Add(launcherGroup);
+			OneSkyProjectGroup = new ProjectGroup(ProjectGroupName, "en");
+			OneSkyService.ProjectGroups.Add(OneSkyProjectGroup);
         }
 
-        return launcherGroup;
+		return OneSkyProjectGroup;
     }
 
-    private static OneSky.Project GetProject(OneSkyService oneSkyService, string GroupName, string ProjectName, string ProjectDescription = "")
+	private static OneSky.Project GetOneSkyProject(OneSkyService OneSkyService, string GroupName, string ProjectName, string ProjectDescription = "")
     {
-        var launcherGroup = GetProjectGroup(oneSkyService, GroupName);
+		var OneSkyProjectGroup = GetOneSkyProjectGroup(OneSkyService, GroupName);
 
-        OneSky.Project appProject = launcherGroup.Projects.FirstOrDefault(p => p.Name == ProjectName);
+		OneSky.Project OneSkyProject = OneSkyProjectGroup.Projects.FirstOrDefault(p => p.Name == ProjectName);
 
-        if (appProject == null)
+		if (OneSkyProject == null)
         {
-            ProjectType projectType = oneSkyService.ProjectTypes.First(pt => pt.Code == "website");
+			ProjectType projectType = OneSkyService.ProjectTypes.First(pt => pt.Code == "website");
 
-            appProject = new OneSky.Project(ProjectName, ProjectDescription, projectType);
-            launcherGroup.Projects.Add(appProject);
+			OneSkyProject = new OneSky.Project(ProjectName, ProjectDescription, projectType);
+			OneSkyProjectGroup.Projects.Add(OneSkyProject);
         }
 
-        return appProject;
+		return OneSkyProject;
     }
 
-    private void ExportFileToDirectory(UploadedFile file, DirectoryInfo destination, IEnumerable<string> cultures)
+	private void ExportOneSkyProjectToDirectory(string RootWorkingDirectory, OneSkyService OneSkyService, ProjectGroup OneSkyProjectGroup, ProjectInfo ProjectInfo)
+	{
+		var OneSkyProject = GetOneSkyProject(OneSkyService, OneSkyProjectGroup.Name, ProjectInfo.ProjectName);
+		var OneSkyFile = OneSkyProject.UploadedFiles.FirstOrDefault(f => f.Filename == ProjectInfo.PortableObjectName);
+
+		//Export
+		if (OneSkyFile != null)
+		{
+			var CulturesToExport = new List<string>();
+			foreach (var OneSkyCulture in OneSkyProject.EnabledCultures)
+			{
+				// Only export the OneSky cultures that we care about for this project
+				if (ProjectInfo.CulturesToGenerate.Contains(OneSkyCulture))
+				{
+					CulturesToExport.Add(OneSkyCulture);
+				}
+			}
+
+			ExportOneSkyFileToDirectory(OneSkyFile, new DirectoryInfo(CombinePaths(RootWorkingDirectory, ProjectInfo.DestinationPath)), CulturesToExport);
+		}
+	}
+
+	private void ExportOneSkyFileToDirectory(UploadedFile OneSkyFile, DirectoryInfo DestinationDirectory, IEnumerable<string> Cultures)
     {
-        foreach (var culture in cultures)
+		foreach (var Culture in Cultures)
         {
-            var cultureDirectory = new DirectoryInfo(Path.Combine(destination.FullName, culture));
-            if (!cultureDirectory.Exists)
+			var CultureDirectory = new DirectoryInfo(Path.Combine(DestinationDirectory.FullName, Culture));
+			if (!CultureDirectory.Exists)
             {
-                cultureDirectory.Create();
+				CultureDirectory.Create();
             }
 
-            using (var memoryStream = new MemoryStream())
+            using (var MemoryStream = new MemoryStream())
             {
-                var exportFile = new FileInfo(Path.Combine(cultureDirectory.FullName, file.Filename));
+				var ExportFile = new FileInfo(Path.Combine(CultureDirectory.FullName, OneSkyFile.Filename));
 
-                var exportTranslationState = file.ExportTranslation(culture, memoryStream).Result;
-                if (exportTranslationState == UploadedFile.ExportTranslationState.Success)
+				var ExportTranslationState = OneSkyFile.ExportTranslation(Culture, MemoryStream).Result;
+				if (ExportTranslationState == UploadedFile.ExportTranslationState.Success)
                 {
-                    memoryStream.Position = 0;
-                    using (Stream fileStream = File.OpenWrite(exportFile.FullName))
-                    {
-                        memoryStream.CopyTo(fileStream);
-                        Console.WriteLine("[SUCCESS] Exporting: " + exportFile.FullName + " Locale: " + culture);
-                    }
-                    FileInfo exportFileCopy = new FileInfo(Path.Combine(exportFile.DirectoryName, Path.GetFileNameWithoutExtension(exportFile.Name) + "_FromOneSky" + exportFile.Extension));
-                    // Add/check out backed up POs from OneSky.
-                    if (P4Enabled)
+					MemoryStream.Position = 0;
+					using (Stream FileStream = File.OpenWrite(ExportFile.FullName))
                     {
-                        UE4Build.AddBuildProductsToChangelist(OneSkyDownloadedPOChangeList, new List<string>() {exportFileCopy.FullName} );
+						MemoryStream.CopyTo(FileStream);
+						Console.WriteLine("[SUCCESS] Exporting: '{0}' ({1})", ExportFile.FullName, Culture);
                     }
-                    File.Copy(exportFile.FullName, exportFileCopy.FullName, true);
+
+					FileInfo ExportFileCopy = new FileInfo(Path.Combine(ExportFile.DirectoryName, String.Format("{0}_FromOneSky{1}", Path.GetFileNameWithoutExtension(ExportFile.Name), ExportFile.Extension)));
+					File.Copy(ExportFile.FullName, ExportFileCopy.FullName, true);
+
+					// Add/check out backed up POs from OneSky.
+					if (P4Enabled)
+					{
+						UE4Build.AddBuildProductsToChangelist(OneSkyDownloadedPOChangeList, new List<string>() { ExportFileCopy.FullName });
+					}
                 }
-                else if (exportTranslationState == UploadedFile.ExportTranslationState.NoContent)
+				else if (ExportTranslationState == UploadedFile.ExportTranslationState.NoContent)
                 {
-                    Console.WriteLine("[WARNING] Exporting: " + exportFile.FullName + " Locale: " + culture + " has no translations!");
+					Console.WriteLine("[WARNING] Exporting: '{0}' ({1}) has no translations!", ExportFile.FullName, Culture);
                 }
                 else
                 {
-                    Console.WriteLine("[FAILED] Exporting: " + exportFile.FullName + " Locale: " + culture);
+					Console.WriteLine("[FAILED] Exporting: '{0}' ({1})", ExportFile.FullName, Culture);
                 }
             }
         }
     }
 
-    static void UploadFilesToProjectForLocale(OneSky.Project project, string[] files, string localeName)
-    {
-        foreach (var currentFile in files)
-        {
-            using (var fileStream = File.OpenRead(currentFile))
-            {
-                // Read the BOM
-                var bom = new byte[3];
-                fileStream.Read(bom, 0, 3);
+	private void UploadProjectToOneSky(string RootWorkingDirectory, OneSkyService OneSkyService, ProjectGroup OneSkyProjectGroup, ProjectInfo ProjectInfo)
+	{
+		var OneSkyProject = GetOneSkyProject(OneSkyService, OneSkyProjectGroup.Name, ProjectInfo.ProjectName);
 
-                //We want to ignore the utf8 BOM
-                if (bom[0] != 0xef || bom[1] != 0xbb || bom[2] != 0xbf)
-                {
-                    fileStream.Position = 0;
-                }
+		// Upload the .po file for the native culture first
+		UploadFileToOneSky(OneSkyProject, new FileInfo(Path.Combine(RootWorkingDirectory, ProjectInfo.DestinationPath, ProjectInfo.NativeCulture, ProjectInfo.PortableObjectName)), ProjectInfo.NativeCulture);
 
-                Console.WriteLine("Uploading: " + currentFile + " Locale: " + localeName);
-                var uploadedFile = project.Upload(Path.GetFileName(currentFile), fileStream, localeName).Result;
-                
-                if (uploadedFile == null)
-                {
-                    Console.WriteLine("[FAILED] Uploading: " + currentFile + " Locale: " + localeName);
-                }
-                else
-                {
-                    Console.WriteLine("[SUCCESS] Uploading: " + currentFile + " Locale: " + localeName);                    
-                }
-            }
-        }
-    }
+		// Upload the remaining .po files for the other cultures
+		foreach (var Culture in ProjectInfo.CulturesToGenerate)
+		{
+			// Skip native culture as we uploaded it above
+			if (Culture != ProjectInfo.NativeCulture)
+			{
+				UploadFileToOneSky(OneSkyProject, new FileInfo(Path.Combine(RootWorkingDirectory, ProjectInfo.DestinationPath, Culture, ProjectInfo.PortableObjectName)), Culture);	
+			}
+		}
+	}
 
-    static void UploadDirectoryToProject(OneSky.Project project, DirectoryInfo directory, string fileExtension)
-    {
-        DirectoryInfo[] cultureDirectories = directory.GetDirectories("*", SearchOption.TopDirectoryOnly);
+	private void UploadFileToOneSky(OneSky.Project OneSkyProject, FileInfo FileToUpload, string Culture)
+	{
+		var FileNameToUpload = FileToUpload.FullName;
+		using (var FileStream = File.OpenRead(FileNameToUpload))
+		{
+			// Read the BOM
+			var UTF8BOM = new byte[3];
+			FileStream.Read(UTF8BOM, 0, 3);
 
-        // Upload, synchronously, the base culture POs.
-        DirectoryInfo nativeCultureDirectory = Array.Find(cultureDirectories, d => { return string.Equals(d.Name, project.BaseCultureName, StringComparison.InvariantCultureIgnoreCase); });
-        if (nativeCultureDirectory != null)
-        {
-            string[] files = Directory.GetFiles(nativeCultureDirectory.FullName, fileExtension, SearchOption.AllDirectories);
-            UploadFilesToProjectForLocale(project, files.Where((FilePath) => { return !FilePath.Contains("_FromOneSky"); }).ToArray(), nativeCultureDirectory.Name);
-        }
+			// We want to ignore the utf8 BOM
+			if (UTF8BOM[0] != 0xef || UTF8BOM[1] != 0xbb || UTF8BOM[2] != 0xbf)
+			{
+				FileStream.Position = 0;
+			}
 
-        // Upload all other POs asynchronously.
-        foreach (var foreignCultureDirectory in Array.FindAll(cultureDirectories, d => { return d != nativeCultureDirectory; }))
-        {
-            string[] files = Directory.GetFiles(foreignCultureDirectory.FullName, fileExtension, SearchOption.AllDirectories);
-            UploadFilesToProjectForLocale(project, files.Where((FilePath) => { return !FilePath.Contains("_FromOneSky");}).ToArray(), foreignCultureDirectory.Name);
-        }
-    }
-}
+			Console.WriteLine("Uploading: '{0}' ({1})", FileNameToUpload, Culture);
 
+			// todo: need to use the branch suffix here so that OneSky will merge them together
+			var UploadedFile = OneSkyProject.Upload(Path.GetFileName(FileNameToUpload), FileStream, Culture).Result;
+
+			if (UploadedFile == null)
+			{
+				Console.WriteLine("[FAILED] Uploading: '{0}' ({1})", FileNameToUpload, Culture);
+			}
+			else
+			{
+				Console.WriteLine("[SUCCESS] Uploading: '{0}' ({1})", FileNameToUpload, Culture);
+			}
+		}
+	}
+}
diff --git a/Engine/Source/Programs/UnrealBuildTool/Configuration/EngineConfiguration.cs b/Engine/Source/Programs/UnrealBuildTool/Configuration/EngineConfiguration.cs
--- a/Engine/Source/Programs/UnrealBuildTool/Configuration/EngineConfiguration.cs
+++ b/Engine/Source/Programs/UnrealBuildTool/Configuration/EngineConfiguration.cs
@@ -50,7 +50,7 @@ namespace UnrealBuildTool
 		static Dictionary<string, List<Command>> FileCache = new Dictionary<string,List<Command>>();
 		static Dictionary<string, ConfigCacheIni> IniCache = new Dictionary<string, ConfigCacheIni>();
 		static Dictionary<string, ConfigCacheIni> BaseIniCache = new Dictionary<string, ConfigCacheIni>();
-		static List<string> RequiredSections = new List<string>(){"AppxManifest", "/Script/AndroidRuntimeSettings.AndroidRuntimeSettings", "/Script/AndroidPlatformEditor.AndroidSDKSettings",
+		static List<string> RequiredSections = new List<string>(){"AppxManifest", "CommonSettings", "/Script/AndroidRuntimeSettings.AndroidRuntimeSettings", "/Script/AndroidPlatformEditor.AndroidSDKSettings",
 																	"/Script/BuildSettings.BuildSettings", "/Script/IOSRuntimeSettings.IOSRuntimeSettings", "/Script/WindowsTargetPlatform.WindowsTargetSettings",
 																	"/Script/UnrealEd.ProjectPackagingSettings", "/Script/PS4PlatformEditor.PS4TargetSettings", "/Script/XboxOneTargetPlatform.XboxOneTargetSettings",
 																	"/Script/HTML5PlatformEditor.HTML5TargetSettings"};
@@ -104,11 +104,25 @@ namespace UnrealBuildTool
 			}
 		}
 
+		/// <summary>
+		/// True if we are loading a hierarchy of config files that should be merged together
+		/// </summary>
+		bool bIsMergingConfigs;
+
 		/// <summary>
 		/// All sections parsed from ini file
 		/// </summary>
 		Dictionary<string, IniSection> Sections;
 
+		/// <summary>
+		/// Constructor. Parses a single ini file. No Platform settings, no engine hierarchy. Do not use this with ini files that have hierarchy!
+		/// </summary>
+		/// <param name="Filename">The ini file to load</param>
+		public ConfigCacheIni(FileReference Filename)
+		{
+			Init(Filename);
+		}
+
 		/// <summary>
 		/// Constructor. Parses ini hierarchy for the specified project.  No Platform settings.
 		/// </summary>
@@ -153,9 +167,22 @@ namespace UnrealBuildTool
 			Init(Platform, BaseIniName, ProjectDirectory, EngineDirectory, EngineOnly, BaseCache);
 		}
 
-		private void Init(UnrealTargetPlatform Platform, string BaseIniName, DirectoryReference ProjectDirectory, DirectoryReference EngineDirectory, bool EngineOnly = false, ConfigCacheIni BaseCache = null)
+		private void InitCommon()
 		{
 			Sections = new Dictionary<string, IniSection>(StringComparer.InvariantCultureIgnoreCase);
+		}
+
+		private void Init(FileReference IniFileName)
+		{
+			InitCommon();
+			bIsMergingConfigs = false;
+			ParseIniFile(IniFileName);
+		}
+
+		private void Init(UnrealTargetPlatform Platform, string BaseIniName, DirectoryReference ProjectDirectory, DirectoryReference EngineDirectory, bool EngineOnly = false, ConfigCacheIni BaseCache = null)
+		{
+			InitCommon();
+			bIsMergingConfigs = true;
 			if (EngineDirectory == null)
 			{
 				EngineDirectory = UnrealBuildTool.EngineDirectory;
@@ -464,7 +491,9 @@ namespace UnrealBuildTool
 			}
 			else
 			{
-				return ParseAction.New;
+				// We use Add rather than New when we're not merging config files together in order 
+				// to mimic the behavior of the C++ config cache when loading a single file
+				return (bIsMergingConfigs) ? ParseAction.New : ParseAction.Add;
 			}
 		}
 
