Renamed the BuildLocalization node to BuildEngineLocalization in order to better reflect its actual role

diff --git a/Engine/Source/Programs/AutomationTool/GUBP/LegacyBranchSetup.cs b/Engine/Source/Programs/AutomationTool/GUBP/LegacyBranchSetup.cs
--- a/Engine/Source/Programs/AutomationTool/GUBP/LegacyBranchSetup.cs
+++ b/Engine/Source/Programs/AutomationTool/GUBP/LegacyBranchSetup.cs
@@ -1121,7 +1121,7 @@ partial class GUBP
 
 		if (HostPlatforms.Contains(UnrealTargetPlatform.Win64) && BranchConfig.HasNode(RootEditorNode.StaticGetFullName(UnrealTargetPlatform.Win64)) && BranchOptions.bShouldBuildLocalization)
 		{
-			BranchConfig.AddNode(new BuildLocalization(BranchOptions.LocalizationBranchSuffix));
+			BranchConfig.AddNode(new BuildEngineLocalization(BranchOptions.LocalizationBranchSuffix));
 		}
 
         BranchConfig.AddNode(new WaitForTestShared(this));
diff --git a/Engine/Source/Programs/AutomationTool/GUBP/LegacyNodes.cs b/Engine/Source/Programs/AutomationTool/GUBP/LegacyNodes.cs
--- a/Engine/Source/Programs/AutomationTool/GUBP/LegacyNodes.cs
+++ b/Engine/Source/Programs/AutomationTool/GUBP/LegacyNodes.cs
@@ -2618,9 +2618,9 @@ partial class GUBP
         }
     }
 
-	public class BuildLocalization : HostPlatformNode
+	public class BuildEngineLocalization : HostPlatformNode
 	{
-		public BuildLocalization(string InLocalizationBranchSuffix)
+		public BuildEngineLocalization(string InLocalizationBranchSuffix)
 			: base(UnrealTargetPlatform.Win64)
 		{
 			//LocalizationBranchSuffix = InLocalizationBranchSuffix;
