Updated the BuildEngineLocalization node name to match the naming convention of other GUBP nodes

diff --git a/Engine/Source/Programs/AutomationTool/GUBP/LegacyBranchSetup.cs b/Engine/Source/Programs/AutomationTool/GUBP/LegacyBranchSetup.cs
--- a/Engine/Source/Programs/AutomationTool/GUBP/LegacyBranchSetup.cs
+++ b/Engine/Source/Programs/AutomationTool/GUBP/LegacyBranchSetup.cs
@@ -1121,7 +1121,7 @@ partial class GUBP
 
 		if (HostPlatforms.Contains(UnrealTargetPlatform.Win64) && BranchConfig.HasNode(RootEditorNode.StaticGetFullName(UnrealTargetPlatform.Win64)) && BranchOptions.bShouldBuildLocalization)
 		{
-			BranchConfig.AddNode(new BuildEngineLocalization(BranchOptions.LocalizationBranchSuffix));
+			BranchConfig.AddNode(new BuildEngineLocalizationNode(BranchOptions.LocalizationBranchSuffix));
 		}
 
         BranchConfig.AddNode(new WaitForTestShared(this));
diff --git a/Engine/Source/Programs/AutomationTool/GUBP/LegacyNodes.cs b/Engine/Source/Programs/AutomationTool/GUBP/LegacyNodes.cs
--- a/Engine/Source/Programs/AutomationTool/GUBP/LegacyNodes.cs
+++ b/Engine/Source/Programs/AutomationTool/GUBP/LegacyNodes.cs
@@ -2618,9 +2618,9 @@ partial class GUBP
         }
     }
 
-	public class BuildEngineLocalization : HostPlatformNode
+	public class BuildEngineLocalizationNode : HostPlatformNode
 	{
-		public BuildEngineLocalization(string InLocalizationBranchSuffix)
+		public BuildEngineLocalizationNode(string InLocalizationBranchSuffix)
 			: base(UnrealTargetPlatform.Win64)
 		{
 			//LocalizationBranchSuffix = InLocalizationBranchSuffix;
@@ -2646,7 +2646,7 @@ partial class GUBP
 
 		public static string StaticGetFullName()
 		{
-			return "BuildLocalization";
+			return "BuildEngineLocalization";
 		}
 
 		public override string GetFullName()
