Remove RootIfAnyForTempStorage(). Now that we zip everything before transferring to temp storage, we don't need to worry about the root path length. It's stored relative to the zip root anyway.

diff --git a/Engine/Source/Programs/AutomationTool/GUBP/GUBP.cs b/Engine/Source/Programs/AutomationTool/GUBP/GUBP.cs
--- a/Engine/Source/Programs/AutomationTool/GUBP/GUBP.cs
+++ b/Engine/Source/Programs/AutomationTool/GUBP/GUBP.cs
@@ -1273,22 +1273,11 @@ public partial class GUBP : BuildCommand
 
             var TempStorageNodeInfo = new TempStorageNodeInfo(JobInfo, NodeToDo.Name);
             
-			string StorageRootIfAny = NodeToDo.RootIfAnyForTempStorage;
-
-			if (bFake)
-			{
-				StorageRootIfAny = ""; // we don't rebase fake runs since those are entirely "records of success", which are always in the logs folder
-			}
-			if (string.IsNullOrEmpty(StorageRootIfAny))
-			{
-				StorageRootIfAny = CmdEnv.LocalRoot;
-			}
-
             if (NodeToDo.IsComplete)
             {
 				// Just fetch the build products from temp storage
 				Log("***** Retrieving GUBP Node {0} -> {1} : {2}", NodeToDo.Name, JobInfo.RootNameForTempStorage, TempStorageNodeInfo.GetRelativeDirectory());
-				NodeToDo.RetrieveBuildProducts(StorageRootIfAny, TempStorageNodeInfo);
+				NodeToDo.RetrieveBuildProducts(TempStorageNodeInfo);
             }
             else
             {
@@ -1310,7 +1299,7 @@ public partial class GUBP : BuildCommand
 				// Archive the build products if the node succeeded
 				if(bResult)
 				{
-					NodeToDo.ArchiveBuildProducts(StorageRootIfAny, TempStorageNodeInfo, !bSaveSharedTempStorage);
+					NodeToDo.ArchiveBuildProducts(TempStorageNodeInfo, !bSaveSharedTempStorage);
 				}
 
 				// Record that the node has finished
diff --git a/Engine/Source/Programs/AutomationTool/GUBP/LegacyNodes.cs b/Engine/Source/Programs/AutomationTool/GUBP/LegacyNodes.cs
--- a/Engine/Source/Programs/AutomationTool/GUBP/LegacyNodes.cs
+++ b/Engine/Source/Programs/AutomationTool/GUBP/LegacyNodes.cs
@@ -34,10 +34,6 @@ partial class GUBP
         /// </summary>
         /// <returns></returns>
         public virtual string GameNameIfAnyForFullGameAggregateNode()
-        {
-            return "";
-        }
-        public virtual string RootIfAnyForTempStorage()
         {
             return "";
         }
@@ -1679,11 +1675,6 @@ partial class GUBP
         {
             return bIsMassive ? 240 : base.TimeoutInMinutes();
         }
-
-        public override string RootIfAnyForTempStorage()
-        {
-            return CombinePaths(Path.GetDirectoryName(GameProj.FilePath), "Saved", "Cooked", CookPlatform);
-        }
         public override void DoBuild(GUBP bp)
         {
             if (HostPlatform == UnrealTargetPlatform.Mac)
@@ -1694,7 +1685,7 @@ partial class GUBP
 			{
 				CommandUtils.CookCommandlet(GameProj.FilePath, "UE4Editor-Cmd.exe", null, null, null, null, CookPlatform);
 			}
-            var CookedPath = RootIfAnyForTempStorage();
+            var CookedPath = CombinePaths(Path.GetDirectoryName(GameProj.FilePath), "Saved", "Cooked", CookPlatform);
             var CookedFiles = CommandUtils.FindFiles("*", true, CookedPath);
             if (CookedFiles.GetLength(0) < 1)
             {
diff --git a/Engine/Source/Programs/AutomationTool/GUBP/Nodes/BuildNode.cs b/Engine/Source/Programs/AutomationTool/GUBP/Nodes/BuildNode.cs
--- a/Engine/Source/Programs/AutomationTool/GUBP/Nodes/BuildNode.cs
+++ b/Engine/Source/Programs/AutomationTool/GUBP/Nodes/BuildNode.cs
@@ -28,7 +28,6 @@ namespace AutomationTool
 		public bool IsTest;
 		public string DisplayGroupName;
         public string GameNameIfAnyForFullGameAggregateNode;
-		public string RootIfAnyForTempStorage;
 
 		public abstract BuildNode Instantiate();
 	}
@@ -57,7 +56,6 @@ namespace AutomationTool
 		public bool IsTest;
 		public string DisplayGroupName;
         public string GameNameIfAnyForFullGameAggregateNode;
-		public string RootIfAnyForTempStorage;
 
 		public List<string> BuildProducts;
 
@@ -79,20 +77,19 @@ namespace AutomationTool
 			IsTest = Template.IsTest;
 			DisplayGroupName = Template.DisplayGroupName;
             GameNameIfAnyForFullGameAggregateNode = Template.GameNameIfAnyForFullGameAggregateNode;
-			RootIfAnyForTempStorage = Template.RootIfAnyForTempStorage;
 		}
 
-		public virtual void ArchiveBuildProducts(string StorageRootIfAny, TempStorageNodeInfo TempStorageNodeInfo, bool bLocalOnly)
+		public virtual void ArchiveBuildProducts(TempStorageNodeInfo TempStorageNodeInfo, bool bLocalOnly)
 		{
-			TempStorage.StoreToTempStorage(TempStorageNodeInfo, BuildProducts, bLocalOnly, StorageRootIfAny);
+			TempStorage.StoreToTempStorage(TempStorageNodeInfo, BuildProducts, bLocalOnly, CommandUtils.CmdEnv.LocalRoot);
 		}
 
-		public virtual void RetrieveBuildProducts(string StorageRootIfAny, TempStorageNodeInfo TempStorageNodeInfo)
+		public virtual void RetrieveBuildProducts(TempStorageNodeInfo TempStorageNodeInfo)
 		{
 			CommandUtils.Log("***** Retrieving GUBP Node {0} from {1}", Name, TempStorageNodeInfo.GetRelativeDirectory());
 			try
 			{
-				BuildProducts = TempStorage.RetrieveFromTempStorage(TempStorageNodeInfo, StorageRootIfAny);
+				BuildProducts = TempStorage.RetrieveFromTempStorage(TempStorageNodeInfo, CommandUtils.CmdEnv.LocalRoot);
 			}
 			catch (Exception Ex)
 			{
diff --git a/Engine/Source/Programs/AutomationTool/GUBP/Nodes/LegacyNode.cs b/Engine/Source/Programs/AutomationTool/GUBP/Nodes/LegacyNode.cs
--- a/Engine/Source/Programs/AutomationTool/GUBP/Nodes/LegacyNode.cs
+++ b/Engine/Source/Programs/AutomationTool/GUBP/Nodes/LegacyNode.cs
@@ -35,7 +35,6 @@ namespace AutomationTool
 			Definition.IsTest = Node.IsTest();
 			Definition.DisplayGroupName = Node.GetDisplayGroupName();
             Definition.GameNameIfAnyForFullGameAggregateNode = Node.GameNameIfAnyForFullGameAggregateNode();
-			Definition.RootIfAnyForTempStorage = Node.RootIfAnyForTempStorage();
 		}
 
 		public override BuildNode Instantiate()
@@ -56,9 +55,9 @@ namespace AutomationTool
 			Node = Template.Node;
 		}
 
-		public override void RetrieveBuildProducts(string StorageRootIfAny, TempStorageNodeInfo TempStorageNodeInfo)
+		public override void RetrieveBuildProducts(TempStorageNodeInfo TempStorageNodeInfo)
 		{
-			base.RetrieveBuildProducts(StorageRootIfAny, TempStorageNodeInfo);
+			base.RetrieveBuildProducts(TempStorageNodeInfo);
 
 			Node.BuildProducts = BuildProducts;
 		}
