Instance deployment handlers, rather than having a single shared instance. Required to support multiple build targets in UBT. Creation is through a method on the UEBuildPlatform.

diff --git a/Engine/Source/Programs/AutomationTool/Android/AndroidPlatform.Automation.cs b/Engine/Source/Programs/AutomationTool/Android/AndroidPlatform.Automation.cs
--- a/Engine/Source/Programs/AutomationTool/Android/AndroidPlatform.Automation.cs
+++ b/Engine/Source/Programs/AutomationTool/Android/AndroidPlatform.Automation.cs
@@ -114,10 +114,10 @@ public class AndroidPlatform : Platform
 	{
 		var Architectures = UnrealBuildTool.AndroidToolChain.GetAllArchitectures();
 		var GPUArchitectures = UnrealBuildTool.AndroidToolChain.GetAllGPUArchitectures();
-		bool bMakeSeparateApks = UnrealBuildTool.Android.UEDeployAndroid.ShouldMakeSeparateApks();
-		bool bPackageDataInsideApk = UnrealBuildTool.Android.UEDeployAndroid.PackageDataInsideApk(false);
+		bool bMakeSeparateApks = UnrealBuildTool.UEDeployAndroid.ShouldMakeSeparateApks();
 
-		var Deploy = UEBuildDeploy.GetBuildDeploy(UnrealTargetPlatform.Android);
+		var Deploy = new UEDeployAndroid();
+		bool bPackageDataInsideApk = Deploy.PackageDataInsideApk(false);
 
 		string BaseApkName = GetFinalApkName(Params, SC.StageExecutables[0], true, "", "");
 		Log("BaseApkName = {0}", BaseApkName);
@@ -319,8 +319,8 @@ public class AndroidPlatform : Platform
 
 		var Architectures = UnrealBuildTool.AndroidToolChain.GetAllArchitectures();
 		var GPUArchitectures = UnrealBuildTool.AndroidToolChain.GetAllGPUArchitectures();
-		bool bMakeSeparateApks = UnrealBuildTool.Android.UEDeployAndroid.ShouldMakeSeparateApks();
-		bool bPackageDataInsideApk = UnrealBuildTool.Android.UEDeployAndroid.PackageDataInsideApk(false);
+		bool bMakeSeparateApks = UnrealBuildTool.UEDeployAndroid.ShouldMakeSeparateApks();
+		bool bPackageDataInsideApk = new UnrealBuildTool.UEDeployAndroid().PackageDataInsideApk(false);
 
 		bool bAddedOBB = false;
 		foreach (string Architecture in Architectures)
@@ -510,7 +510,7 @@ public class AndroidPlatform : Platform
 		string ApkName = GetFinalApkName(Params, SC.StageExecutables[0], true, DeviceArchitecture, GPUArchitecture);
 
 		// make sure APK is up to date (this is fast if so)
-		var Deploy = UEBuildDeploy.GetBuildDeploy(UnrealTargetPlatform.Android);
+		var Deploy = new UEDeployAndroid();
 		if (!Params.Prebuilt)
 		{
 			string CookFlavor = SC.FinalCookPlatform.IndexOf("_") > 0 ? SC.FinalCookPlatform.Substring(SC.FinalCookPlatform.IndexOf("_")) : "";
@@ -969,7 +969,7 @@ public class AndroidPlatform : Platform
 
 	private string GetBestDeviceArchitecture(ProjectParams Params)
 	{
-		bool bMakeSeparateApks = UnrealBuildTool.Android.UEDeployAndroid.ShouldMakeSeparateApks();
+		bool bMakeSeparateApks = UnrealBuildTool.UEDeployAndroid.ShouldMakeSeparateApks();
 		// if we are joining all .so's into a single .apk, there's no need to find the best one - there is no other one
 		if (!bMakeSeparateApks)
 		{
@@ -982,7 +982,7 @@ public class AndroidPlatform : Platform
 		ProcessResult ABIResult = RunAdbCommand(Params, " shell getprop ro.product.cpu.abi", null, ERunOptions.AppMustExist);
 
 		// the output is just the architecture
-		string DeviceArch = UnrealBuildTool.Android.UEDeployAndroid.GetUE4Arch(ABIResult.Output.Trim());
+		string DeviceArch = UnrealBuildTool.UEDeployAndroid.GetUE4Arch(ABIResult.Output.Trim());
 
 		// if the architecture wasn't built, look for a backup
 		if (!AppArchitectures.Contains(DeviceArch))
@@ -1034,7 +1034,7 @@ public class AndroidPlatform : Platform
 
 	private string GetBestGPUArchitecture(ProjectParams Params)
 	{
-		bool bMakeSeparateApks = UnrealBuildTool.Android.UEDeployAndroid.ShouldMakeSeparateApks();
+		bool bMakeSeparateApks = UnrealBuildTool.UEDeployAndroid.ShouldMakeSeparateApks();
 		// if we are joining all .so's into a single .apk, there's no need to find the best one - there is no other one
 		if (!bMakeSeparateApks)
 		{
diff --git a/Engine/Source/Programs/AutomationTool/IOS/IOSPlatform.Automation.cs b/Engine/Source/Programs/AutomationTool/IOS/IOSPlatform.Automation.cs
--- a/Engine/Source/Programs/AutomationTool/IOS/IOSPlatform.Automation.cs
+++ b/Engine/Source/Programs/AutomationTool/IOS/IOSPlatform.Automation.cs
@@ -146,7 +146,7 @@ public class IOSPlatform : Platform
 		if (UnrealBuildTool.BuildHostPlatform.Current.Platform == UnrealTargetPlatform.Mac)
 		{
 			// copy in all of the artwork and plist
-			var DeployHandler = UEBuildDeploy.GetBuildDeploy(UnrealTargetPlatform.IOS);
+			var DeployHandler = new UEDeployIOS();
 
 			DeployHandler.PrepForUATPackageOrDeploy(Params.ShortProjectName,
 				Path.GetDirectoryName(Params.RawProjectPath),
@@ -610,7 +610,7 @@ public class IOSPlatform : Platform
 //				if (!File.Exists(TargetPListFile))
 				{
 					// ensure the plist, entitlements, and provision files are properly copied
-					UnrealBuildTool.IOS.UEDeployIOS.GeneratePList((SC.IsCodeBasedProject ? SC.ProjectRoot : SC.LocalRoot + "/Engine"), !SC.IsCodeBasedProject, (SC.IsCodeBasedProject ? SC.ShortProjectName : "UE4Game"), SC.ShortProjectName, SC.LocalRoot + "/Engine", (SC.IsCodeBasedProject ? SC.ProjectRoot : SC.LocalRoot + "/Engine") + "/Binaries/IOS/Payload/" + (SC.IsCodeBasedProject ? SC.ShortProjectName : "UE4Game") + ".app");
+					UnrealBuildTool.UEDeployIOS.GeneratePList((SC.IsCodeBasedProject ? SC.ProjectRoot : SC.LocalRoot + "/Engine"), !SC.IsCodeBasedProject, (SC.IsCodeBasedProject ? SC.ShortProjectName : "UE4Game"), SC.ShortProjectName, SC.LocalRoot + "/Engine", (SC.IsCodeBasedProject ? SC.ProjectRoot : SC.LocalRoot + "/Engine") + "/Binaries/IOS/Payload/" + (SC.IsCodeBasedProject ? SC.ShortProjectName : "UE4Game") + ".app");
 				}
 
 				SC.StageFiles(StagedFileType.NonUFS, SourcePath, Path.GetFileName(TargetPListFile), false, null, "", false, false, "Info.plist");
diff --git a/Engine/Source/Programs/UnrealBuildTool/Android/UEBuildAndroid.cs b/Engine/Source/Programs/UnrealBuildTool/Android/UEBuildAndroid.cs
--- a/Engine/Source/Programs/UnrealBuildTool/Android/UEBuildAndroid.cs
+++ b/Engine/Source/Programs/UnrealBuildTool/Android/UEBuildAndroid.cs
@@ -476,5 +476,11 @@ namespace UnrealBuildTool
 
 			return AllBinaries;
 		}
+
+		public override bool TryCreateDeploymentHandler(FileReference ProjectFile, out UEBuildDeploy DeploymentHandler)
+		{
+			DeploymentHandler = new UEDeployAndroid();
+			return true;
+		}
 	}
 }
diff --git a/Engine/Source/Programs/UnrealBuildTool/Android/UEDeployAndroid.cs b/Engine/Source/Programs/UnrealBuildTool/Android/UEDeployAndroid.cs
--- a/Engine/Source/Programs/UnrealBuildTool/Android/UEDeployAndroid.cs
+++ b/Engine/Source/Programs/UnrealBuildTool/Android/UEDeployAndroid.cs
@@ -9,37 +9,19 @@ using System.Diagnostics;
 using System.IO;
 using Microsoft.Win32;
 
-namespace UnrealBuildTool.Android
+namespace UnrealBuildTool
 {
 	public class UEDeployAndroid : UEBuildDeploy
 	{
-		/// <summary>
-		/// Register the platform with the UEBuildDeploy class
-		/// </summary>
-		public override void RegisterBuildDeploy()
-		{
-			string NDKPath = Environment.GetEnvironmentVariable("ANDROID_HOME");
-
-			// we don't have an NDKROOT specified
-			if (String.IsNullOrEmpty(NDKPath))
-			{
-				Log.TraceVerbose("        Unable to register Android deployment class because the ANDROID_HOME environment variable isn't set or points to something that doesn't exist");
-			}
-			else
-			{
-				UEBuildDeploy.RegisterBuildDeploy(UnrealTargetPlatform.Android, this);
-			}
-		}
-
 		/// <summary>
 		/// Internal usage for GetApiLevel
 		/// </summary>
-		private static List<string> PossibleApiLevels = null;
+		private List<string> PossibleApiLevels = null;
 
 		/// <summary>
 		/// Simple function to pipe output asynchronously
 		/// </summary>
-		private static void ParseApiLevel(object Sender, DataReceivedEventArgs Event)
+		private void ParseApiLevel(object Sender, DataReceivedEventArgs Event)
 		{
 			// DataReceivedEventHandler is fired with a null string when the output stream is closed.  We don't want to
 			// print anything for that event.
@@ -58,9 +40,9 @@ namespace UnrealBuildTool.Android
 			}
 		}
 
-		static private Dictionary<string, ConfigCacheIni> ConfigCache = null;
+		private Dictionary<string, ConfigCacheIni> ConfigCache = null;
 
-		private static ConfigCacheIni GetConfigCacheIni(string baseIniName)
+		private ConfigCacheIni GetConfigCacheIni(string baseIniName)
 		{
 			if (ConfigCache == null)
 			{
@@ -77,8 +59,8 @@ namespace UnrealBuildTool.Android
 			return config;
 		}
 
-		static private string CachedSDKLevel = null;
-		private static string GetSdkApiLevel()
+		private string CachedSDKLevel = null;
+		private string GetSdkApiLevel()
 		{
 			if (CachedSDKLevel == null)
 			{
@@ -133,7 +115,7 @@ namespace UnrealBuildTool.Android
 			return VersionString;
 		}
 
-		public static bool PackageDataInsideApk(bool bDisallowPackagingDataInApk, ConfigCacheIni Ini = null)
+		public bool PackageDataInsideApk(bool bDisallowPackagingDataInApk, ConfigCacheIni Ini = null)
 		{
 			if (bDisallowPackagingDataInApk)
 			{
@@ -153,7 +135,7 @@ namespace UnrealBuildTool.Android
 			return bPackageDataInsideApk;
 		}
 
-		public static bool DisableVerifyOBBOnStartUp(ConfigCacheIni Ini = null)
+		public bool DisableVerifyOBBOnStartUp(ConfigCacheIni Ini = null)
 		{
 			// make a new one if one wasn't passed in
 			if (Ini == null)
@@ -1835,7 +1817,7 @@ namespace UnrealBuildTool.Android
 			}
 		}
 
-		private static void CopyPluginLibs(string EngineDirectory, string UE4BuildPath, string UE4Arch)
+		private void CopyPluginLibs(string EngineDirectory, string UE4BuildPath, string UE4Arch)
 		{
 			ConfigCacheIni Ini = GetConfigCacheIni("Engine");
 			string Arch = GetNDKArch(UE4Arch);
diff --git a/Engine/Source/Programs/UnrealBuildTool/Configuration/UEBuildDeploy.cs b/Engine/Source/Programs/UnrealBuildTool/Configuration/UEBuildDeploy.cs
--- a/Engine/Source/Programs/UnrealBuildTool/Configuration/UEBuildDeploy.cs
+++ b/Engine/Source/Programs/UnrealBuildTool/Configuration/UEBuildDeploy.cs
@@ -8,72 +8,11 @@ using System.Diagnostics;
 
 namespace UnrealBuildTool
 {
-	public interface IUEBuildDeploy
-	{
-		/// <summary>
-		/// Register the platform with the UEBuildDeploy class.
-		/// </summary>
-		void RegisterBuildDeploy();
-
-		/// <summary>
-		/// Prepare the target for deployment.
-		/// </summary>
-		/// <param name="InTarget">The target for deployment.</param>
-		/// <returns>bool True if successful, false if not.</returns>
-		bool PrepTargetForDeployment(UEBuildTarget InTarget);
-
-		/// <summary>
-		/// Prepare the target for deployment.
-		/// </summary>
-		bool PrepForUATPackageOrDeploy(string ProjectName, string ProjectDirectory, string ExecutablePath, string EngineDirectory, bool bForDistribution, string CookFlavor, bool bIsDataDeploy);
-	}
-
 	/// <summary>
 	///  Base class to handle deploy of a target for a given platform
 	/// </summary>
-	public abstract class UEBuildDeploy : IUEBuildDeploy
+	public abstract class UEBuildDeploy
 	{
-		static Dictionary<UnrealTargetPlatform, IUEBuildDeploy> BuildDeployDictionary = new Dictionary<UnrealTargetPlatform, IUEBuildDeploy>();
-
-		/// <summary>
-		/// Register the given platforms UEBuildDeploy instance
-		/// </summary>
-		/// <param name="InPlatform">  The UnrealTargetPlatform to register with</param>
-		/// <param name="InBuildDeploy"> The UEBuildDeploy instance to use for the InPlatform</param>
-		public static void RegisterBuildDeploy(UnrealTargetPlatform InPlatform, IUEBuildDeploy InBuildDeploy)
-		{
-			if (BuildDeployDictionary.ContainsKey(InPlatform) == true)
-			{
-				Log.TraceWarning("RegisterBuildDeply Warning: Registering build deploy {0} for {1} when it is already set to {2}",
-					InBuildDeploy.ToString(), InPlatform.ToString(), BuildDeployDictionary[InPlatform].ToString());
-				BuildDeployDictionary[InPlatform] = InBuildDeploy;
-			}
-			else
-			{
-				BuildDeployDictionary.Add(InPlatform, InBuildDeploy);
-			}
-		}
-
-		/// <summary>
-		/// Retrieve the UEBuildDeploy instance for the given TargetPlatform
-		/// </summary>
-		/// <param name="InPlatform">  The UnrealTargetPlatform being built</param>
-		/// <returns>UEBuildDeploy  The instance of the build deploy</returns>
-		public static IUEBuildDeploy GetBuildDeploy(UnrealTargetPlatform InPlatform)
-		{
-			if (BuildDeployDictionary.ContainsKey(InPlatform) == true)
-			{
-				return BuildDeployDictionary[InPlatform];
-			}
-			// A platform does not *have* to have a deployment handler...
-			return null;
-		}
-
-		/// <summary>
-		/// Register the platform with the UEBuildDeploy class
-		/// </summary>
-		public abstract void RegisterBuildDeploy();
-
 		/// <summary>
 		/// Prepare the target for deployment
 		/// </summary>
diff --git a/Engine/Source/Programs/UnrealBuildTool/Configuration/UEBuildPlatform.cs b/Engine/Source/Programs/UnrealBuildTool/Configuration/UEBuildPlatform.cs
--- a/Engine/Source/Programs/UnrealBuildTool/Configuration/UEBuildPlatform.cs
+++ b/Engine/Source/Programs/UnrealBuildTool/Configuration/UEBuildPlatform.cs
@@ -246,6 +246,14 @@ namespace UnrealBuildTool
 		/// return true if the project uses the default build config
 		/// </summary>
 		bool HasDefaultBuildConfig(UnrealTargetPlatform InPlatform, DirectoryReference InProjectPath);
+
+		/// <summary>
+		/// Create a build deployment handler
+		/// </summary>
+		/// <param name="ProjectFile">The project file of the target being deployed. Used to find any deployment specific settings.</param>
+		/// <param name="DeploymentHandler">The output deployment handler</param>
+		/// <returns>True if the platform requires a deployment handler, false otherwise</returns>
+		bool TryCreateDeploymentHandler(FileReference ProjectFile, out UEBuildDeploy DeploymentHandler);
 	}
 
 	public abstract partial class UEBuildPlatform : IUEBuildPlatform
@@ -944,6 +952,14 @@ namespace UnrealBuildTool
 			return DoProjectSettingsMatchDefault(Platform, ProjectDirectoryName, "/Script/BuildSettings.BuildSettings",
 				BoolKeys, null, null);
 		}
+
+		/// <summary>
+		/// Create a build deployment handler
+		/// </summary>
+		/// <param name="ProjectFile">The project file of the target being deployed. Used to find any deployment specific settings.</param>
+		/// <param name="DeploymentHandler">The output deployment handler</param>
+		/// <returns>True if the platform requires a deployment handler, false otherwise</returns>
+		public abstract bool TryCreateDeploymentHandler(FileReference ProjectFile, out UEBuildDeploy DeploymentHandler);
 	}
 
 	// AutoSDKs handling portion
diff --git a/Engine/Source/Programs/UnrealBuildTool/HTML5/UEBuildHTML5.cs b/Engine/Source/Programs/UnrealBuildTool/HTML5/UEBuildHTML5.cs
--- a/Engine/Source/Programs/UnrealBuildTool/HTML5/UEBuildHTML5.cs
+++ b/Engine/Source/Programs/UnrealBuildTool/HTML5/UEBuildHTML5.cs
@@ -368,5 +368,17 @@ namespace UnrealBuildTool
 		{
 			// 			InBuildTarget.ExtraModuleNames.Add("ES2RHI");
 		}
+
+		/// <summary>
+		/// Create a build deployment handler
+		/// </summary>
+		/// <param name="ProjectFile">The project file of the target being deployed. Used to find any deployment specific settings.</param>
+		/// <param name="DeploymentHandler">The output deployment handler</param>
+		/// <returns>True if the platform requires a deployment handler, false otherwise</returns>
+		public override bool TryCreateDeploymentHandler(FileReference ProjectFile, out UEBuildDeploy DeploymentHandler)
+		{
+			DeploymentHandler = null;
+			return false;
+		}
 	}
 }
diff --git a/Engine/Source/Programs/UnrealBuildTool/IOS/IOSToolChain.cs b/Engine/Source/Programs/UnrealBuildTool/IOS/IOSToolChain.cs
--- a/Engine/Source/Programs/UnrealBuildTool/IOS/IOSToolChain.cs
+++ b/Engine/Source/Programs/UnrealBuildTool/IOS/IOSToolChain.cs
@@ -1013,11 +1013,8 @@ namespace UnrealBuildTool
 					if (Directory.Exists(Project))
 					{
 						// ensure the plist, entitlements, and provision files are properly copied
-						var DeployHandler = UEBuildDeploy.GetBuildDeploy(Target.Platform);
-						if (DeployHandler != null)
-						{
-							DeployHandler.PrepTargetForDeployment(Target);
-						}
+						var DeployHandler = new UEDeployIOS();
+						DeployHandler.PrepTargetForDeployment(Target);
 
 						// code sign the project
 						string CmdLine = XcodeDeveloperDir + "usr/bin/xcodebuild" +
@@ -1141,11 +1138,8 @@ namespace UnrealBuildTool
 				if (BuildConfiguration.bCreateStubIPA || bUseDangerouslyFastMode)
 				{
 					// ensure the plist, entitlements, and provision files are properly copied
-					var DeployHandler = UEBuildDeploy.GetBuildDeploy(Target.Platform);
-					if (DeployHandler != null)
-					{
-						DeployHandler.PrepTargetForDeployment(Target);
-					}
+					var DeployHandler = new UEDeployIOS();
+					DeployHandler.PrepTargetForDeployment(Target);
 
 					if (!bUseDangerouslyFastMode)
 					{
diff --git a/Engine/Source/Programs/UnrealBuildTool/IOS/UEBuildIOS.cs b/Engine/Source/Programs/UnrealBuildTool/IOS/UEBuildIOS.cs
--- a/Engine/Source/Programs/UnrealBuildTool/IOS/UEBuildIOS.cs
+++ b/Engine/Source/Programs/UnrealBuildTool/IOS/UEBuildIOS.cs
@@ -462,6 +462,18 @@ namespace UnrealBuildTool
 				}
 			}
 		}
+
+		/// <summary>
+		/// Create a build deployment handler
+		/// </summary>
+		/// <param name="ProjectFile">The project file of the target being deployed. Used to find any deployment specific settings.</param>
+		/// <param name="DeploymentHandler">The output deployment handler</param>
+		/// <returns>True if the platform requires a deployment handler, false otherwise</returns>
+		public override bool TryCreateDeploymentHandler(FileReference ProjectFile, out UEBuildDeploy DeploymentHandler)
+		{
+			DeploymentHandler = new UEDeployIOS();
+			return true;
+		}
 	}
 }
 
diff --git a/Engine/Source/Programs/UnrealBuildTool/IOS/UEDeployIOS.cs b/Engine/Source/Programs/UnrealBuildTool/IOS/UEDeployIOS.cs
--- a/Engine/Source/Programs/UnrealBuildTool/IOS/UEDeployIOS.cs
+++ b/Engine/Source/Programs/UnrealBuildTool/IOS/UEDeployIOS.cs
@@ -8,19 +8,10 @@ using System.Xml;
 using System.IO;
 using System.Diagnostics;
 
-namespace UnrealBuildTool.IOS
+namespace UnrealBuildTool
 {
 	public class UEDeployIOS : UEBuildDeploy
 	{
-		/// <summary>
-		/// Register the platform with the UEBuildDeploy class
-		/// </summary>
-		public override void RegisterBuildDeploy()
-		{
-			// TODO: print debug info and handle any cases that would keep this from registering
-			UEBuildDeploy.RegisterBuildDeploy(UnrealTargetPlatform.IOS, this);
-		}
-
 		class VersionUtilities
 		{
 			public static string BuildDirectory
diff --git a/Engine/Source/Programs/UnrealBuildTool/Linux/UEBuildLinux.cs b/Engine/Source/Programs/UnrealBuildTool/Linux/UEBuildLinux.cs
--- a/Engine/Source/Programs/UnrealBuildTool/Linux/UEBuildLinux.cs
+++ b/Engine/Source/Programs/UnrealBuildTool/Linux/UEBuildLinux.cs
@@ -393,5 +393,17 @@ namespace UnrealBuildTool
 		public override void SetupBinaries(UEBuildTarget InBuildTarget)
 		{
 		}
+
+		/// <summary>
+		/// Create a build deployment handler
+		/// </summary>
+		/// <param name="ProjectFile">The project file of the target being deployed. Used to find any deployment specific settings.</param>
+		/// <param name="DeploymentHandler">The output deployment handler</param>
+		/// <returns>True if the platform requires a deployment handler, false otherwise</returns>
+		public override bool TryCreateDeploymentHandler(FileReference ProjectFile, out UEBuildDeploy DeploymentHandler)
+		{
+			DeploymentHandler = null;
+			return false;
+		}
 	}
 }
diff --git a/Engine/Source/Programs/UnrealBuildTool/Mac/UEBuildMac.cs b/Engine/Source/Programs/UnrealBuildTool/Mac/UEBuildMac.cs
--- a/Engine/Source/Programs/UnrealBuildTool/Mac/UEBuildMac.cs
+++ b/Engine/Source/Programs/UnrealBuildTool/Mac/UEBuildMac.cs
@@ -218,5 +218,17 @@ namespace UnrealBuildTool
 		{
 			return false;
 		}
+
+		/// <summary>
+		/// Create a build deployment handler
+		/// </summary>
+		/// <param name="ProjectFile">The project file of the target being deployed. Used to find any deployment specific settings.</param>
+		/// <param name="DeploymentHandler">The output deployment handler</param>
+		/// <returns>True if the platform requires a deployment handler, false otherwise</returns>
+		public override bool TryCreateDeploymentHandler(FileReference ProjectFile, out UEBuildDeploy DeploymentHandler)
+		{
+			DeploymentHandler = new UEDeployMac();
+			return true;
+		}
 	}
 }
diff --git a/Engine/Source/Programs/UnrealBuildTool/Mac/UEDeployMac.cs b/Engine/Source/Programs/UnrealBuildTool/Mac/UEDeployMac.cs
--- a/Engine/Source/Programs/UnrealBuildTool/Mac/UEDeployMac.cs
+++ b/Engine/Source/Programs/UnrealBuildTool/Mac/UEDeployMac.cs
@@ -9,19 +9,10 @@ using System.Xml;
 using System.IO;
 using Ionic.Zip;
 
-namespace UnrealBuildTool.IOS
+namespace UnrealBuildTool
 {
 	class UEDeployMac : UEBuildDeploy
 	{
-		/// <summary>
-		/// Register the platform with the UEBuildDeploy class
-		/// </summary>
-		public override void RegisterBuildDeploy()
-		{
-			// TODO: print debug info and handle any cases that would keep this from registering
-			UEBuildDeploy.RegisterBuildDeploy(UnrealTargetPlatform.Mac, this);
-		}
-
 		public override bool PrepTargetForDeployment(UEBuildTarget InTarget)
 		{
 			Log.TraceInformation("Deploying now!");
diff --git a/Engine/Source/Programs/UnrealBuildTool/System/UnrealBuildTool.cs b/Engine/Source/Programs/UnrealBuildTool/System/UnrealBuildTool.cs
--- a/Engine/Source/Programs/UnrealBuildTool/System/UnrealBuildTool.cs
+++ b/Engine/Source/Programs/UnrealBuildTool/System/UnrealBuildTool.cs
@@ -506,12 +506,6 @@ namespace UnrealBuildTool
 							var TempInst = (UEToolChain)(UBTAssembly.CreateInstance(CheckType.FullName, true));
 							TempInst.RegisterToolChain();
 						}
-						else if (Utils.ImplementsInterface<IUEBuildDeploy>(CheckType))
-						{
-							Log.TraceVerbose("    Registering build deploy  : {0}", CheckType.ToString());
-							var TempInst = (UEBuildDeploy)(UBTAssembly.CreateInstance(CheckType.FullName, true));
-							TempInst.RegisterBuildDeploy();
-						}
 						else if (CheckType.IsSubclassOf(typeof(UEPlatformProjectGenerator)))
 						{
 							ProjectGeneratorList.Add(CheckType);
@@ -1393,14 +1387,15 @@ namespace UnrealBuildTool
 							if ((Result == ECompilationResult.Succeeded) && (BuildConfiguration.bDeployAfterCompile == true) && (BuildConfiguration.bXGEExport == false) &&
 								(UEBuildConfiguration.bGenerateManifest == false) && (UEBuildConfiguration.bGenerateExternalFileList == false) && (UEBuildConfiguration.bCleanProject == false))
 							{
-								var DeployHandler = UEBuildDeploy.GetBuildDeploy(CheckPlatform);
-								if (DeployHandler != null)
+								IUEBuildPlatform BuildPlatform = UEBuildPlatform.GetBuildPlatform(CheckPlatform);
+
+								UEBuildDeploy DeploymentHandler;
+								if (BuildPlatform.TryCreateDeploymentHandler(UnrealBuildTool.GetUProjectFile(), out DeploymentHandler))
 								{
 									// We need to be able to identify the Target.Type we can derive it from the Arguments.
 									BuildConfiguration.bFlushBuildDirOnRemoteMac = false;
 									var TargetDescs = UEBuildTarget.ParseTargetCommandLine(Arguments, GetUProjectFile());
 									UEBuildTarget CheckTarget = UEBuildTarget.CreateTarget(TargetDescs[0]);	// @todo ubtmake: This may not work in assembler only mode.  We don't want to be loading target rules assemblies here either.
-									IUEBuildPlatform BuildPlatform = UEBuildPlatform.GetBuildPlatform(TargetDescs[0].Platform);
 									IUEToolChain ToolChain = UEToolChain.GetPlatformToolChain(BuildPlatform.GetCPPTargetPlatform(TargetDescs[0].Platform));
 									CheckTarget.SetupGlobalEnvironment(ToolChain);
 									if ((CheckTarget.TargetType == TargetRules.TargetType.Game) ||
@@ -1409,7 +1404,7 @@ namespace UnrealBuildTool
 									{
 										CheckTarget.AppName = CheckTarget.TargetName;
 									}
-									DeployHandler.PrepTargetForDeployment(CheckTarget);
+									DeploymentHandler.PrepTargetForDeployment(CheckTarget);
 								}
 							}
 						}
diff --git a/Engine/Source/Programs/UnrealBuildTool/System/XcodeProjectFileGenerator.cs b/Engine/Source/Programs/UnrealBuildTool/System/XcodeProjectFileGenerator.cs
--- a/Engine/Source/Programs/UnrealBuildTool/System/XcodeProjectFileGenerator.cs
+++ b/Engine/Source/Programs/UnrealBuildTool/System/XcodeProjectFileGenerator.cs
@@ -628,7 +628,7 @@ namespace UnrealBuildTool
 						GameName = "UE4Game";
 					}
 					Directory.CreateDirectory(Path.GetDirectoryName(InfoPlistPath));
-					IOS.UEDeployIOS.GeneratePList(ProjectPath, bIsUE4Game, GameName, TargetName, EngineRelative + "Engine", ProjectPath + "/Binaries/IOS/Payload");
+					UEDeployIOS.GeneratePList(ProjectPath, bIsUE4Game, GameName, TargetName, EngineRelative + "Engine", ProjectPath + "/Binaries/IOS/Payload");
 				}
 			}
 		}
diff --git a/Engine/Source/Programs/UnrealBuildTool/UWP/UEBuildUWP.cs b/Engine/Source/Programs/UnrealBuildTool/UWP/UEBuildUWP.cs
--- a/Engine/Source/Programs/UnrealBuildTool/UWP/UEBuildUWP.cs
+++ b/Engine/Source/Programs/UnrealBuildTool/UWP/UEBuildUWP.cs
@@ -475,5 +475,17 @@ namespace UnrealBuildTool
 			// Windows applications have many shared binaries between games
 			return false;
 		}
+
+		/// <summary>
+		/// Create a build deployment handler
+		/// </summary>
+		/// <param name="ProjectFile">The project file of the target being deployed. Used to find any deployment specific settings.</param>
+		/// <param name="DeploymentHandler">The output deployment handler</param>
+		/// <returns>True if the platform requires a deployment handler, false otherwise</returns>
+		public override bool TryCreateDeploymentHandler(FileReference ProjectFile, out UEBuildDeploy DeploymentHandler)
+		{
+			DeploymentHandler = new UWPDeploy();
+			return true;
+		}
 	}
 }
diff --git a/Engine/Source/Programs/UnrealBuildTool/UWP/UWPDeploy.cs b/Engine/Source/Programs/UnrealBuildTool/UWP/UWPDeploy.cs
--- a/Engine/Source/Programs/UnrealBuildTool/UWP/UWPDeploy.cs
+++ b/Engine/Source/Programs/UnrealBuildTool/UWP/UWPDeploy.cs
@@ -15,16 +15,6 @@ namespace UnrealBuildTool
 	/// </summary>
 	class UWPDeploy : UEBuildDeploy
 	{
-		/// <summary>
-		/// Register the platform with the UEBuildDeploy class
-		/// </summary>
-		public override void RegisterBuildDeploy()
-		{
-			// Register this deployment handle for UWP
-			Log.TraceVerbose("        Registering for {0}", UnrealTargetPlatform.UWP.ToString());
-			UEBuildDeploy.RegisterBuildDeploy(UnrealTargetPlatform.UWP, this);
-		}
-
 		/// <summary>
 		/// Utility function to delete a file
 		/// </summary>
diff --git a/Engine/Source/Programs/UnrealBuildTool/WinRT/UEBuildWinRT.cs b/Engine/Source/Programs/UnrealBuildTool/WinRT/UEBuildWinRT.cs
--- a/Engine/Source/Programs/UnrealBuildTool/WinRT/UEBuildWinRT.cs
+++ b/Engine/Source/Programs/UnrealBuildTool/WinRT/UEBuildWinRT.cs
@@ -634,5 +634,17 @@ namespace UnrealBuildTool
 		{
 			return WinRTPlatform.bCompileWinRT;
 		}
+
+		/// <summary>
+		/// Create a build deployment handler
+		/// </summary>
+		/// <param name="ProjectFile">The project file of the target being deployed. Used to find any deployment specific settings.</param>
+		/// <param name="DeploymentHandler">The output deployment handler</param>
+		/// <returns>True if the platform requires a deployment handler, false otherwise</returns>
+		public override bool TryCreateDeploymentHandler(FileReference ProjectFile, out UEBuildDeploy DeploymentHandler)
+		{
+			DeploymentHandler = new WinRTDeploy();
+			return true;
+		}
 	}
 }
diff --git a/Engine/Source/Programs/UnrealBuildTool/WinRT/WinRTDeploy.cs b/Engine/Source/Programs/UnrealBuildTool/WinRT/WinRTDeploy.cs
--- a/Engine/Source/Programs/UnrealBuildTool/WinRT/WinRTDeploy.cs
+++ b/Engine/Source/Programs/UnrealBuildTool/WinRT/WinRTDeploy.cs
@@ -13,16 +13,6 @@ namespace UnrealBuildTool
 	/// </summary>
 	class WinRTDeploy : UEBuildDeploy
 	{
-		/// <summary>
-		/// Register the platform with the UEBuildDeploy class
-		/// </summary>
-		public override void RegisterBuildDeploy()
-		{
-			// Register this deployment handle for WinRT
-			Log.TraceVerbose("        Registering for {0}", UnrealTargetPlatform.WinRT.ToString());
-			UEBuildDeploy.RegisterBuildDeploy(UnrealTargetPlatform.WinRT, this);
-		}
-
 		/// <summary>
 		/// Utility function to delete a file
 		/// </summary>
diff --git a/Engine/Source/Programs/UnrealBuildTool/Windows/UEBuildWindows.cs b/Engine/Source/Programs/UnrealBuildTool/Windows/UEBuildWindows.cs
--- a/Engine/Source/Programs/UnrealBuildTool/Windows/UEBuildWindows.cs
+++ b/Engine/Source/Programs/UnrealBuildTool/Windows/UEBuildWindows.cs
@@ -788,5 +788,17 @@ namespace UnrealBuildTool
 			// Windows applications have many shared binaries between games
 			return false;
 		}
+
+		/// <summary>
+		/// Create a build deployment handler
+		/// </summary>
+		/// <param name="ProjectFile">The project file of the target being deployed. Used to find any deployment specific settings.</param>
+		/// <param name="DeploymentHandler">The output deployment handler</param>
+		/// <returns>True if the platform requires a deployment handler, false otherwise</returns>
+		public override bool TryCreateDeploymentHandler(FileReference ProjectFile, out UEBuildDeploy DeploymentHandler)
+		{
+			DeploymentHandler = null;
+			return false;
+		}
 	}
 }
