Fix parsing of changelist descriptions that contain '\\r' characters. Was previously treating it at the end of the changelist summary, and throwing an exception when the next line wasn't "Affected Files".

diff --git a/Engine/Source/Programs/AutomationTool/AutomationUtils/P4Utils.cs b/Engine/Source/Programs/AutomationTool/AutomationUtils/P4Utils.cs
--- a/Engine/Source/Programs/AutomationTool/AutomationUtils/P4Utils.cs
+++ b/Engine/Source/Programs/AutomationTool/AutomationUtils/P4Utils.cs
@@ -1041,44 +1041,31 @@ namespace AutomationTool
 						for( ; LineIndex < Lines.Length; ++LineIndex )
 						{
 							Line = Lines[ LineIndex ];
-
-							if( String.IsNullOrEmpty( Line ) )
+							if(Line.Length > 0)
 							{
-								// Summaries end with a blank line (no tabs)
-								break;
-							}
+								// Stop once we reach a line that doesn't begin with a tab. It's possible (through changelist descriptions that contain embedded newlines, like \r\r\n on Windows) to get 
+								// empty lines that don't begin with a tab as we expect.
+								if(Line[0] != '\t')
+								{
+									break;
+								}
 
-							// Summary lines are supposed to begin with a single tab character (even empty lines)
-							if( Line[0] != '\t' )
-							{
-								throw new AutomationException("Was expecting every line of the P4 changes summary to start with a tab character");
-							}
+								// Remove the tab
+								var SummaryLine = Line.Substring( 1 );
 
-							// Remove the tab
-							var SummaryLine = Line.Substring( 1 );
+								// Add a CR if we already had some summary text
+								if( !String.IsNullOrEmpty( DescribeRecord.Summary ) )
+								{
+									DescribeRecord.Summary += "\n";
+								}
 
-							// Add a CR if we already had some summary text
-							if( !String.IsNullOrEmpty( DescribeRecord.Summary ) )
-							{
-								DescribeRecord.Summary += "\n";
+								// Append the summary line!
+								DescribeRecord.Summary += SummaryLine;
 							}
-
-							// Append the summary line!
-							DescribeRecord.Summary += SummaryLine;
 						}
 
-
-						++LineIndex;
-						if( LineIndex >= Lines.Length )
-						{
-							throw new AutomationException("Was expecting 'Affected files' to appear after the summary output from P4, but there were no more lines to read");
-						}
-
-						// If the summary ends with an empty newline, it doesn't seem to be prefixed with a tab
-						while (LineIndex < Lines.Length && Lines[LineIndex].Length == 0)
-						{
-							LineIndex++;
-						}
+						// Remove any trailing newlines from the end of the summary
+						DescribeRecord.Summary = DescribeRecord.Summary.TrimEnd('\n');
 
 						Line = Lines[ LineIndex ];
 
