StageNonMonolithic going away in favor of build reciepts

Merging 2621823

//depot/UE4-UT/...

to //depot/UE4/...

diff --git a/Engine/Source/Programs/AutomationTool/AutomationUtils/ProjectParams.cs b/Engine/Source/Programs/AutomationTool/AutomationUtils/ProjectParams.cs
--- a/Engine/Source/Programs/AutomationTool/AutomationUtils/ProjectParams.cs
+++ b/Engine/Source/Programs/AutomationTool/AutomationUtils/ProjectParams.cs
@@ -242,7 +242,6 @@ namespace AutomationTool
 			this.Stage = InParams.Stage;
 			this.SkipStage = InParams.SkipStage;
             this.StageDirectoryParam = InParams.StageDirectoryParam;
-            this.StageNonMonolithic = InParams.StageNonMonolithic;
 			this.Manifests = InParams.Manifests;
             this.CreateChunkInstall = InParams.CreateChunkInstall;
 			this.UE4Exe = InParams.UE4Exe;
@@ -318,7 +317,6 @@ namespace AutomationTool
 			string StageCommandline = null,
             string BundleName = null,
             string StageDirectoryParam = null,
-            bool? StageNonMonolithic = null,
 			string UE4Exe = null,
 			string SignPak = null,
 			List<UnrealTargetConfiguration> ClientConfigsToBuild = null,
@@ -523,7 +521,6 @@ namespace AutomationTool
 				this.Stage = true;
 			}
 			this.StageDirectoryParam = ParseParamValueIfNotSpecified(Command, StageDirectoryParam, "stagingdirectory", String.Empty, true);
-            this.StageNonMonolithic = GetParamValueIfNotSpecified(Command, StageNonMonolithic, this.StageNonMonolithic, "StageNonMonolithic");
 			this.bCodeSign = GetParamValueIfNotSpecified(Command, CodeSign, CommandUtils.IsBuildMachine, "CodeSign");
 			this.Manifests = GetParamValueIfNotSpecified(Command, Manifests, this.Manifests, "manifests");
             this.CreateChunkInstall = GetParamValueIfNotSpecified(Command, CreateChunkInstall, this.CreateChunkInstall, "createchunkinstall");
@@ -934,12 +931,7 @@ namespace AutomationTool
 
 		[Help("stagingdirectory=Path", "Directory to copy the builds to, i.e. -stagingdirectory=C:\\Stage")]
 		public string StageDirectoryParam;
-
-        /// <summary>
-        /// Whether the project should use non monolithic staging
-        /// </summary>
-        public bool StageNonMonolithic;
-
+        
 		[Help("ue4exe=ExecutableName", "Name of the UE4 Editor executable, i.e. -ue4exe=UE4Editor.exe")]
 		public string UE4Exe;
 
diff --git a/Engine/Source/Programs/AutomationTool/Linux/LinuxPlatform.Automation.cs b/Engine/Source/Programs/AutomationTool/Linux/LinuxPlatform.Automation.cs
--- a/Engine/Source/Programs/AutomationTool/Linux/LinuxPlatform.Automation.cs
+++ b/Engine/Source/Programs/AutomationTool/Linux/LinuxPlatform.Automation.cs
@@ -76,69 +76,48 @@ public abstract class BaseLinuxPlatform : Platform
 
 		SC.StageFiles(StagedFileType.NonUFS, CombinePaths(SC.ProjectRoot, "Content/Splash"), "Splash.bmp", false, null, null, true);
 
-        if (Params.StageNonMonolithic)
-        {
-            if (SC.DedicatedServer)
-            {
-                SC.StageFiles(StagedFileType.NonUFS, CombinePaths(SC.LocalRoot, "Engine/Binaries", SC.PlatformDir), "UE4Server*");
-                SC.StageFiles(StagedFileType.NonUFS, CombinePaths(SC.LocalRoot, "Engine/Binaries", SC.PlatformDir), "libUE4Server-*.so");
-                SC.StageFiles(StagedFileType.NonUFS, CombinePaths(SC.LocalRoot, "Engine/Plugins"), "libUE4Server-*.so", true, null, null, true);
-                SC.StageFiles(StagedFileType.NonUFS, CombinePaths(SC.ProjectRoot, "Binaries", SC.PlatformDir), "libUE4Server-*.so");
-                SC.StageFiles(StagedFileType.NonUFS, CombinePaths(SC.ProjectRoot, "Plugins"), "libUE4Server-*.so", true, null, null, true);
-            }
-            else
-            {
-                SC.StageFiles(StagedFileType.NonUFS, CombinePaths(SC.LocalRoot, "Engine/Binaries", SC.PlatformDir), "UE4*");
-                SC.StageFiles(StagedFileType.NonUFS, CombinePaths(SC.LocalRoot, "Engine/Binaries", SC.PlatformDir), "libUE4-*.so");
-                SC.StageFiles(StagedFileType.NonUFS, CombinePaths(SC.LocalRoot, "Engine/Plugins"), "libUE4-*.so", true, null, null, true);
-                SC.StageFiles(StagedFileType.NonUFS, CombinePaths(SC.ProjectRoot, "Binaries", SC.PlatformDir), "libUE4-*.so");
-                SC.StageFiles(StagedFileType.NonUFS, CombinePaths(SC.ProjectRoot, "Plugins"), "libUE4-*.so", true, null, null, true);
-            }
-        }
-        else
-        {
-	    	List<string> Exes = GetExecutableNames(SC);
-
-		    foreach (var Exe in Exes)
-		    {
-
-			    if (Exe.StartsWith(CombinePaths(SC.RuntimeProjectRootDir, "Binaries", SC.PlatformDir)))
-			    {
-				    // remap the project root. For content-only projects, rename the executable to the project name.
-				    if (!Params.IsCodeBasedProject && Exe == Exes[0])
-				    {
-					    SC.StageFiles(StagedFileType.NonUFS, CombinePaths(SC.ProjectRoot, "Binaries", SC.PlatformDir), Path.GetFileNameWithoutExtension(Exe), true, null, CommandUtils.CombinePaths(SC.RelativeProjectRootForStage, "Binaries", SC.PlatformDir), false, true, SC.ShortProjectName);
-				    }
-				    else
-				    {
-					    SC.StageFiles(StagedFileType.NonUFS, CombinePaths(SC.ProjectRoot, "Binaries", SC.PlatformDir), Path.GetFileNameWithoutExtension(Exe), true, null, CommandUtils.CombinePaths(SC.RelativeProjectRootForStage, "Binaries", SC.PlatformDir), false);
-				    }
-			    }
-			    else if (Exe.StartsWith(CombinePaths(SC.RuntimeRootDir, "Engine/Binaries", SC.PlatformDir)))
-			    {
-				    // Move the executable for content-only projects into the project directory, using the project name, so it can figure out the UProject to look for and is consistent with code projects.
-				    if (!Params.IsCodeBasedProject && Exe == Exes[0])
-				    {
-						// ensure the ue4game binary exists, if applicable
-						if (!SC.IsCodeBasedProject && !FileExists_NoExceptions(Params.ProjectGameExeFilename) && !SC.bIsCombiningMultiplePlatforms)
-						{
-							LogError("Failed to find game binary " + Params.ProjectGameExeFilename);
-                            throw new AutomationException(ExitCode.Error_MissingExecutable, "Stage Failed. Could not find game binary {0}. You may need to build the UE4 project with your target configuration and platform.", Params.ProjectGameExeFilename);
-						}
+	    List<string> Exes = GetExecutableNames(SC);
+
+		foreach (var Exe in Exes)
+		{
+
+			if (Exe.StartsWith(CombinePaths(SC.RuntimeProjectRootDir, "Binaries", SC.PlatformDir)))
+			{
+				// remap the project root. For content-only projects, rename the executable to the project name.
+				if (!Params.IsCodeBasedProject && Exe == Exes[0])
+				{
+					SC.StageFiles(StagedFileType.NonUFS, CombinePaths(SC.ProjectRoot, "Binaries", SC.PlatformDir), Path.GetFileNameWithoutExtension(Exe), true, null, CommandUtils.CombinePaths(SC.RelativeProjectRootForStage, "Binaries", SC.PlatformDir), false, true, SC.ShortProjectName);
+				}
+				else
+				{
+					SC.StageFiles(StagedFileType.NonUFS, CombinePaths(SC.ProjectRoot, "Binaries", SC.PlatformDir), Path.GetFileNameWithoutExtension(Exe), true, null, CommandUtils.CombinePaths(SC.RelativeProjectRootForStage, "Binaries", SC.PlatformDir), false);
+				}
+			}
+			else if (Exe.StartsWith(CombinePaths(SC.RuntimeRootDir, "Engine/Binaries", SC.PlatformDir)))
+			{
+				// Move the executable for content-only projects into the project directory, using the project name, so it can figure out the UProject to look for and is consistent with code projects.
+				if (!Params.IsCodeBasedProject && Exe == Exes[0])
+				{
+					// ensure the ue4game binary exists, if applicable
+					if (!SC.IsCodeBasedProject && !FileExists_NoExceptions(Params.ProjectGameExeFilename) && !SC.bIsCombiningMultiplePlatforms)
+					{
+						LogError("Failed to find game binary " + Params.ProjectGameExeFilename);
+                        throw new AutomationException(ExitCode.Error_MissingExecutable, "Stage Failed. Could not find game binary {0}. You may need to build the UE4 project with your target configuration and platform.", Params.ProjectGameExeFilename);
+					}
 
-					    SC.StageFiles(StagedFileType.NonUFS, CombinePaths(SC.LocalRoot, "Engine/Binaries", SC.PlatformDir), Path.GetFileNameWithoutExtension(Exe), true, null, CommandUtils.CombinePaths(SC.RelativeProjectRootForStage, "Binaries", SC.PlatformDir), false, true, SC.ShortProjectName);
-				    }
-				    else
-				    {
-					    SC.StageFiles(StagedFileType.NonUFS, CombinePaths(SC.LocalRoot, "Engine/Binaries", SC.PlatformDir), Path.GetFileNameWithoutExtension(Exe), true, null, null, false);
-				    }
-			    }
-			    else
-			    {
-				    throw new AutomationException("Can't stage the exe {0} because it doesn't start with {1} or {2}", Exe, CombinePaths(SC.RuntimeProjectRootDir, "Binaries", SC.PlatformDir), CombinePaths(SC.RuntimeRootDir, "Engine/Binaries", SC.PlatformDir));
-                }
+					SC.StageFiles(StagedFileType.NonUFS, CombinePaths(SC.LocalRoot, "Engine/Binaries", SC.PlatformDir), Path.GetFileNameWithoutExtension(Exe), true, null, CommandUtils.CombinePaths(SC.RelativeProjectRootForStage, "Binaries", SC.PlatformDir), false, true, SC.ShortProjectName);
+				}
+				else
+				{
+					SC.StageFiles(StagedFileType.NonUFS, CombinePaths(SC.LocalRoot, "Engine/Binaries", SC.PlatformDir), Path.GetFileNameWithoutExtension(Exe), true, null, null, false);
+				}
+			}
+			else
+			{
+				throw new AutomationException("Can't stage the exe {0} because it doesn't start with {1} or {2}", Exe, CombinePaths(SC.RuntimeProjectRootDir, "Binaries", SC.PlatformDir), CombinePaths(SC.RuntimeRootDir, "Engine/Binaries", SC.PlatformDir));
             }
-		}
+        }
+		
 	}
 
 	public override string GetCookPlatform(bool bDedicatedServer, bool bIsClientOnly, string CookFlavor)
diff --git a/Engine/Source/Programs/AutomationTool/Mac/MacPlatform.Automation.cs b/Engine/Source/Programs/AutomationTool/Mac/MacPlatform.Automation.cs
--- a/Engine/Source/Programs/AutomationTool/Mac/MacPlatform.Automation.cs
+++ b/Engine/Source/Programs/AutomationTool/Mac/MacPlatform.Automation.cs
@@ -64,94 +64,45 @@ public class MacPlatform : Platform
 
 	public override void GetFilesToDeployOrStage(ProjectParams Params, DeploymentContext SC)
 	{
-	
-		if (Params.StageNonMonolithic)
+		// the first app is the "main" one, the rest are marked as debug files for exclusion from chunking/distribution
+		StagedFileType WorkingFileType = StagedFileType.NonUFS;
+
+		List<string> Exes = GetExecutableNames(SC);
+		foreach (var Exe in Exes)
 		{
-			if (SC.DedicatedServer)
+			string AppBundlePath = "";
+			if (Exe.StartsWith(CombinePaths(SC.RuntimeProjectRootDir, "Binaries", SC.PlatformDir)))
 			{
-				if (SC.StageTargetConfigurations.Contains(UnrealTargetConfiguration.Development))
-				{
-					SC.StageFiles(StagedFileType.NonUFS, CombinePaths(SC.LocalRoot, "Engine/Binaries", SC.PlatformDir, "UE4Server.app"));
-					SC.StageFiles(StagedFileType.NonUFS, CombinePaths(SC.ProjectRoot, "Binaries", SC.PlatformDir), "UE4Server-" + SC.ShortProjectName + ".dylib");
-				}
-				if (SC.StageTargetConfigurations.Contains(UnrealTargetConfiguration.Test))
-				{
-					SC.StageFiles(StagedFileType.NonUFS, CombinePaths(SC.LocalRoot, "Engine/Binaries", SC.PlatformDir, "UE4Server-Mac-Test.app"));
-					SC.StageFiles(StagedFileType.NonUFS, CombinePaths(SC.ProjectRoot, "Binaries", SC.PlatformDir), "UE4Server-" + SC.ShortProjectName + "-Mac-Test.dylib");
-				}
-				if (SC.StageTargetConfigurations.Contains(UnrealTargetConfiguration.Shipping))
-				{
-					SC.StageFiles(StagedFileType.NonUFS, CombinePaths(SC.LocalRoot, "Engine/Binaries", SC.PlatformDir, "UE4Server-Mac-Shipping.app"));
-					SC.StageFiles(StagedFileType.NonUFS, CombinePaths(SC.ProjectRoot, "Binaries", SC.PlatformDir), "UE4Server-" + SC.ShortProjectName + "-Mac-Shipping.dylib");
-				}
-
-				SC.StageFiles(StagedFileType.NonUFS, CombinePaths(SC.LocalRoot, "Engine/Plugins"), "UE4Server-*.dylib", true);
-				SC.StageFiles(StagedFileType.NonUFS, CombinePaths(SC.ProjectRoot, "Plugins"), "UE4Server-*.dylib", true);
+				AppBundlePath = CombinePaths(SC.ShortProjectName, "Binaries", SC.PlatformDir, Path.GetFileNameWithoutExtension(Exe) + ".app");
+				StageAppBundle(SC, WorkingFileType, CombinePaths(SC.ProjectRoot, "Binaries", SC.PlatformDir, Path.GetFileNameWithoutExtension(Exe) + ".app"), AppBundlePath);
 			}
-			else
+			else if (Exe.StartsWith(CombinePaths(SC.RuntimeRootDir, "Engine/Binaries", SC.PlatformDir)))
 			{
-				if (SC.StageTargetConfigurations.Contains(UnrealTargetConfiguration.Development))
-				{
-					SC.StageFiles(StagedFileType.NonUFS, CombinePaths(SC.LocalRoot, "Engine/Binaries", SC.PlatformDir, "UE4.app"));
-					SC.StageFiles(StagedFileType.NonUFS, CombinePaths(SC.ProjectRoot, "Binaries", SC.PlatformDir), "UE4-" + SC.ShortProjectName + ".dylib");
-				}
-				if (SC.StageTargetConfigurations.Contains(UnrealTargetConfiguration.Test))
-				{
-					SC.StageFiles(StagedFileType.NonUFS, CombinePaths(SC.LocalRoot, "Engine/Binaries", SC.PlatformDir, "UE4-Mac-Test.app"));
-					SC.StageFiles(StagedFileType.NonUFS, CombinePaths(SC.ProjectRoot, "Binaries", SC.PlatformDir), "UE4-" + SC.ShortProjectName + "-Mac-Test.dylib");
-				}
-				if (SC.StageTargetConfigurations.Contains(UnrealTargetConfiguration.Shipping))
+				AppBundlePath = CombinePaths("Engine/Binaries", SC.PlatformDir, Path.GetFileNameWithoutExtension(Exe) + ".app");
+
+				string AbsoluteBundlePath = CombinePaths (SC.LocalRoot, AppBundlePath);
+				// ensure the ue4game binary exists, if applicable
+				if (!SC.IsCodeBasedProject && !Directory.Exists(AbsoluteBundlePath) && !SC.bIsCombiningMultiplePlatforms)
 				{
-					SC.StageFiles(StagedFileType.NonUFS, CombinePaths(SC.LocalRoot, "Engine/Binaries", SC.PlatformDir, "UE4-Mac-Shipping.app"));
-					SC.StageFiles(StagedFileType.NonUFS, CombinePaths(SC.ProjectRoot, "Binaries", SC.PlatformDir), "UE4-" + SC.ShortProjectName + "-Mac-Shipping.dylib");
+					LogError("Failed to find app bundle " + AbsoluteBundlePath);
+					throw new AutomationException(ExitCode.Error_MissingExecutable, "Could not find app bundle {0}. You may need to build the UE4 project with your target configuration and platform.", AbsoluteBundlePath);
 				}
 
-				SC.StageFiles(StagedFileType.NonUFS, CombinePaths(SC.LocalRoot, "Engine/Plugins"), "UE4-*.dylib", true);
-				SC.StageFiles(StagedFileType.NonUFS, CombinePaths(SC.ProjectRoot, "Plugins"), "UE4-*.dylib", true, null, null, true);
+				StageAppBundle(SC, WorkingFileType, CombinePaths(SC.LocalRoot, "Engine/Binaries", SC.PlatformDir, Path.GetFileNameWithoutExtension(Exe) + ".app"), AppBundlePath);
 			}
-		}
-		else
-		{
-			// the first app is the "main" one, the rest are marked as debug files for exclusion from chunking/distribution
-			StagedFileType WorkingFileType = StagedFileType.NonUFS;
 
-			List<string> Exes = GetExecutableNames(SC);
-			foreach (var Exe in Exes)
+			if (!string.IsNullOrEmpty(AppBundlePath))
 			{
-				string AppBundlePath = "";
-				if (Exe.StartsWith(CombinePaths(SC.RuntimeProjectRootDir, "Binaries", SC.PlatformDir)))
-				{
-					AppBundlePath = CombinePaths(SC.ShortProjectName, "Binaries", SC.PlatformDir, Path.GetFileNameWithoutExtension(Exe) + ".app");
-					StageAppBundle(SC, WorkingFileType, CombinePaths(SC.ProjectRoot, "Binaries", SC.PlatformDir, Path.GetFileNameWithoutExtension(Exe) + ".app"), AppBundlePath);
-				}
-				else if (Exe.StartsWith(CombinePaths(SC.RuntimeRootDir, "Engine/Binaries", SC.PlatformDir)))
-				{
-					AppBundlePath = CombinePaths("Engine/Binaries", SC.PlatformDir, Path.GetFileNameWithoutExtension(Exe) + ".app");
-
-					string AbsoluteBundlePath = CombinePaths (SC.LocalRoot, AppBundlePath);
-					// ensure the ue4game binary exists, if applicable
-					if (!SC.IsCodeBasedProject && !Directory.Exists(AbsoluteBundlePath) && !SC.bIsCombiningMultiplePlatforms)
-					{
-						LogError("Failed to find app bundle " + AbsoluteBundlePath);
-						throw new AutomationException(ExitCode.Error_MissingExecutable, "Could not find app bundle {0}. You may need to build the UE4 project with your target configuration and platform.", AbsoluteBundlePath);
-					}
-
-					StageAppBundle(SC, WorkingFileType, CombinePaths(SC.LocalRoot, "Engine/Binaries", SC.PlatformDir, Path.GetFileNameWithoutExtension(Exe) + ".app"), AppBundlePath);
-				}
+				SC.StageFiles(WorkingFileType, CombinePaths(SC.ProjectRoot, "Build/Mac"), "Application.icns", false, null, CombinePaths(AppBundlePath, "Contents/Resources"), true);
 
-				if (!string.IsNullOrEmpty(AppBundlePath))
+				if (Params.bUsesSteam)
 				{
-					SC.StageFiles(WorkingFileType, CombinePaths(SC.ProjectRoot, "Build/Mac"), "Application.icns", false, null, CombinePaths(AppBundlePath, "Contents/Resources"), true);
-
-					if (Params.bUsesSteam)
-					{
-						SC.StageFiles(StagedFileType.NonUFS, CombinePaths(SC.LocalRoot, "Engine/Source/ThirdParty/Steamworks/Steamv132/sdk/redistributable_bin/osx32"), "libsteam_api.dylib", false, null, CombinePaths(AppBundlePath, "Contents/MacOS"), true);
-					}
+					SC.StageFiles(StagedFileType.NonUFS, CombinePaths(SC.LocalRoot, "Engine/Source/ThirdParty/Steamworks/Steamv132/sdk/redistributable_bin/osx32"), "libsteam_api.dylib", false, null, CombinePaths(AppBundlePath, "Contents/MacOS"), true);
 				}
-
-				// the first app is the "main" one, the rest are marked as debug files for exclusion from chunking/distribution
-				WorkingFileType = StagedFileType.DebugNonUFS;
 			}
+
+			// the first app is the "main" one, the rest are marked as debug files for exclusion from chunking/distribution
+			WorkingFileType = StagedFileType.DebugNonUFS;
 		}
 
 		if (SC.bStageCrashReporter)
