Change deployment context and project params to use file references rather than raw paths.

diff --git a/Engine/Source/Programs/AutomationTool/Android/AndroidPlatform.Automation.cs b/Engine/Source/Programs/AutomationTool/Android/AndroidPlatform.Automation.cs
--- a/Engine/Source/Programs/AutomationTool/Android/AndroidPlatform.Automation.cs
+++ b/Engine/Source/Programs/AutomationTool/Android/AndroidPlatform.Automation.cs
@@ -29,7 +29,7 @@ public class AndroidPlatform : Platform
 
 	private static string GetFinalApkName(ProjectParams Params, string DecoratedExeName, bool bRenameUE4Game, string Architecture, string GPUArchitecture)
 	{
-		string ProjectDir = Path.Combine(Path.GetDirectoryName(Path.GetFullPath(Params.RawProjectPath)), "Binaries/Android");
+		string ProjectDir = Path.Combine(Path.GetDirectoryName(Path.GetFullPath(Params.RawProjectPath.FullName)), "Binaries/Android");
 
 		if (Params.Prebuilt)
 		{
@@ -1200,19 +1200,19 @@ public class AndroidPlatform : Platform
     
 	#region Hooks
 
-	public override void PostBuildTarget(UE4Build Build, string ProjectName, string UProjectPath, string Config)
+	public override void PostBuildTarget(UE4Build Build, FileReference UProjectPath, string TargetName, string Config)
 	{
 		// Run UBT w/ the prep for deployment only option
 		// This is required as UBT will 'fake' success when building via UAT and run
 		// the deployment prep step before all the required parts are present.
-		if (ProjectName.Length > 0)
+		if (!String.IsNullOrEmpty(TargetName) && TargetName.Length > 0)
 		{
-			string ProjectToBuild = ProjectName;
-			if (ProjectToBuild != "UE4Game" && !string.IsNullOrEmpty(UProjectPath))
+			string ProjectArg = "";
+			if(UProjectPath != null)
 			{
-				ProjectToBuild = UProjectPath;
+				ProjectArg = string.Format(" -project=\"{0}\"", UProjectPath);
 			}
-			string UBTCommand = string.Format("\"{0}\" Android {1} -prepfordeploy", ProjectToBuild, Config);
+			string UBTCommand = string.Format("{0} Android {1} -prepfordeploy{2}", TargetName, Config, ProjectArg);
 			CommandUtils.RunUBT(UE4Build.CmdEnv, Build.UBTExecutable, UBTCommand);
 		}
 	}
diff --git a/Engine/Source/Programs/AutomationTool/AutomationUtils/CommandletUtils.cs b/Engine/Source/Programs/AutomationTool/AutomationUtils/CommandletUtils.cs
--- a/Engine/Source/Programs/AutomationTool/AutomationUtils/CommandletUtils.cs
+++ b/Engine/Source/Programs/AutomationTool/AutomationUtils/CommandletUtils.cs
@@ -4,6 +4,7 @@ using System.Collections.Generic;
 using System.Linq;
 using System.Text;
 using System.IO;
+using UnrealBuildTool;
 
 
 namespace AutomationTool
@@ -21,7 +22,7 @@ namespace AutomationTool
 		/// <param name="Dirs">List of directories to cook, can be null</param>
 		/// <param name="TargetPlatform">Target platform.</param>
 		/// <param name="Parameters">List of additional parameters.</param>
-		public static void CookCommandlet(string ProjectName, string UE4Exe = "UE4Editor-Cmd.exe", string[] Maps = null, string[] Dirs = null, string InternationalizationPreset = "", string[] Cultures = null, string TargetPlatform = "WindowsNoEditor", string Parameters = "-Unversioned")
+		public static void CookCommandlet(FileReference ProjectName, string UE4Exe = "UE4Editor-Cmd.exe", string[] Maps = null, string[] Dirs = null, string InternationalizationPreset = "", string[] Cultures = null, string TargetPlatform = "WindowsNoEditor", string Parameters = "-Unversioned")
 		{
 			string MapsToCook = "";
 			if (IsNullOrEmpty(Maps))
@@ -59,7 +60,7 @@ namespace AutomationTool
         /// <param name="Maps">List of maps to cook, can be null in which case -MapIniSection=AllMaps is used.</param>
         /// <param name="TargetPlatform">Target platform.</param>
         /// <param name="Parameters">List of additional parameters.</param>
-        public static void DDCCommandlet(string ProjectName, string UE4Exe = "UE4Editor-Cmd.exe", string[] Maps = null, string TargetPlatform = "WindowsNoEditor", string Parameters = "")
+        public static void DDCCommandlet(FileReference ProjectName, string UE4Exe = "UE4Editor-Cmd.exe", string[] Maps = null, string TargetPlatform = "WindowsNoEditor", string Parameters = "")
         {
             string MapsToCook = "";
             if (!IsNullOrEmpty(Maps))
@@ -79,7 +80,7 @@ namespace AutomationTool
         /// <param name="Maps">List of maps to cook, can be null in which case -MapIniSection=AllMaps is used.</param>
         /// <param name="TargetPlatform">Target platform.</param>
         /// <param name="Parameters">List of additional parameters.</param>
-        public static List<string> GenerateDistillFileSetsCommandlet(string ProjectName, string ManifestFile, string UE4Exe = "UE4Editor-Cmd.exe", string[] Maps = null, string Parameters = "")
+        public static List<string> GenerateDistillFileSetsCommandlet(FileReference ProjectName, string ManifestFile, string UE4Exe = "UE4Editor-Cmd.exe", string[] Maps = null, string Parameters = "")
         {
             string MapsToCook = "";
             if (!IsNullOrEmpty(Maps))
@@ -136,7 +137,7 @@ namespace AutomationTool
         /// <param name="ProjectName">Project name.</param>
         /// <param name="UE4Exe">The name of the UE4 Editor executable to use.</param>
         /// <param name="Parameters">List of additional parameters.</param>
-        public static void UpdateGameProjectCommandlet(string ProjectName, string UE4Exe = "UE4Editor-Cmd.exe", string Parameters = "")
+        public static void UpdateGameProjectCommandlet(FileReference ProjectName, string UE4Exe = "UE4Editor-Cmd.exe", string Parameters = "")
         {
             RunCommandlet(ProjectName, UE4Exe, "UpdateGameProject", Parameters);
         }
@@ -148,7 +149,7 @@ namespace AutomationTool
 		/// <param name="UE4Exe">The name of the UE4 Editor executable to use.</param>
 		/// <param name="Commandlet">Commandlet name.</param>
 		/// <param name="Parameters">Command line parameters (without -run=)</param>
-		public static void RunCommandlet(string ProjectName, string UE4Exe, string Commandlet, string Parameters = null)
+		public static void RunCommandlet(FileReference ProjectName, string UE4Exe, string Commandlet, string Parameters = null)
 		{
 			Log("Running UE4Editor {0} for project {1}", Commandlet, ProjectName);
 
@@ -168,7 +169,7 @@ namespace AutomationTool
 			Log("Commandlet log file is {0}", LocalLogFile);
 			string Args = String.Format(
 				"{0} -run={1} {2} -abslog={3} -stdout -FORCELOGFLUSH -CrashForUAT -unattended {5}{4}",
-				(ProjectName == null) ? "" : CommandUtils.MakePathSafeToUseWithCommandLine(ProjectName),
+				(ProjectName == null) ? "" : CommandUtils.MakePathSafeToUseWithCommandLine(ProjectName.FullName),
 				Commandlet,
 				String.IsNullOrEmpty(Parameters) ? "" : Parameters,
 				CommandUtils.MakePathSafeToUseWithCommandLine(LocalLogFile),
@@ -211,7 +212,7 @@ namespace AutomationTool
 				CommandUtils.LogWarning("Commandlet {0} failed to copy the local log file from {1} to {2}. The log file will be lost.", Commandlet, LocalDDCLogFile, DestDDCLogFile);
 			}
 
-            string ProjectStatsDirectory = CombinePaths(Path.GetDirectoryName(ProjectName), "Saved", "Stats");
+            string ProjectStatsDirectory = CombinePaths(Path.GetDirectoryName(ProjectName.FullName), "Saved", "Stats");
             if (Directory.Exists(ProjectStatsDirectory))
             {
                 string DestCookerStats = CmdEnv.LogFolder;
diff --git a/Engine/Source/Programs/AutomationTool/AutomationUtils/DeploymentContext.cs b/Engine/Source/Programs/AutomationTool/AutomationUtils/DeploymentContext.cs
--- a/Engine/Source/Programs/AutomationTool/AutomationUtils/DeploymentContext.cs
+++ b/Engine/Source/Programs/AutomationTool/AutomationUtils/DeploymentContext.cs
@@ -23,9 +23,9 @@ public struct StageTarget
 public class DeploymentContext //: ProjectParams
 {
 	/// <summary>
-	/// Full path where the project exists (For uprojects this should include the uproj filename, otherwise just project folder)
+	/// Full path to the .uproject file
 	/// </summary>
-	public string RawProjectPath;
+	public FileReference RawProjectPath;
 
 	/// <summary>
 	///  true if we should stage crash reporter
@@ -232,7 +232,7 @@ public class DeploymentContext //: ProjectParams
 	public bool bIsCombiningMultiplePlatforms = false;
 
 	public DeploymentContext(
-		string RawProjectPathOrName,
+		FileReference RawProjectPathOrName,
 		string InLocalRoot,
 		string BaseStageDirectory,
 		string BaseArchiveDirectory,
@@ -299,12 +299,12 @@ public class DeploymentContext //: ProjectParams
         StageDirectory = CommandUtils.CombinePaths(BaseStageDirectory, FinalCookPlatform);
         ArchiveDirectory = CommandUtils.CombinePaths(BaseArchiveDirectory, FinalCookPlatform);
 
-		if (!CommandUtils.FileExists(RawProjectPath))
+		if (!CommandUtils.FileExists(RawProjectPath.FullName))
 		{
 			throw new AutomationException("Can't find uproject file {0}.", RawProjectPathOrName);
 		}
 
-		ProjectRoot = CommandUtils.CombinePaths(CommandUtils.GetDirectoryName(Path.GetFullPath(RawProjectPath)));
+		ProjectRoot = CommandUtils.CombinePaths(CommandUtils.GetDirectoryName(RawProjectPath.FullName));
 
 		if (!CommandUtils.DirectoryExists(ProjectRoot))
 		{
@@ -313,7 +313,7 @@ public class DeploymentContext //: ProjectParams
 
         RelativeProjectRootForStage = ShortProjectName;
 
-		ProjectArgForCommandLines = CommandUtils.MakePathSafeToUseWithCommandLine(RawProjectPath);
+		ProjectArgForCommandLines = CommandUtils.MakePathSafeToUseWithCommandLine(RawProjectPath.FullName);
 		CookSourceRuntimeRootDir = RuntimeRootDir = LocalRoot;
 		RuntimeProjectRootDir = ProjectRoot;
 
diff --git a/Engine/Source/Programs/AutomationTool/AutomationUtils/Platform.cs b/Engine/Source/Programs/AutomationTool/AutomationUtils/Platform.cs
--- a/Engine/Source/Programs/AutomationTool/AutomationUtils/Platform.cs
+++ b/Engine/Source/Programs/AutomationTool/AutomationUtils/Platform.cs
@@ -482,7 +482,7 @@ namespace AutomationTool
 
 		}
 
-		public virtual void PostBuildTarget(UE4Build Build, string ProjectName, string UProjectPath, string Config)
+		public virtual void PostBuildTarget(UE4Build Build, FileReference UProjectPath, string TargetName, string Config)
 		{
 
 		}
diff --git a/Engine/Source/Programs/AutomationTool/AutomationUtils/ProjectParams.cs b/Engine/Source/Programs/AutomationTool/AutomationUtils/ProjectParams.cs
--- a/Engine/Source/Programs/AutomationTool/AutomationUtils/ProjectParams.cs
+++ b/Engine/Source/Programs/AutomationTool/AutomationUtils/ProjectParams.cs
@@ -306,7 +306,7 @@ namespace AutomationTool
 		/// If a parameter value is not set, it will be parsed from the command line if the command is null, the default value will be used.
 		/// </summary>
 		public ProjectParams(			
-			string RawProjectPath,
+			FileReference RawProjectPath,
 
 			CommandUtils Command = null,
 			string Device = null,			
@@ -759,9 +759,9 @@ namespace AutomationTool
 		#region Shared
 
 		/// <summary>
-		/// Shared: Full path where the project exists (For uprojects this should include the uproj filename, otherwise just project folder)
+		/// Shared: Full path to the .uproject file
 		/// </summary>
-		public string RawProjectPath { private set; get; }
+		public FileReference RawProjectPath { private set; get; }
 
 		/// <summary>
 		/// Shared: The current project is a foreign project, commandline: -foreign
@@ -922,10 +922,10 @@ namespace AutomationTool
                 }
                 if ( HasDLCName )
                 {
-                     return Path.GetFullPath( CommandUtils.CombinePaths(Path.GetDirectoryName(RawProjectPath), "Plugins", DLCName, "Saved", "StagedBuilds" ) );
+                     return Path.GetFullPath( CommandUtils.CombinePaths(Path.GetDirectoryName(RawProjectPath.FullName), "Plugins", DLCName, "Saved", "StagedBuilds" ) );
                 }
                 // default return the project saved\stagedbuilds directory
-                return Path.GetFullPath( CommandUtils.CombinePaths(Path.GetDirectoryName(RawProjectPath), "Saved", "StagedBuilds") );
+                return Path.GetFullPath( CommandUtils.CombinePaths(Path.GetDirectoryName(RawProjectPath.FullName), "Saved", "StagedBuilds") );
 			}
 		}
 
@@ -948,7 +948,7 @@ namespace AutomationTool
 		{
 			get
 			{
-                return Path.GetFullPath(String.IsNullOrEmpty(ArchiveDirectoryParam) ? CommandUtils.CombinePaths(Path.GetDirectoryName(RawProjectPath), "ArchivedBuilds") : ArchiveDirectoryParam);
+                return Path.GetFullPath(String.IsNullOrEmpty(ArchiveDirectoryParam) ? CommandUtils.CombinePaths(Path.GetDirectoryName(RawProjectPath.FullName), "ArchivedBuilds") : ArchiveDirectoryParam);
 			}
 		}
 
@@ -1728,7 +1728,7 @@ namespace AutomationTool
 			// the full solution for per-platform packaging settings.
 			if (!Manifests)
 			{				
-				ConfigCacheIni GameIni = ConfigCacheIni.CreateConfigCacheIni(UnrealTargetPlatform.Unknown, "Game", new DirectoryReference(Path.GetDirectoryName(RawProjectPath)));
+				ConfigCacheIni GameIni = ConfigCacheIni.CreateConfigCacheIni(UnrealTargetPlatform.Unknown, "Game", RawProjectPath.Directory);
 				String IniPath = "/Script/UnrealEd.ProjectPackagingSettings";
 				bool bSetting = false;
 				if (!GameIni.GetBool(IniPath, "bGenerateChunks", out bSetting))
@@ -1847,7 +1847,7 @@ namespace AutomationTool
 		}
 		private bool bIsCodeBasedProject;
 
-		public string CodeBasedUprojectPath
+		public FileReference CodeBasedUprojectPath
 		{
             get { return IsCodeBasedProject ? RawProjectPath : null; }
 		}
@@ -2031,15 +2031,15 @@ namespace AutomationTool
 
 		public void Validate()
 		{
-			if (String.IsNullOrEmpty(RawProjectPath))
+			if (RawProjectPath == null)
 			{
 				throw new AutomationException("RawProjectPath can't be empty.");
 			}
-            if (!RawProjectPath.EndsWith(".uproject", StringComparison.InvariantCultureIgnoreCase))
+            if (!RawProjectPath.HasExtension(".uproject"))
             {
                 throw new AutomationException("RawProjectPath {0} must end with .uproject", RawProjectPath);
             }
-            if (!CommandUtils.FileExists_NoExceptions(RawProjectPath))
+            if (!CommandUtils.FileExists_NoExceptions(RawProjectPath.FullName))
             {
                 throw new AutomationException("RawProjectPath {0} file must exist", RawProjectPath);
             }
diff --git a/Engine/Source/Programs/AutomationTool/AutomationUtils/ProjectUtils.cs b/Engine/Source/Programs/AutomationTool/AutomationUtils/ProjectUtils.cs
--- a/Engine/Source/Programs/AutomationTool/AutomationUtils/ProjectUtils.cs
+++ b/Engine/Source/Programs/AutomationTool/AutomationUtils/ProjectUtils.cs
@@ -27,7 +27,7 @@ namespace AutomationTool
 		/// <summary>
 		/// Full Project path. Must be a .uproject file
 		/// </summary>
-		public string RawProjectPath;
+		public FileReference RawProjectPath;
 
 		/// <summary>
 		/// True if the uproject contains source code.
@@ -107,9 +107,9 @@ namespace AutomationTool
 		/// <param name="RawProjectPath">Full project path.</param>
 		/// <param name="bIsUProjectFile">True if a uproject.</param>
 		/// <returns>Short project name</returns>
-		public static string GetShortProjectName(string RawProjectPath)
+		public static string GetShortProjectName(FileReference RawProjectPath)
 		{
-			return CommandUtils.GetFilenameWithoutAnyExtensions(RawProjectPath);
+			return CommandUtils.GetFilenameWithoutAnyExtensions(RawProjectPath.FullName);
 		}
 
 		/// <summary>
@@ -117,12 +117,12 @@ namespace AutomationTool
 		/// </summary>
 		/// <param name="RawProjectPath">Full project path.</param>
 		/// <returns>Properties of the project.</returns>
-		public static ProjectProperties GetProjectProperties(string RawProjectPath, List<UnrealTargetPlatform> ClientTargetPlatforms = null)
+		public static ProjectProperties GetProjectProperties(FileReference RawProjectPath, List<UnrealTargetPlatform> ClientTargetPlatforms = null)
 		{
 			string ProjectKey = "UE4";
-			if (!String.IsNullOrEmpty(RawProjectPath))
+			if (RawProjectPath != null)
 			{
-				ProjectKey = CommandUtils.ConvertSeparators(PathSeparator.Slash, Path.GetFullPath(RawProjectPath));
+				ProjectKey = CommandUtils.ConvertSeparators(PathSeparator.Slash, RawProjectPath.FullName);
 			}
 			ProjectProperties Properties;
 			if (PropertiesCache.TryGetValue(ProjectKey, out Properties) == false)
@@ -138,7 +138,7 @@ namespace AutomationTool
 		/// </summary>
 		/// <param name="RawProjectPath">Full project path.</param>
 		/// <returns>True if the project is a UProject file with source code.</returns>
-		public static bool IsCodeBasedUProjectFile(string RawProjectPath)
+		public static bool IsCodeBasedUProjectFile(FileReference RawProjectPath)
 		{
 			return GetProjectProperties(RawProjectPath, null).bIsCodeBasedProject;
 		}
@@ -158,18 +158,18 @@ namespace AutomationTool
 			return ProjectClientBinariesPath;
 		}
 
-		private static bool RequiresTempTarget(string RawProjectPath, List<UnrealTargetPlatform> ClientTargetPlatforms)
+		private static bool RequiresTempTarget(FileReference RawProjectPath, List<UnrealTargetPlatform> ClientTargetPlatforms)
 		{
 			// check to see if we already have a Target.cs file
-			if (File.Exists (Path.Combine (Path.GetDirectoryName (RawProjectPath), "Source", Path.GetFileNameWithoutExtension (RawProjectPath) + ".Target.cs")))
+			if (File.Exists (Path.Combine (Path.GetDirectoryName (RawProjectPath.FullName), "Source", RawProjectPath.GetFileNameWithoutExtension() + ".Target.cs")))
 			{
 				return false;
 			}
-			else if (Directory.Exists(Path.Combine(Path.GetDirectoryName(RawProjectPath), "Source")))
+			else if (Directory.Exists(Path.Combine(Path.GetDirectoryName(RawProjectPath.FullName), "Source")))
 			{
 				// wasn't one in the main Source directory, let's check all sub-directories
 				//@todo: may want to read each target.cs to see if it has a target corresponding to the project name as a final check
-				FileInfo[] Files = (new DirectoryInfo( Path.Combine (Path.GetDirectoryName (RawProjectPath), "Source")).GetFiles ("*.Target.cs", SearchOption.AllDirectories));
+				FileInfo[] Files = (new DirectoryInfo( Path.Combine (Path.GetDirectoryName (RawProjectPath.FullName), "Source")).GetFiles ("*.Target.cs", SearchOption.AllDirectories));
 				if (Files.Length > 0)
 				{
 					return false;
@@ -204,15 +204,15 @@ namespace AutomationTool
 			}
 
 			// Read the project descriptor, and find all the plugins available to this project
-			ProjectDescriptor Project = ProjectDescriptor.FromFile(RawProjectPath);
-			List<PluginInfo> AvailablePlugins = Plugins.ReadAvailablePlugins(new DirectoryReference(BuildConfiguration.RelativeEnginePath), new FileReference(RawProjectPath));
+			ProjectDescriptor Project = ProjectDescriptor.FromFile(RawProjectPath.FullName);
+			List<PluginInfo> AvailablePlugins = Plugins.ReadAvailablePlugins(new DirectoryReference(BuildConfiguration.RelativeEnginePath), RawProjectPath);
 
 			// check the target platforms for any differences in build settings or additional plugins
 			bool RetVal = false;
 			foreach (UnrealTargetPlatform TargetPlatformType in TargetPlatforms)
 			{
 				IUEBuildPlatform BuildPlat = UEBuildPlatform.GetBuildPlatform(TargetPlatformType, true);
-				if (!GlobalCommandLine.Rocket && BuildPlat != null && !(BuildPlat as UEBuildPlatform).HasDefaultBuildConfig(TargetPlatformType, new DirectoryReference(Path.GetDirectoryName(RawProjectPath))))
+				if (!GlobalCommandLine.Rocket && BuildPlat != null && !(BuildPlat as UEBuildPlatform).HasDefaultBuildConfig(TargetPlatformType, RawProjectPath.Directory))
 				{
 					RetVal = true;
 					break;
@@ -238,17 +238,17 @@ namespace AutomationTool
 			return RetVal;
 		}
 
-		private static void GenerateTempTarget(string RawProjectPath)
+		private static void GenerateTempTarget(FileReference RawProjectPath)
 		{
 			// read in the template target cs file
 			var TempCSFile = CommandUtils.CombinePaths(CommandUtils.CmdEnv.LocalRoot, "Engine", "Build", "Target.cs.template");
 			string TargetCSFile = File.ReadAllText(TempCSFile);
 
 			// replace {GAME_NAME} with the game name
-			TargetCSFile = TargetCSFile.Replace("{GAME_NAME}", Path.GetFileNameWithoutExtension(RawProjectPath));
+			TargetCSFile = TargetCSFile.Replace("{GAME_NAME}", Path.GetFileNameWithoutExtension(RawProjectPath.FullName));
 
 			// write out the file in a new Source directory
-			string FileName = CommandUtils.CombinePaths(Path.GetDirectoryName(RawProjectPath), "Intermediate", "Source", Path.GetFileNameWithoutExtension(RawProjectPath) + ".Target.cs");
+			string FileName = CommandUtils.CombinePaths(Path.GetDirectoryName(RawProjectPath.FullName), "Intermediate", "Source", Path.GetFileNameWithoutExtension(RawProjectPath.FullName) + ".Target.cs");
 			if (!Directory.Exists(Path.GetDirectoryName(FileName)))
 			{
 				Directory.CreateDirectory(Path.GetDirectoryName(FileName));
@@ -262,25 +262,25 @@ namespace AutomationTool
 		/// </summary>
 		/// <param name="RawProjectPath">Full project path.</param>
 		/// <returns>Project properties.</returns>
-		private static ProjectProperties DetectProjectProperties(string RawProjectPath, List<UnrealTargetPlatform> ClientTargetPlatforms)
+		private static ProjectProperties DetectProjectProperties(FileReference RawProjectPath, List<UnrealTargetPlatform> ClientTargetPlatforms)
 		{
 			var Properties = new ProjectProperties();
 			Properties.RawProjectPath = RawProjectPath;
 
 			// detect if the project is content only, but has non-default build settings
 			List<string> ExtraSearchPaths = null;
-			if (!string.IsNullOrEmpty(RawProjectPath))
+			if (RawProjectPath != null)
 			{
 				if (RequiresTempTarget(RawProjectPath, ClientTargetPlatforms))
 				{
 					GenerateTempTarget(RawProjectPath);
 					Properties.bWasGenerated = true;
 					ExtraSearchPaths = new List<string>();
-					ExtraSearchPaths.Add(CommandUtils.CombinePaths(Path.GetDirectoryName(RawProjectPath), "Intermediate", "Source"));
+					ExtraSearchPaths.Add(CommandUtils.CombinePaths(Path.GetDirectoryName(RawProjectPath.FullName), "Intermediate", "Source"));
 				}
-				else if (File.Exists(Path.Combine(Path.GetDirectoryName(RawProjectPath), "Intermediate", "Source", Path.GetFileNameWithoutExtension(RawProjectPath) + ".Target.cs")))
+				else if (File.Exists(Path.Combine(Path.GetDirectoryName(RawProjectPath.FullName), "Intermediate", "Source", Path.GetFileNameWithoutExtension(RawProjectPath.FullName) + ".Target.cs")))
 				{
-					File.Delete(Path.Combine(Path.GetDirectoryName(RawProjectPath), "Intermediate", "Source", Path.GetFileNameWithoutExtension(RawProjectPath) + ".Target.cs"));
+					File.Delete(Path.Combine(Path.GetDirectoryName(RawProjectPath.FullName), "Intermediate", "Source", Path.GetFileNameWithoutExtension(RawProjectPath.FullName) + ".Target.cs"));
 				}
 			}
 
@@ -292,14 +292,14 @@ namespace AutomationTool
 			else
 			{
 				// should never ask for engine targets if we can't compile
-				if (String.IsNullOrEmpty(RawProjectPath))
+				if (RawProjectPath == null)
 				{
 					throw new AutomationException("Cannot determine engine targets if we can't compile.");
 				}
 
 				Properties.bIsCodeBasedProject = Properties.bWasGenerated;
 				// if there's a Source directory with source code in it, then mark us as having source code
-				string SourceDir = CommandUtils.CombinePaths(Path.GetDirectoryName(RawProjectPath), "Source");
+				string SourceDir = CommandUtils.CombinePaths(Path.GetDirectoryName(RawProjectPath.FullName), "Source");
 				if (Directory.Exists(SourceDir))
 				{
 					string[] CppFiles = Directory.GetFiles(SourceDir, "*.cpp", SearchOption.AllDirectories);
@@ -309,14 +309,14 @@ namespace AutomationTool
 			}
 
 			// check to see if the uproject loads modules, only if we haven't already determined it is a code based project
-			if (!Properties.bIsCodeBasedProject)
+			if (!Properties.bIsCodeBasedProject && RawProjectPath != null)
 			{
-				string uprojectStr = File.ReadAllText(RawProjectPath);
+				string uprojectStr = File.ReadAllText(RawProjectPath.FullName);
 				Properties.bIsCodeBasedProject = uprojectStr.Contains("\"Modules\"");
 			}
 
 			// Get all ini files
-			if (!String.IsNullOrWhiteSpace(RawProjectPath))
+			if (RawProjectPath != null)
 			{
 				CommandUtils.LogVerbose("Loading ini files for {0}", RawProjectPath);
 
@@ -325,7 +325,7 @@ namespace AutomationTool
 				{
 					if (TargetPlatformType != UnrealTargetPlatform.Unknown)
 					{
-						var Config = ConfigCacheIni.CreateConfigCacheIni(TargetPlatformType, "Engine", new DirectoryReference(Path.GetDirectoryName(RawProjectPath)), new DirectoryReference(EngineDirectory));
+						var Config = ConfigCacheIni.CreateConfigCacheIni(TargetPlatformType, "Engine", RawProjectPath.Directory, new DirectoryReference(EngineDirectory));
 						Properties.EngineConfigs.Add(TargetPlatformType, Config);
 					}
 				}
@@ -334,7 +334,7 @@ namespace AutomationTool
 				{
 					if (TargetPlatformType != UnrealTargetPlatform.Unknown)
 					{
-						var Config = ConfigCacheIni.CreateConfigCacheIni(TargetPlatformType, "Game", new DirectoryReference(Path.GetDirectoryName(RawProjectPath)));
+						var Config = ConfigCacheIni.CreateConfigCacheIni(TargetPlatformType, "Game", RawProjectPath.Directory);
 						Properties.GameConfigs.Add(TargetPlatformType, Config);
 					}
 				}
@@ -351,7 +351,7 @@ namespace AutomationTool
 		/// <param name="TargetType">Target type.</param>
 		/// <param name="bIsUProjectFile">True if uproject file.</param>
 		/// <returns>Binaries path.</returns>
-		public static string GetClientProjectBinariesRootPath(string RawProjectPath, TargetRules.TargetType TargetType, bool bIsCodeBasedProject)
+		public static string GetClientProjectBinariesRootPath(FileReference RawProjectPath, TargetRules.TargetType TargetType, bool bIsCodeBasedProject)
 		{
 			var BinPath = String.Empty;
 			switch (TargetType)
@@ -367,7 +367,7 @@ namespace AutomationTool
 					}
 					else
 					{
-						BinPath = CommandUtils.CombinePaths(CommandUtils.GetDirectoryName(RawProjectPath), "Binaries");
+						BinPath = CommandUtils.CombinePaths(CommandUtils.GetDirectoryName(RawProjectPath.FullName), "Binaries");
 					}
 					break;
 			}
@@ -404,13 +404,13 @@ namespace AutomationTool
 
 			var GameFolders = new List<DirectoryReference>();
 			var RulesFolder = new DirectoryReference(GetRulesAssemblyFolder());
-			if (!String.IsNullOrEmpty(Properties.RawProjectPath))
+			if (Properties.RawProjectPath != null)
 			{
 				CommandUtils.LogVerbose("Looking for targets for project {0}", Properties.RawProjectPath);
 
 				TargetsDllFilename = FileReference.Combine(RulesFolder, String.Format("UATRules{0}.dll", Properties.RawProjectPath.GetHashCode()));
 
-				FullProjectPath = CommandUtils.GetDirectoryName(Properties.RawProjectPath);
+				FullProjectPath = CommandUtils.GetDirectoryName(Properties.RawProjectPath.FullName);
 				GameFolders.Add(new DirectoryReference(FullProjectPath));
 				CommandUtils.LogVerbose("Searching for target rule files in {0}", FullProjectPath);
 			}
@@ -602,7 +602,7 @@ namespace AutomationTool
         public class BranchUProject
         {
             public string GameName;
-            public string FilePath;
+            public FileReference FilePath;
             public ProjectProperties Properties;
 
             public BranchUProject(UnrealBuildTool.UProjectInfo InfoEntry)
@@ -610,14 +610,14 @@ namespace AutomationTool
                 GameName = InfoEntry.GameName;
 
                 //not sure what the heck this path is relative to
-                FilePath = InfoEntry.FilePath.FullName;
+                FilePath = InfoEntry.FilePath;
 
-                if (!CommandUtils.FileExists_NoExceptions(FilePath))
+                if (!CommandUtils.FileExists_NoExceptions(FilePath.FullName))
                 {
                     throw new AutomationException("Could not resolve relative path corrctly {0} -> {1} which doesn't exist.", InfoEntry.FilePath, FilePath);
                 }
 
-                Properties = ProjectUtils.GetProjectProperties(Path.GetFullPath(FilePath));
+                Properties = ProjectUtils.GetProjectProperties(FilePath);
 
 
 
@@ -625,7 +625,7 @@ namespace AutomationTool
             public BranchUProject()
             {
                 GameName = "UE4";
-                Properties = ProjectUtils.GetProjectProperties("");
+                Properties = ProjectUtils.GetProjectProperties(null);
                 if (!Properties.Targets.ContainsKey(TargetRules.TargetType.Editor))
                 {
                     throw new AutomationException("Base UE4 project did not contain an editor target.");
@@ -732,7 +732,7 @@ namespace AutomationTool
 					{
 						NonCodeProjects.Add(UProject);
 						// the base project uses BlankProject if it really needs a .uproject file
-						if (String.IsNullOrEmpty(BaseEngineProject.FilePath) && UProject.GameName == "BlankProject")
+						if (BaseEngineProject.FilePath == null && UProject.GameName == "BlankProject")
 						{
 							BaseEngineProject.FilePath = UProject.FilePath;
 						}
diff --git a/Engine/Source/Programs/AutomationTool/AutomationUtils/UBTUtils.cs b/Engine/Source/Programs/AutomationTool/AutomationUtils/UBTUtils.cs
--- a/Engine/Source/Programs/AutomationTool/AutomationUtils/UBTUtils.cs
+++ b/Engine/Source/Programs/AutomationTool/AutomationUtils/UBTUtils.cs
@@ -2,6 +2,7 @@
 using System;
 using System.Collections.Generic;
 using System.Text;
+using UnrealBuildTool;
 
 namespace AutomationTool
 {
@@ -58,16 +59,16 @@ namespace AutomationTool
 		/// <param name="Platform">Platform to build for.</param>
 		/// <param name="Config">Configuration to build.</param>
 		/// <param name="AdditionalArgs">Additional arguments to pass on to UBT.</param>
-		public static string UBTCommandline(string Project, string Target, string Platform, string Config, string AdditionalArgs = "")
+		public static string UBTCommandline(FileReference Project, string Target, string Platform, string Config, string AdditionalArgs = "")
 		{
 			string CmdLine;
-			if (String.IsNullOrEmpty(Project))
+			if (Project == null)
 			{
 				CmdLine = String.Format("{0} {1} {2} {3}", Target, Platform, Config, AdditionalArgs);
 			}
 			else
 			{
-				CmdLine = String.Format("{0} {1} {2} -Project={3} {4}", Target, Platform, Config, CommandUtils.MakePathSafeToUseWithCommandLine(Project), AdditionalArgs);
+				CmdLine = String.Format("{0} {1} {2} -Project={3} {4}", Target, Platform, Config, CommandUtils.MakePathSafeToUseWithCommandLine(Project.FullName), AdditionalArgs);
 			}
 			return CmdLine;
 		}
@@ -83,7 +84,7 @@ namespace AutomationTool
 		/// <param name="Config">Configuration to build.</param>
 		/// <param name="AdditionalArgs">Additional arguments to pass on to UBT.</param>
 		/// <param name="LogName">Optional logifle name.</param>
-		public static void RunUBT(CommandEnvironment Env, string UBTExecutable, string Project, string Target, string Platform, string Config, string AdditionalArgs = "", string LogName = null, Dictionary<string, string> EnvVars = null)
+		public static void RunUBT(CommandEnvironment Env, string UBTExecutable, FileReference Project, string Target, string Platform, string Config, string AdditionalArgs = "", string LogName = null, Dictionary<string, string> EnvVars = null)
 		{
 			RunUBT(Env, UBTExecutable, UBTCommandline(Project, Target, Platform, Config, AdditionalArgs), LogName, EnvVars);
 		}
diff --git a/Engine/Source/Programs/AutomationTool/AutomationUtils/UE4Build.cs b/Engine/Source/Programs/AutomationTool/AutomationUtils/UE4Build.cs
--- a/Engine/Source/Programs/AutomationTool/AutomationUtils/UE4Build.cs
+++ b/Engine/Source/Programs/AutomationTool/AutomationUtils/UE4Build.cs
@@ -6,6 +6,7 @@ using System.IO;
 using System.Xml;
 using System.Diagnostics;
 using System.Web.Script.Serialization;
+using UnrealBuildTool;
 
 namespace AutomationTool
 {
@@ -150,18 +151,18 @@ namespace AutomationTool
 			public string CommandLine;
 			public UnrealBuildTool.UnrealTargetPlatform Platform;
 			public string Config;
-			public string ProjectName;
-			public string UProjectPath;
+			public string TargetName;
+			public FileReference UProjectPath;
 			public List<string> XgeXmlFiles;
 			public string OutputCaption;
 		}
 
-		XGEItem XGEPrepareBuildWithUBT(string ProjectName, string TargetName, UnrealBuildTool.UnrealTargetPlatform Platform, string Config, string UprojectPath, bool ForceMonolithic = false, bool ForceNonUnity = false, bool ForceDebugInfo = false, string InAddArgs = "", bool ForceUnity = false, Dictionary<string, string> EnvVars = null)
+		XGEItem XGEPrepareBuildWithUBT(string TargetName, UnrealBuildTool.UnrealTargetPlatform Platform, string Config, FileReference UprojectPath, bool ForceMonolithic = false, bool ForceNonUnity = false, bool ForceDebugInfo = false, string InAddArgs = "", bool ForceUnity = false, Dictionary<string, string> EnvVars = null)
 		{
 			string AddArgs = "";
-			if (string.IsNullOrEmpty(UprojectPath) == false)
+			if (UprojectPath != null)
 			{
-				AddArgs += " " + CommandUtils.MakePathSafeToUseWithCommandLine(UprojectPath);
+				AddArgs += " " + CommandUtils.MakePathSafeToUseWithCommandLine(UprojectPath.FullName);
 			}
 			AddArgs += " " + InAddArgs;
 			if (ForceMonolithic)
@@ -192,20 +193,20 @@ namespace AutomationTool
 
 			using(TelemetryStopwatch GenerateManifestStopwatch = new TelemetryStopwatch("GenerateXGEManifest.{0}.{1}.{2}", TargetName, Platform.ToString(), Config))
 			{
-				RunUBT(CmdEnv, UBTExecutable: UBTExecutable, Project: ProjectName, Target: TargetName, Platform: Platform.ToString(), Config: Config, AdditionalArgs: "-generatemanifest -nobuilduht -xgeexport" + AddArgs, EnvVars: EnvVars);
+				RunUBT(CmdEnv, UBTExecutable: UBTExecutable, Project: UprojectPath, Target: TargetName, Platform: Platform.ToString(), Config: Config, AdditionalArgs: "-generatemanifest -nobuilduht -xgeexport" + AddArgs, EnvVars: EnvVars);
 			}
 
 			PrepareManifest(UBTManifest, true);
 
 			Result.Platform = Platform;
 			Result.Config = Config;
-			Result.ProjectName = string.IsNullOrEmpty(ProjectName) ? TargetName : ProjectName; // use the target as the project if no project is set
+			Result.TargetName = TargetName;
 			Result.UProjectPath = UprojectPath;
 			Result.Manifest = ReadManifest(UBTManifest);
-			Result.OutputCaption = String.Format("{0}-{1}-{2}", Path.GetFileNameWithoutExtension(Result.ProjectName), Platform.ToString(), Config.ToString());
+			Result.OutputCaption = String.Format("{0}-{1}-{2}", TargetName, Platform.ToString(), Config.ToString());
 			DeleteFile(UBTManifest);
 
-			Result.CommandLine = UBTExecutable + " " + UBTCommandline(Project: ProjectName, Target: TargetName, Platform: Platform.ToString(), Config: Config, AdditionalArgs: "-noxge -nobuilduht" + AddArgs);
+			Result.CommandLine = UBTExecutable + " " + UBTCommandline(Project: UprojectPath, Target: TargetName, Platform: Platform.ToString(), Config: Config, AdditionalArgs: "-noxge -nobuilduht" + AddArgs);
 
 			Result.XgeXmlFiles = new List<string>();
 			foreach (var XGEFile in FindXGEFiles())
@@ -236,7 +237,7 @@ namespace AutomationTool
 		void XGEFinishBuildWithUBT(XGEItem Item)
 		{
 			// allow the platform to perform any actions after building a target (seems almost like this should be done in UBT)
-			Platform.Platforms[Item.Platform].PostBuildTarget(this, Item.ProjectName, Item.UProjectPath, Item.Config);
+			Platform.Platforms[Item.Platform].PostBuildTarget(this, Item.UProjectPath, Item.TargetName, Item.Config);
 
 			foreach (string ManifestItem in Item.Manifest.BuildProducts)
 			{
@@ -256,12 +257,12 @@ namespace AutomationTool
 			}
 		}
 
-		void CleanWithUBT(string ProjectName, string TargetName, UnrealBuildTool.UnrealTargetPlatform Platform, string Config, string UprojectPath, bool ForceMonolithic = false, bool ForceNonUnity = false, bool ForceDebugInfo = false, string InAddArgs = "", bool ForceUnity = false, Dictionary<string, string> EnvVars = null)
+		void CleanWithUBT(string TargetName, UnrealBuildTool.UnrealTargetPlatform Platform, string Config, FileReference UprojectPath, bool ForceMonolithic = false, bool ForceNonUnity = false, bool ForceDebugInfo = false, string InAddArgs = "", bool ForceUnity = false, Dictionary<string, string> EnvVars = null)
 		{
 			string AddArgs = "";
-			if (string.IsNullOrEmpty(UprojectPath) == false)
+			if (UprojectPath != null)
 			{
-				AddArgs += " " + CommandUtils.MakePathSafeToUseWithCommandLine(UprojectPath);
+				AddArgs += " " + CommandUtils.MakePathSafeToUseWithCommandLine(UprojectPath.FullName);
 			}
 			AddArgs += " " + InAddArgs;
 			if (ForceMonolithic)
@@ -288,16 +289,16 @@ namespace AutomationTool
 			PrepareUBT();
 			using(TelemetryStopwatch CleanStopwatch = new TelemetryStopwatch("CleanWithUBT.{0}.{1}.{2}", TargetName, Platform.ToString(), Config))
 			{
-				RunUBT(CmdEnv, UBTExecutable: UBTExecutable, Project: ProjectName, Target: TargetName, Platform: Platform.ToString(), Config: Config, AdditionalArgs: "-clean" + AddArgs, EnvVars: EnvVars);
+				RunUBT(CmdEnv, UBTExecutable: UBTExecutable, Project: UprojectPath, Target: TargetName, Platform: Platform.ToString(), Config: Config, AdditionalArgs: "-clean" + AddArgs, EnvVars: EnvVars);
 			}
         }
 
-		void BuildWithUBT(string ProjectName, string TargetName, UnrealBuildTool.UnrealTargetPlatform TargetPlatform, string Config, string UprojectPath, bool ForceMonolithic = false, bool ForceNonUnity = false, bool ForceDebugInfo = false, bool ForceFlushMac = false, bool DisableXGE = false, string InAddArgs = "", bool ForceUnity = false, Dictionary<string, string> EnvVars = null)
+		void BuildWithUBT(string TargetName, UnrealBuildTool.UnrealTargetPlatform TargetPlatform, string Config, FileReference UprojectPath, bool ForceMonolithic = false, bool ForceNonUnity = false, bool ForceDebugInfo = false, bool ForceFlushMac = false, bool DisableXGE = false, string InAddArgs = "", bool ForceUnity = false, Dictionary<string, string> EnvVars = null)
 		{
 			string AddArgs = "";
-			if (string.IsNullOrEmpty(UprojectPath) == false)
+			if (UprojectPath != null)
 			{
-				AddArgs += " " + CommandUtils.MakePathSafeToUseWithCommandLine(UprojectPath);
+				AddArgs += " " + CommandUtils.MakePathSafeToUseWithCommandLine(UprojectPath.FullName);
 			}
 			AddArgs += " " + InAddArgs;
 			if (ForceMonolithic)
@@ -337,17 +338,17 @@ namespace AutomationTool
 				DeleteFile(UBTManifest);
 				using(TelemetryStopwatch PrepareManifestStopwatch = new TelemetryStopwatch("PrepareUBTManifest.{0}.{1}.{2}", TargetName, TargetPlatform.ToString(), Config))
 				{
-					RunUBT(CmdEnv, UBTExecutable: UBTExecutable, Project: ProjectName, Target: TargetName, Platform: TargetPlatform.ToString(), Config: Config, AdditionalArgs: AddArgs + " -generatemanifest", EnvVars: EnvVars);
+					RunUBT(CmdEnv, UBTExecutable: UBTExecutable, Project: UprojectPath, Target: TargetName, Platform: TargetPlatform.ToString(), Config: Config, AdditionalArgs: AddArgs + " -generatemanifest", EnvVars: EnvVars);
 				}
                 PrepareManifest(UBTManifest, false);
 			}
 
 			using(TelemetryStopwatch CompileStopwatch = new TelemetryStopwatch("Compile.{0}.{1}.{2}", TargetName, TargetPlatform.ToString(), Config))
 			{
-				RunUBT(CmdEnv, UBTExecutable: UBTExecutable, Project: ProjectName, Target: TargetName, Platform: TargetPlatform.ToString(), Config: Config, AdditionalArgs: AddArgs, EnvVars: EnvVars);
+				RunUBT(CmdEnv, UBTExecutable: UBTExecutable, Project: UprojectPath, Target: TargetName, Platform: TargetPlatform.ToString(), Config: Config, AdditionalArgs: AddArgs, EnvVars: EnvVars);
 			}
             // allow the platform to perform any actions after building a target (seems almost like this should be done in UBT)
-			Platform.Platforms[TargetPlatform].PostBuildTarget(this, string.IsNullOrEmpty(ProjectName) ? TargetName : ProjectName, UprojectPath, Config);
+			Platform.Platforms[TargetPlatform].PostBuildTarget(this, UprojectPath, TargetName, Config);
 
 			if (UseManifest)
 			{
@@ -584,14 +585,11 @@ namespace AutomationTool
 		[DebuggerDisplay("{TargetName} {Platform} {Config}")]
 		public class BuildTarget
 		{
-			/// Unreal project name, if needed
-			public string ProjectName;
-
 			/// Name of the target
 			public string TargetName;
 
 			/// For code-based projects with a .uproject, the TargetName isn't enough for UBT to find the target, this will point UBT to the target
-			public string UprojectPath;
+			public FileReference UprojectPath;
 
 			/// Platform to build
 			public UnrealBuildTool.UnrealTargetPlatform Platform;
@@ -649,7 +647,7 @@ namespace AutomationTool
 			/// <param name="InConfiguration">Configuration</param>
 			/// <param name="InUprojectPath">Path to optional uproject file</param>
 			/// <param name="InAddArgs">Specifies additional arguments for UBT</param>
-			public void AddTarget(string TargetName, UnrealBuildTool.UnrealTargetPlatform InPlatform, UnrealBuildTool.UnrealTargetConfiguration InConfiguration, string InUprojectPath = null, string InAddArgs = "")
+			public void AddTarget(string TargetName, UnrealBuildTool.UnrealTargetPlatform InPlatform, UnrealBuildTool.UnrealTargetConfiguration InConfiguration, FileReference InUprojectPath = null, string InAddArgs = "")
 			{
 				// Is this platform a compilable target?
 				if (!Platform.Platforms[InPlatform].CanBeCompiled())
@@ -675,7 +673,7 @@ namespace AutomationTool
 			/// <param name="InConfiguration">Configuration</param>
 			/// <param name="InUprojectPath">Path to optional uproject file</param>
 			/// <param name="InAddArgs">Specifies additional arguments for UBT</param>
-			public void AddTargets(string[] TargetNames, UnrealBuildTool.UnrealTargetPlatform InPlatform, UnrealBuildTool.UnrealTargetConfiguration InConfiguration, string InUprojectPath = null, string InAddArgs = "")
+			public void AddTargets(string[] TargetNames, UnrealBuildTool.UnrealTargetPlatform InPlatform, UnrealBuildTool.UnrealTargetConfiguration InConfiguration, FileReference InUprojectPath = null, string InAddArgs = "")
 			{
 				// Is this platform a compilable target?
 				if (!Platform.Platforms[InPlatform].CanBeCompiled())
@@ -1343,7 +1341,7 @@ namespace AutomationTool
 				{
 					foreach (var Target in Agenda.Targets)
 					{
-						CleanWithUBT(Target.ProjectName, Target.TargetName, Target.Platform, Target.Config.ToString(), Target.UprojectPath, bForceMonolithic, bForceNonUnity, bForceDebugInfo, Target.UBTArgs, bForceUnity);
+						CleanWithUBT(Target.TargetName, Target.Platform, Target.Config.ToString(), Target.UprojectPath, bForceMonolithic, bForceNonUnity, bForceDebugInfo, Target.UBTArgs, bForceUnity);
 					}
 				}
 				foreach (var Target in Agenda.Targets)
@@ -1353,7 +1351,7 @@ namespace AutomationTool
 						// When building a target for Mac or iOS, use UBT's -flushmac option to clean up the remote builder
 						bool bForceFlushMac = DeleteBuildProducts && (Target.Platform == UnrealBuildTool.UnrealTargetPlatform.Mac || Target.Platform == UnrealBuildTool.UnrealTargetPlatform.IOS);
 						LogSetProgress(InShowProgress, "Building header tool...");
-						BuildWithUBT(Target.ProjectName, Target.TargetName, Target.Platform, Target.Config.ToString(), Target.UprojectPath, bForceMonolithic, bForceNonUnity, bForceDebugInfo, bForceFlushMac, true, Target.UBTArgs, bForceUnity);
+						BuildWithUBT(Target.TargetName, Target.Platform, Target.Config.ToString(), Target.UprojectPath, bForceMonolithic, bForceNonUnity, bForceDebugInfo, bForceFlushMac, true, Target.UBTArgs, bForceUnity);
 					}
 				}
 
@@ -1363,7 +1361,7 @@ namespace AutomationTool
 				{
 					if (Target.TargetName != "UnrealHeaderTool" && CanUseXGE(Target.Platform))
 					{
-						XGEItem Item = XGEPrepareBuildWithUBT(Target.ProjectName, Target.TargetName, Target.Platform, Target.Config.ToString(), Target.UprojectPath, bForceMonolithic, bForceNonUnity, bForceDebugInfo, Target.UBTArgs, bForceUnity);
+						XGEItem Item = XGEPrepareBuildWithUBT(Target.TargetName, Target.Platform, Target.Config.ToString(), Target.UprojectPath, bForceMonolithic, bForceNonUnity, bForceDebugInfo, Target.UBTArgs, bForceUnity);
 						XGEItems.Add(Item);
 					}
 				}
@@ -1408,14 +1406,14 @@ namespace AutomationTool
 				{
 					foreach (var Target in Agenda.Targets)
 					{
-						CleanWithUBT(Target.ProjectName, Target.TargetName, Target.Platform, Target.Config.ToString(), Target.UprojectPath, bForceMonolithic, bForceNonUnity, bForceDebugInfo, Target.UBTArgs, bForceUnity);
+						CleanWithUBT(Target.TargetName, Target.Platform, Target.Config.ToString(), Target.UprojectPath, bForceMonolithic, bForceNonUnity, bForceDebugInfo, Target.UBTArgs, bForceUnity);
 					}
 				}
 				foreach (var Target in Agenda.Targets)
 				{
 					// When building a target for Mac or iOS, use UBT's -flushmac option to clean up the remote builder
 					bool bForceFlushMac = DeleteBuildProducts && (Target.Platform == UnrealBuildTool.UnrealTargetPlatform.Mac || Target.Platform == UnrealBuildTool.UnrealTargetPlatform.IOS);
-					BuildWithUBT(Target.ProjectName, Target.TargetName, Target.Platform, Target.Config.ToString(), Target.UprojectPath, bForceMonolithic, bForceNonUnity, bForceDebugInfo, bForceFlushMac, true, Target.UBTArgs, bForceUnity);
+					BuildWithUBT(Target.TargetName, Target.Platform, Target.Config.ToString(), Target.UprojectPath, bForceMonolithic, bForceNonUnity, bForceDebugInfo, bForceFlushMac, true, Target.UBTArgs, bForceUnity);
 				}
 			}
 
@@ -1583,14 +1581,14 @@ namespace AutomationTool
 			}
 		}
 
-		string GetUBTManifest(string UProjectPath, string InAddArgs)
+		string GetUBTManifest(FileReference UProjectPath, string InAddArgs)
 		{
 			// Can't write to Engine directory on 
             bool bForceRocket = InAddArgs.ToLowerInvariant().Contains(" -rocket"); //awful
 
-            if ((GlobalCommandLine.Rocket || bForceRocket) && !String.IsNullOrEmpty(UProjectPath))
+            if ((GlobalCommandLine.Rocket || bForceRocket) && UProjectPath != null)
 			{
-				return Path.Combine(Path.GetDirectoryName(UProjectPath), "Intermediate/Build/Manifest.xml");				
+				return Path.Combine(Path.GetDirectoryName(UProjectPath.FullName), "Intermediate/Build/Manifest.xml");				
 			}
 			else
 			{
diff --git a/Engine/Source/Programs/AutomationTool/GUBP/LegacyNodes.cs b/Engine/Source/Programs/AutomationTool/GUBP/LegacyNodes.cs
--- a/Engine/Source/Programs/AutomationTool/GUBP/LegacyNodes.cs
+++ b/Engine/Source/Programs/AutomationTool/GUBP/LegacyNodes.cs
@@ -972,8 +972,8 @@ partial class GUBP
 
 		public static bool IsFeaturePack(BranchInfo.BranchUProject InGameProj)
 		{
-			bool bHasContents = CommandUtils.FileExists(CommandUtils.CombinePaths(CommandUtils.GetDirectoryName(InGameProj.FilePath), "contents.txt"));
-			bool bHasManifest = CommandUtils.FileExists(CommandUtils.CombinePaths(CommandUtils.GetDirectoryName(InGameProj.FilePath), "manifest.json"));
+			bool bHasContents = CommandUtils.FileExists(CommandUtils.CombinePaths(CommandUtils.GetDirectoryName(InGameProj.FilePath.FullName), "contents.txt"));
+			bool bHasManifest = CommandUtils.FileExists(CommandUtils.CombinePaths(CommandUtils.GetDirectoryName(InGameProj.FilePath.FullName), "manifest.json"));
 			return bHasContents && bHasManifest;
 		}
 
@@ -1017,7 +1017,7 @@ partial class GUBP
 			BuildProducts = new List<string>();
 			foreach(BranchInfo.BranchUProject Project in Projects)
 			{
-				string ContentsFileName = CommandUtils.CombinePaths(CommandUtils.GetDirectoryName(Project.FilePath), "contents.txt");
+				string ContentsFileName = CommandUtils.CombinePaths(CommandUtils.GetDirectoryName(Project.FilePath.FullName), "contents.txt");
 
 				// Make sure we delete the output file. It may be read-only.
 				string OutputFileName = GetOutputFile(Project);
@@ -1114,7 +1114,7 @@ partial class GUBP
 
 		public static bool IsSample(GUBPBranchConfig BranchConfig, BranchInfo.BranchUProject GameProj)
 		{
-			return (GameProj.GameName != BranchConfig.Branch.BaseEngineProject.GameName && (GameProj.FilePath.Contains("Samples") || GameProj.FilePath.Contains("Templates")));
+			return (GameProj.GameName != BranchConfig.Branch.BaseEngineProject.GameName && (GameProj.FilePath.FullName.Contains("Samples") || GameProj.FilePath.FullName.Contains("Templates")));
 		}
 
 		public override string GetDisplayGroupName()
@@ -1685,7 +1685,7 @@ partial class GUBP
 			{
 				CommandUtils.CookCommandlet(GameProj.FilePath, "UE4Editor-Cmd.exe", null, null, null, null, CookPlatform);
 			}
-            var CookedPath = CombinePaths(Path.GetDirectoryName(GameProj.FilePath), "Saved", "Cooked", CookPlatform);
+            var CookedPath = CombinePaths(Path.GetDirectoryName(GameProj.FilePath.FullName), "Saved", "Cooked", CookPlatform);
             var CookedFiles = CommandUtils.FindFiles("*", true, CookedPath);
             if (CookedFiles.GetLength(0) < 1)
             {
@@ -1992,9 +1992,9 @@ partial class GUBP
         {
             BuildProducts = new List<string>();
             string ProjectArg = "";
-            if (!String.IsNullOrEmpty(GameProj.FilePath))
+            if (GameProj.FilePath != null)
             {
-                ProjectArg = " -project=\"" + GameProj.FilePath + "\"";
+                ProjectArg = " -project=\"" + GameProj.FilePath.FullName + "\"";
             }
             string Args = String.Format("BuildCookRun{0} -SkipBuild -SkipCook -Stage -Pak -Package -NoSubmit", ProjectArg);
 
@@ -2229,7 +2229,7 @@ partial class GUBP
         public override void DoTest(GUBP bp)
         {
             string ProjectArg = "";
-            if (!String.IsNullOrEmpty(GameProj.FilePath))
+            if (GameProj.FilePath != null)
             {
                 ProjectArg = " -project=\"" + GameProj.FilePath + "\"";
             }
@@ -2595,9 +2595,9 @@ partial class GUBP
         public override void DoTest(GUBP bp)
         {
             string ProjectArg = "";
-            if (!String.IsNullOrEmpty(GameProj.FilePath) && UATCommandLine.IndexOf("-project=", StringComparison.InvariantCultureIgnoreCase) < 0)
+            if (GameProj.FilePath != null && UATCommandLine.IndexOf("-project=", StringComparison.InvariantCultureIgnoreCase) < 0)
             {
-                ProjectArg = " -project=\"" + GameProj.FilePath + "\"";
+                ProjectArg = " -project=\"" + GameProj.FilePath.FullName + "\"";
             }
             string WorkingCommandline = UATCommandLine + ProjectArg + " -NoSubmit -addcmdline=\"-DisablePS4TMAPI\"";
             if (WorkingCommandline.Contains("-project=\"\""))
diff --git a/Engine/Source/Programs/AutomationTool/HTML5/HTML5Platform.Automation.cs b/Engine/Source/Programs/AutomationTool/HTML5/HTML5Platform.Automation.cs
--- a/Engine/Source/Programs/AutomationTool/HTML5/HTML5Platform.Automation.cs
+++ b/Engine/Source/Programs/AutomationTool/HTML5/HTML5Platform.Automation.cs
@@ -31,14 +31,14 @@ public class HTML5Platform : Platform
 		HTML5SDKInfo.SetupEmscriptenTemp();
 		HTML5SDKInfo.SetUpEmscriptenConfigFile();
 
-		string PackagePath = Path.Combine(Path.GetDirectoryName(Params.RawProjectPath), "Binaries", "HTML5");
+		string PackagePath = Path.Combine(Path.GetDirectoryName(Params.RawProjectPath.FullName), "Binaries", "HTML5");
 		if (!Directory.Exists(PackagePath))
 		{
 			Directory.CreateDirectory(PackagePath);
 		}
 		string FinalDataLocation = Path.Combine(PackagePath, Params.ShortProjectName) + ".data";
 
-		var ConfigCache = new UnrealBuildTool.ConfigCacheIni(UnrealTargetPlatform.HTML5, "Engine", Path.GetDirectoryName(Params.RawProjectPath), CommandUtils.CombinePaths(CommandUtils.CmdEnv.LocalRoot, "Engine"));
+		var ConfigCache = new UnrealBuildTool.ConfigCacheIni(UnrealTargetPlatform.HTML5, "Engine", Path.GetDirectoryName(Params.RawProjectPath.FullName), CommandUtils.CombinePaths(CommandUtils.CmdEnv.LocalRoot, "Engine"));
 
 		if (HTMLPakAutomation.CanCreateMapPaks(Params))
 		{
@@ -353,7 +353,7 @@ public class HTML5Platform : Platform
 			throw new AutomationException("iOS is currently only able to package one target configuration at a time, but StageTargetConfigurations contained {0} configurations", SC.StageTargetConfigurations.Count);
 		}
 
-		string PackagePath = Path.Combine(Path.GetDirectoryName(Params.RawProjectPath), "Binaries", "HTML5");
+		string PackagePath = Path.Combine(Path.GetDirectoryName(Params.RawProjectPath.FullName), "Binaries", "HTML5");
 		string FinalDataLocation = Path.Combine(PackagePath, Params.ShortProjectName) + ".data";
 
 		// copy the "Executable" to the archive directory
@@ -417,7 +417,7 @@ public class HTML5Platform : Platform
 		// open the webpage
 		Int32 ServerPort = 8000;
 
-		var ConfigCache = new UnrealBuildTool.ConfigCacheIni(UnrealTargetPlatform.HTML5, "Engine", Path.GetDirectoryName(Params.RawProjectPath), CombinePaths(CmdEnv.LocalRoot, "Engine"));
+		var ConfigCache = new UnrealBuildTool.ConfigCacheIni(UnrealTargetPlatform.HTML5, "Engine", Path.GetDirectoryName(Params.RawProjectPath.FullName), CombinePaths(CmdEnv.LocalRoot, "Engine"));
 		ConfigCache.GetInt32("/Script/HTML5PlatformEditor.HTML5TargetSettings", "DeployServerPort", out ServerPort);
 		string WorkingDirectory = Path.GetDirectoryName(ClientApp);
 		string url = Path.GetFileName(ClientApp) +".html";
@@ -533,7 +533,7 @@ public class HTML5Platform : Platform
 #region AMAZON S3 
 	public void UploadToS3(DeploymentContext SC)
 	{
-		ConfigCacheIni Ini = ConfigCacheIni.CreateConfigCacheIni(SC.StageTargetPlatform.PlatformType, "Engine", new DirectoryReference(Path.GetDirectoryName(SC.RawProjectPath)));
+		ConfigCacheIni Ini = ConfigCacheIni.CreateConfigCacheIni(SC.StageTargetPlatform.PlatformType, "Engine", DirectoryReference.FromFile(SC.RawProjectPath));
 		bool Upload = false;
 
 		string KeyId = "";
diff --git a/Engine/Source/Programs/AutomationTool/HTML5/HTML5Platform.PakFiles.Automation.cs b/Engine/Source/Programs/AutomationTool/HTML5/HTML5Platform.PakFiles.Automation.cs
--- a/Engine/Source/Programs/AutomationTool/HTML5/HTML5Platform.PakFiles.Automation.cs
+++ b/Engine/Source/Programs/AutomationTool/HTML5/HTML5Platform.PakFiles.Automation.cs
@@ -23,7 +23,7 @@ public class HTMLPakAutomation
 	public static bool CanCreateMapPaks(ProjectParams Param)
 	{
 		bool UseAsyncLevelLoading = false;
-		var ConfigCache = new UnrealBuildTool.ConfigCacheIni(UnrealTargetPlatform.HTML5, "Engine", Path.GetDirectoryName(Param.RawProjectPath), CommandUtils.CombinePaths(CommandUtils.CmdEnv.LocalRoot, "Engine"));
+		var ConfigCache = new UnrealBuildTool.ConfigCacheIni(UnrealTargetPlatform.HTML5, "Engine", DirectoryReference.FromFile(Param.RawProjectPath), new DirectoryReference(CommandUtils.CombinePaths(CommandUtils.CmdEnv.LocalRoot, "Engine")));
 		ConfigCache.GetBool("/Script/HTML5PlatformEditor.HTML5TargetSettings", "UseAsyncLevelLoading", out UseAsyncLevelLoading);
 
 		if (Param.Run)
@@ -45,7 +45,7 @@ public class HTMLPakAutomation
 			PakOrderFileLocation = CommandUtils.CombinePaths(PakOrderFileLocationBase, "EditorOpenOrder.log");
 		}
 
-		string PakPath = Path.Combine(new string[] { Path.GetDirectoryName(Params.RawProjectPath), "Binaries", "HTML5", SC.ShortProjectName});
+		string PakPath = Path.Combine(new string[] { Path.GetDirectoryName(Params.RawProjectPath.FullName), "Binaries", "HTML5", SC.ShortProjectName});
 		if (Directory.Exists(PakPath))
 		{
 			Directory.Delete(PakPath,true);
@@ -86,7 +86,7 @@ public class HTMLPakAutomation
 	{
 		string StagedDir = Path.Combine(Params.BaseStageDirectory, "HTML5");// CommandUtils.CombinePaths(SC.ProjectRoot, "Saved", "Cooked", "HTML5");
 	
-		 string PakPath = Path.Combine( new string [] { Path.GetDirectoryName(Params.RawProjectPath), "Binaries", "HTML5", SC.ShortProjectName, "Content","Paks"});
+		 string PakPath = Path.Combine( new string [] { Path.GetDirectoryName(Params.RawProjectPath.FullName), "Binaries", "HTML5", SC.ShortProjectName, "Content","Paks"});
 		 if (!Directory.Exists(PakPath))
 		 {
 			 Directory.CreateDirectory(PakPath);
@@ -117,7 +117,7 @@ public class HTMLPakAutomation
 
 		string StagedDir = Path.Combine(Params.BaseStageDirectory, "HTML5");// CommandUtils.CombinePaths(SC.ProjectRoot, "Saved", "Cooked", "HTML5");
 
-		string PakPath = Path.Combine(new string[] { Path.GetDirectoryName(Params.RawProjectPath), "Binaries", "HTML5", Params.ShortProjectName, "Content", "Paks" });
+		string PakPath = Path.Combine(new string[] { Path.GetDirectoryName(Params.RawProjectPath.FullName), "Binaries", "HTML5", Params.ShortProjectName, "Content", "Paks" });
 		if (!Directory.Exists(PakPath))
 		{
 			Directory.CreateDirectory(PakPath);
@@ -307,7 +307,7 @@ public class HTMLPakAutomation
 	{
 		string CookedDir = Path.Combine(Params.BaseStageDirectory, "HTML5");
 
-		string PakPath = Path.Combine(new string[] { Path.GetDirectoryName(Params.RawProjectPath), "Binaries", "HTML5", Params.ShortProjectName, "Content", "Paks" });
+		string PakPath = Path.Combine(new string[] { Path.GetDirectoryName(Params.RawProjectPath.FullName), "Binaries", "HTML5", Params.ShortProjectName, "Content", "Paks" });
 		if (!Directory.Exists(PakPath))
 		{
 			Directory.CreateDirectory(PakPath);
diff --git a/Engine/Source/Programs/AutomationTool/IOS/IOSPlatform.Automation.cs b/Engine/Source/Programs/AutomationTool/IOS/IOSPlatform.Automation.cs
--- a/Engine/Source/Programs/AutomationTool/IOS/IOSPlatform.Automation.cs
+++ b/Engine/Source/Programs/AutomationTool/IOS/IOSPlatform.Automation.cs
@@ -86,7 +86,7 @@ public class IOSPlatform : Platform
 
 	protected string MakeIPAFileName( UnrealTargetConfiguration TargetConfiguration, ProjectParams Params )
 	{
-		string ProjectIPA = Path.Combine(Path.GetDirectoryName(Params.RawProjectPath), "Binaries", "IOS", (Params.Distribution ? "Distro_" : "") + Params.ShortProjectName);
+		string ProjectIPA = Path.Combine(Path.GetDirectoryName(Params.RawProjectPath.FullName), "Binaries", "IOS", (Params.Distribution ? "Distro_" : "") + Params.ShortProjectName);
 		if (TargetConfiguration != UnrealTargetConfiguration.Development)
 		{
 			ProjectIPA += "-" + PlatformType.ToString() + "-" + TargetConfiguration.ToString();
@@ -149,7 +149,7 @@ public class IOSPlatform : Platform
 			var DeployHandler = new UEDeployIOS();
 
 			DeployHandler.PrepForUATPackageOrDeploy(Params.ShortProjectName,
-				Path.GetDirectoryName(Params.RawProjectPath),
+				Path.GetDirectoryName(Params.RawProjectPath.FullName),
 				CombinePaths(Path.GetDirectoryName(Params.ProjectGameExeFilename), SC.StageExecutables[0]),
 				CombinePaths(SC.LocalRoot, "Engine"),
 				Params.Distribution, 
@@ -240,7 +240,7 @@ public class IOSPlatform : Platform
 
 				if (RemoteToolChain.bUseRPCUtil)
 				{
-					string IPPArguments = "RepackageFromStage \"" + (Params.IsCodeBasedProject ? Params.RawProjectPath : "Engine") + "\"";
+					string IPPArguments = "RepackageFromStage \"" + (Params.IsCodeBasedProject ? Params.RawProjectPath.FullName : "Engine") + "\"";
 					IPPArguments += " -config " + TargetConfiguration.ToString();
 
 					if (TargetConfiguration == UnrealTargetConfiguration.Shipping)
@@ -285,7 +285,7 @@ public class IOSPlatform : Platform
 				{
 					List<string> IPPArguments = new List<string>();
 					IPPArguments.Add("RepackageFromStage");
-					IPPArguments.Add(Params.IsCodeBasedProject ? Params.RawProjectPath : "Engine");
+					IPPArguments.Add(Params.IsCodeBasedProject ? Params.RawProjectPath.FullName : "Engine");
 					IPPArguments.Add("-config");
 					IPPArguments.Add(TargetConfiguration.ToString());
 
@@ -317,7 +317,7 @@ public class IOSPlatform : Platform
 					IPPArguments.Add(" -stagedir");
 					IPPArguments.Add(CombinePaths(Params.BaseStageDirectory, "IOS"));
 					IPPArguments.Add(" -project");
-					IPPArguments.Add(Params.RawProjectPath);
+					IPPArguments.Add(Params.RawProjectPath.FullName);
 					if (Params.IterativeDeploy)
 					{
 						IPPArguments.Add(" -iterate");
@@ -359,28 +359,28 @@ public class IOSPlatform : Platform
 		else
 		{
 			// create the ipa
-			string IPAName = CombinePaths(Path.GetDirectoryName(Params.RawProjectPath), "Binaries", "IOS", (Params.Distribution ? "Distro_" : "") + Params.ShortProjectName + (SC.StageTargetConfigurations[0] != UnrealTargetConfiguration.Development ? ("-IOS-" + SC.StageTargetConfigurations[0].ToString()) : "") + ".ipa");
+			string IPAName = CombinePaths(Path.GetDirectoryName(Params.RawProjectPath.FullName), "Binaries", "IOS", (Params.Distribution ? "Distro_" : "") + Params.ShortProjectName + (SC.StageTargetConfigurations[0] != UnrealTargetConfiguration.Development ? ("-IOS-" + SC.StageTargetConfigurations[0].ToString()) : "") + ".ipa");
 
 			if (!Params.IterativeDeploy || !File.Exists(IPAName) || bNeedsIPA)
 			{
 				bCreatedIPA = true;
 
 				// code sign the app
-				CodeSign(Path.GetDirectoryName(Params.ProjectGameExeFilename), Params.IsCodeBasedProject ? Params.ShortProjectName : Path.GetFileNameWithoutExtension(Params.ProjectGameExeFilename), Params.RawProjectPath, SC.StageTargetConfigurations[0], SC.LocalRoot, Params.ShortProjectName, Path.GetDirectoryName(Params.RawProjectPath), SC.IsCodeBasedProject, Params.Distribution, Params.Provision);
+				CodeSign(Path.GetDirectoryName(Params.ProjectGameExeFilename), Params.IsCodeBasedProject ? Params.ShortProjectName : Path.GetFileNameWithoutExtension(Params.ProjectGameExeFilename), Params.RawProjectPath, SC.StageTargetConfigurations[0], SC.LocalRoot, Params.ShortProjectName, Path.GetDirectoryName(Params.RawProjectPath.FullName), SC.IsCodeBasedProject, Params.Distribution, Params.Provision);
 
 				// now generate the ipa
-				PackageIPA(Path.GetDirectoryName(Params.ProjectGameExeFilename), Params.IsCodeBasedProject ? Params.ShortProjectName : Path.GetFileNameWithoutExtension(Params.ProjectGameExeFilename), Params.ShortProjectName, Path.GetDirectoryName(Params.RawProjectPath), SC.StageTargetConfigurations[0], Params.Distribution);
+				PackageIPA(Path.GetDirectoryName(Params.ProjectGameExeFilename), Params.IsCodeBasedProject ? Params.ShortProjectName : Path.GetFileNameWithoutExtension(Params.ProjectGameExeFilename), Params.ShortProjectName, Path.GetDirectoryName(Params.RawProjectPath.FullName), SC.StageTargetConfigurations[0], Params.Distribution);
 			}
 		}
 
 		PrintRunTime();
 	}
 
-	private string EnsureXcodeProjectExists(string RawProjectPath, string LocalRoot, string ShortProjectName, string ProjectRoot, bool IsCodeBasedProject, out bool bWasGenerated)
+	private string EnsureXcodeProjectExists(FileReference RawProjectPath, string LocalRoot, string ShortProjectName, string ProjectRoot, bool IsCodeBasedProject, out bool bWasGenerated)
 	{
 		// first check for ue4.xcodeproj
 		bWasGenerated = false;
-		string XcodeProj = RawProjectPath.Replace(".uproject", "_IOS.xcodeproj");
+		string XcodeProj = RawProjectPath.FullName.Replace(".uproject", "_IOS.xcodeproj");
 		Console.WriteLine ("Project: " + XcodeProj);
 		//		if (!Directory.Exists (XcodeProj))
 		{
@@ -408,7 +408,7 @@ public class IOSPlatform : Platform
 		return XcodeProj;
 	}
 
-	private void CodeSign(string BaseDirectory, string GameName, string RawProjectPath, UnrealTargetConfiguration TargetConfig, string LocalRoot, string ProjectName, string ProjectDirectory, bool IsCode, bool Distribution = false, string Provision = null)
+	private void CodeSign(string BaseDirectory, string GameName, FileReference RawProjectPath, UnrealTargetConfiguration TargetConfig, string LocalRoot, string ProjectName, string ProjectDirectory, bool IsCode, bool Distribution = false, string Provision = null)
 	{
 		// check for the proper xcodeproject
 		bool bWasGenerated = false;
diff --git a/Engine/Source/Programs/AutomationTool/Scripts/BuildCookRun.Automation.cs b/Engine/Source/Programs/AutomationTool/Scripts/BuildCookRun.Automation.cs
--- a/Engine/Source/Programs/AutomationTool/Scripts/BuildCookRun.Automation.cs
+++ b/Engine/Source/Programs/AutomationTool/Scripts/BuildCookRun.Automation.cs
@@ -130,7 +130,7 @@ public class BuildCookRun : BuildCommand
 		const string EngineEntryMap = "/Engine/Maps/Entry";
 		Log("Trying to find DefaultMap in ini files");
 		string DefaultMap = null;
-		var ProjectFolder = GetDirectoryName(Params.RawProjectPath);
+		var ProjectFolder = GetDirectoryName(Params.RawProjectPath.FullName);
 		var DefaultGameEngineConfig = CombinePaths(ProjectFolder, "Config", "DefaultEngine.ini");
 		if (FileExists(DefaultGameEngineConfig))
 		{
@@ -274,8 +274,8 @@ public class BuildCookRun : BuildCommand
 		Ini.Commit();
 	}
 
-	private string ProjectFullPath;
-	public virtual string ProjectPath
+	private FileReference ProjectFullPath;
+	public virtual FileReference ProjectPath
 	{
 		get
 		{
@@ -287,13 +287,13 @@ public class BuildCookRun : BuildCommand
 				{
 					var DestSample = ParseParamValue("DestSample", "CopiedHoverShip");
 					var Dest = ParseParamValue("ForeignDest", CombinePaths(@"C:\testue4\foreign\", DestSample + "_ _Dir"));
-					ProjectFullPath = CombinePaths(Dest, DestSample + ".uproject");
+					ProjectFullPath = new FileReference(CombinePaths(Dest, DestSample + ".uproject"));
 				}
 				else if (bForeignCode)
 				{
 					var DestSample = ParseParamValue("DestSample", "PlatformerGame");
 					var Dest = ParseParamValue("ForeignDest", CombinePaths(@"C:\testue4\foreign\", DestSample + "_ _Dir"));
-					ProjectFullPath = CombinePaths(Dest, DestSample + ".uproject");
+					ProjectFullPath = new FileReference(CombinePaths(Dest, DestSample + ".uproject"));
 				}
 				else
 				{
@@ -308,20 +308,23 @@ public class BuildCookRun : BuildCommand
 					{
 						ProjectName = CombinePaths(CmdEnv.LocalRoot, ProjectName);
 					}
-					if (!FileExists_NoExceptions(ProjectName))
+					if(FileExists_NoExceptions(ProjectName))
+					{
+						ProjectFullPath = new FileReference(ProjectName);
+					}
+					else
 					{
 						var Branch = new BranchInfo(new List<UnrealTargetPlatform> { UnrealBuildTool.BuildHostPlatform.Current.Platform });
 						var GameProj = Branch.FindGame(OriginalProjectName);
 						if (GameProj != null)
 						{
-							ProjectName = GameProj.FilePath;
+							ProjectFullPath = GameProj.FilePath;
+						}
+						if (!FileExists_NoExceptions(ProjectFullPath.FullName))
+						{
+							throw new AutomationException("Could not find a project file {0}.", ProjectName);
 						}
 					}
-					if (!FileExists_NoExceptions(ProjectName))
-					{
-						throw new AutomationException("Could not find a project file {0}.", ProjectName);
-					}
-					ProjectFullPath = Path.GetFullPath(ProjectName);
 				}
 			}
 			return ProjectFullPath;
diff --git a/Engine/Source/Programs/AutomationTool/Scripts/CookCommand.Automation.cs b/Engine/Source/Programs/AutomationTool/Scripts/CookCommand.Automation.cs
--- a/Engine/Source/Programs/AutomationTool/Scripts/CookCommand.Automation.cs
+++ b/Engine/Source/Programs/AutomationTool/Scripts/CookCommand.Automation.cs
@@ -275,7 +275,7 @@ public partial class Project : CommandUtils
     private static void DiffCookedContent( ProjectParams Params)
     {
         List<UnrealTargetPlatform> PlatformsToCook = Params.ClientTargetPlatforms;
-        string ProjectPath = Path.GetFullPath(Params.RawProjectPath);
+        string ProjectPath = Params.RawProjectPath.FullName;
 
         var CookedSandboxesPath = CombinePaths(GetDirectoryName(ProjectPath), "Saved", "Cooked");
 
@@ -395,7 +395,7 @@ public partial class Project : CommandUtils
 
 	private static void CleanupCookedData(List<string> PlatformsToCook, ProjectParams Params)
 	{
-		var ProjectPath = Path.GetFullPath(Params.RawProjectPath);
+		var ProjectPath = Params.RawProjectPath.FullName;
 		var CookedSandboxesPath = CombinePaths(GetDirectoryName(ProjectPath), "Saved", "Cooked");
 		var CleanDirs = new string[PlatformsToCook.Count];
 		for (int DirIndex = 0; DirIndex < CleanDirs.Length; ++DirIndex)
diff --git a/Engine/Source/Programs/AutomationTool/Scripts/CopyBuildToStagingDirectory.Automation.cs b/Engine/Source/Programs/AutomationTool/Scripts/CopyBuildToStagingDirectory.Automation.cs
--- a/Engine/Source/Programs/AutomationTool/Scripts/CopyBuildToStagingDirectory.Automation.cs
+++ b/Engine/Source/Programs/AutomationTool/Scripts/CopyBuildToStagingDirectory.Automation.cs
@@ -308,7 +308,7 @@ public partial class Project : CommandUtils
 		}
 		SC.StageFiles(StagedFileType.NonUFS, GetIntermediateCommandlineDir(SC), CommandLineFile, false, null, "", true, false);
 
-        ConfigCacheIni PlatformGameConfig = ConfigCacheIni.CreateConfigCacheIni(SC.StageTargetPlatform.PlatformType, "Game", new DirectoryReference(CommandUtils.GetDirectoryName(Params.RawProjectPath)));
+        ConfigCacheIni PlatformGameConfig = ConfigCacheIni.CreateConfigCacheIni(SC.StageTargetPlatform.PlatformType, "Game", new DirectoryReference(CommandUtils.GetDirectoryName(Params.RawProjectPath.FullName)));
         var ProjectContentRoot = CombinePaths(SC.ProjectRoot, "Content");
         var StageContentRoot = CombinePaths(SC.RelativeProjectRootForStage, "Content");
         
@@ -389,10 +389,10 @@ public partial class Project : CommandUtils
             }
 			
 			// Engine & Game Plugins. Push the Engine/Source working directory so UBT code has a correct RelativeEnginePath in ReadAvailablePlugins
-			ProjectDescriptor Project = ProjectDescriptor.FromFile(SC.RawProjectPath);
+			ProjectDescriptor Project = ProjectDescriptor.FromFile(SC.RawProjectPath.FullName);
 			LogLog("Searching for plugins with CurrentWorkingDir: " + Directory.GetCurrentDirectory());
 			LogLog("Searching for plugins in: " + SC.RawProjectPath);
-			List<PluginInfo> AvailablePlugins = Plugins.ReadAvailablePlugins(new DirectoryReference(CombinePaths(SC.LocalRoot, "Engine")), new FileReference(SC.RawProjectPath));
+			List<PluginInfo> AvailablePlugins = Plugins.ReadAvailablePlugins(new DirectoryReference(CombinePaths(SC.LocalRoot, "Engine")), SC.RawProjectPath);
 			foreach (PluginInfo Plugin in AvailablePlugins)
 			{
 				LogLog("Considering Plugin for Stage: " + Plugin.File.FullName);
@@ -1096,7 +1096,7 @@ public partial class Project : CommandUtils
 
 	private static string GetTmpPackagingPath(ProjectParams Params, DeploymentContext SC)
 	{
-		return CombinePaths(Path.GetDirectoryName(Params.RawProjectPath), "Saved", "TmpPackaging", SC.StageTargetPlatform.GetCookPlatform(SC.DedicatedServer, false, Params.CookFlavor));
+		return CombinePaths(Path.GetDirectoryName(Params.RawProjectPath.FullName), "Saved", "TmpPackaging", SC.StageTargetPlatform.GetCookPlatform(SC.DedicatedServer, false, Params.CookFlavor));
 	}
 
 	private static bool ShouldCreatePak(ProjectParams Params, DeploymentContext SC)
@@ -1626,7 +1626,7 @@ public partial class Project : CommandUtils
 			{
 				foreach(UnrealTargetConfiguration Config in ConfigsToProcess)
 				{
-					string ReceiptBaseDir = Params.IsCodeBasedProject? Path.GetDirectoryName(Params.RawProjectPath) : EngineDir;
+					string ReceiptBaseDir = Params.IsCodeBasedProject? Path.GetDirectoryName(Params.RawProjectPath.FullName) : EngineDir;
 
 					Platform PlatformInstance = Platform.Platforms[StagePlatform];
 					UnrealTargetPlatform[] SubPlatformsToStage = PlatformInstance.GetStagePlatforms();
@@ -1670,7 +1670,7 @@ public partial class Project : CommandUtils
 						}
 
 						// Convert the paths to absolute
-						Receipt.ExpandPathVariables(new DirectoryReference(EngineDir), new DirectoryReference(Path.GetDirectoryName(Params.RawProjectPath)));
+						Receipt.ExpandPathVariables(new DirectoryReference(EngineDir), Params.RawProjectPath.Directory);
 						TargetsToStage.Add(new StageTarget{ Receipt = Receipt, RequireFilesExist = bRequireStagedFilesToExist });
 					}
 				}
diff --git a/Engine/Source/Programs/AutomationTool/Scripts/ListMobileDevices.Automation.cs b/Engine/Source/Programs/AutomationTool/Scripts/ListMobileDevices.Automation.cs
--- a/Engine/Source/Programs/AutomationTool/Scripts/ListMobileDevices.Automation.cs
+++ b/Engine/Source/Programs/AutomationTool/Scripts/ListMobileDevices.Automation.cs
@@ -12,7 +12,7 @@ namespace AutomationScripts.Automation
 
 			var GlobalParams = new ProjectParams(
 				Command: this,
-				RawProjectPath: @"D:\UE-Main\UE4\Samples\Games\TappyChicken\TappyChicken.uproject"
+				RawProjectPath: new FileReference(@"D:\UE-Main\UE4\Samples\Games\TappyChicken\TappyChicken.uproject")
 				);
 
 			if (ParseParam("android"))
diff --git a/Engine/Source/Programs/AutomationTool/Scripts/RocketBuild.Automation.cs b/Engine/Source/Programs/AutomationTool/Scripts/RocketBuild.Automation.cs
--- a/Engine/Source/Programs/AutomationTool/Scripts/RocketBuild.Automation.cs
+++ b/Engine/Source/Programs/AutomationTool/Scripts/RocketBuild.Automation.cs
@@ -877,7 +877,7 @@ namespace Rocket
 			foreach (string Template in CurrentTemplates)
 			{
 				BranchInfo.BranchUProject Project = BranchConfig.Branch.FindGameChecked(Template);
-				Filter.Include("/" + Utils.StripBaseDirectory(Path.GetDirectoryName(Project.FilePath), CommandUtils.CmdEnv.LocalRoot).Replace('\\', '/') + "/...");
+				Filter.Include("/" + Utils.StripBaseDirectory(Path.GetDirectoryName(Project.FilePath.FullName), CommandUtils.CmdEnv.LocalRoot).Replace('\\', '/') + "/...");
 			}
 
 			// Include all the standard rules
diff --git a/Engine/Source/Programs/AutomationTool/Scripts/RunProjectCommand.Automation.cs b/Engine/Source/Programs/AutomationTool/Scripts/RunProjectCommand.Automation.cs
--- a/Engine/Source/Programs/AutomationTool/Scripts/RunProjectCommand.Automation.cs
+++ b/Engine/Source/Programs/AutomationTool/Scripts/RunProjectCommand.Automation.cs
@@ -1006,11 +1006,11 @@ public partial class Project : CommandUtils
 		return Result;
 	}
 
-	private static ProcessResult RunCookOnTheFlyServer(string ProjectName, string ServerLogFile, string TargetPlatform, string AdditionalCommandLine)
+	private static ProcessResult RunCookOnTheFlyServer(FileReference ProjectName, string ServerLogFile, string TargetPlatform, string AdditionalCommandLine)
 	{
 		var ServerApp = HostPlatform.Current.GetUE4ExePath("UE4Editor.exe");
 		var Args = String.Format("{0} -run=cook -cookonthefly -unattended -CrashForUAT -FORCELOGFLUSH -log",
-			CommandUtils.MakePathSafeToUseWithCommandLine(ProjectName));
+			CommandUtils.MakePathSafeToUseWithCommandLine(ProjectName.FullName));
 		if (!String.IsNullOrEmpty(ServerLogFile))
 		{
 			Args += " -abslog=" + CommandUtils.MakePathSafeToUseWithCommandLine(ServerLogFile);
@@ -1055,7 +1055,7 @@ public partial class Project : CommandUtils
 
 		Log("Running UnrealFileServer *******");
 		var Args = String.Format("{0} -abslog={1} -unattended -CrashForUAT -FORCELOGFLUSH -log {2}",
-						CommandUtils.MakePathSafeToUseWithCommandLine(Params.RawProjectPath),
+						CommandUtils.MakePathSafeToUseWithCommandLine(Params.RawProjectPath.FullName),
 						CommandUtils.MakePathSafeToUseWithCommandLine(ServerLogFile),
 						AdditionalCommandLine);
 		if (IsBuildMachine)
diff --git a/Engine/Source/Programs/AutomationTool/Scripts/UE4BuildUtils.cs b/Engine/Source/Programs/AutomationTool/Scripts/UE4BuildUtils.cs
--- a/Engine/Source/Programs/AutomationTool/Scripts/UE4BuildUtils.cs
+++ b/Engine/Source/Programs/AutomationTool/Scripts/UE4BuildUtils.cs
@@ -21,7 +21,7 @@ public class UE4BuildUtils : CommandUtils
 	{
 		BuildProduct(Command, new UE4Build.BuildTarget()
 			{
-				ProjectName = "",
+				UprojectPath = null,
 				TargetName = "BuildPatchTool",
 				Platform = InPlatform,
 				Config = UnrealBuildTool.UnrealTargetConfiguration.Development,
@@ -37,7 +37,7 @@ public class UE4BuildUtils : CommandUtils
 	{
 		BuildProduct(Command, new UE4Build.BuildTarget()
 		{
-			ProjectName = "",
+			UprojectPath = null,
 			TargetName = "UnrealHeaderTool",
 			Platform = InPlatform,
 			Config = UnrealBuildTool.UnrealTargetConfiguration.Development,
diff --git a/Engine/Source/Programs/UnrealBuildTool/System/FileSystemReference.cs b/Engine/Source/Programs/UnrealBuildTool/System/FileSystemReference.cs
--- a/Engine/Source/Programs/UnrealBuildTool/System/FileSystemReference.cs
+++ b/Engine/Source/Programs/UnrealBuildTool/System/FileSystemReference.cs
@@ -496,6 +496,16 @@ namespace UnrealBuildTool
 		{
 			return new DirectoryReference(AbsolutePath, AbsolutePath.ToLowerInvariant());
 		}
+
+		/// <summary>
+		/// Gets the parent directory for a file, or returns null if it's null.
+		/// </summary>
+		/// <param name="File">The file to create a directory reference for</param>
+		/// <returns>The directory containing the file  </returns>
+		public static DirectoryReference FromFile(FileReference File)
+		{
+			return (File == null)? null : File.Directory;
+		}
 	}
 
 	/// <summary>
diff --git a/Engine/Source/Programs/UnrealBuildTool/System/UnrealBuildTool.cs b/Engine/Source/Programs/UnrealBuildTool/System/UnrealBuildTool.cs
--- a/Engine/Source/Programs/UnrealBuildTool/System/UnrealBuildTool.cs
+++ b/Engine/Source/Programs/UnrealBuildTool/System/UnrealBuildTool.cs
@@ -2784,13 +2784,13 @@ namespace UnrealBuildTool
 		/// Sets up UBT when running from UAT
 		/// </summary>
 		/// <param name="UProjectDir"></param>
-		public static void SetupUBTFromUAT(string UProjectFile)
+		public static void SetupUBTFromUAT(FileReference UProjectFile)
 		{
 			// Reset project file
 			SetProjectFile(null);
 			// when running UAT, the working directory is the root UE4 dir
 			BuildConfiguration.RelativeEnginePath = "Engine";
-			SetProjectFile(new FileReference(UProjectFile));
+			SetProjectFile(UProjectFile);
 		}
 
 
