Merging CL #2807711:

[BuildWatcher] UAT: Downgrading UploadSymbols warning to log message

diff --git a/Engine/Source/Programs/AutomationTool/AutomationUtils/Platform.cs b/Engine/Source/Programs/AutomationTool/AutomationUtils/Platform.cs
--- a/Engine/Source/Programs/AutomationTool/AutomationUtils/Platform.cs
+++ b/Engine/Source/Programs/AutomationTool/AutomationUtils/Platform.cs
@@ -186,7 +186,7 @@ namespace AutomationTool
 
 		public virtual void UploadSymbols(ProjectParams Params, DeploymentContext SC)
 		{
-			LogWarning("{0} does not implement UploadSymbols...", PlatformType);
+			Log("{0} does not implement UploadSymbols...", PlatformType);
 		}
 
 		/// <summary>
