(WIP) Integrating the Blueprint conversion process into UAT and UBT.

diff --git a/Engine/Source/Programs/AutomationTool/AutomationUtils/CommandletUtils.cs b/Engine/Source/Programs/AutomationTool/AutomationUtils/CommandletUtils.cs
--- a/Engine/Source/Programs/AutomationTool/AutomationUtils/CommandletUtils.cs
+++ b/Engine/Source/Programs/AutomationTool/AutomationUtils/CommandletUtils.cs
@@ -161,6 +161,20 @@ namespace AutomationTool
             RunCommandlet(ProjectName, UE4Exe, "UpdateGameProject", Parameters);
         }
 
+        /// <summary>
+        /// Runs the GenerateNativePluginFromBlueprint commandlet.
+        /// </summary>
+        /// <param name="ManifestPath"></param>
+        /// <param name="UE4Exe"></param>
+        /// <param name="Parameters"></param>
+        public static void GenerateNativePluginFromBlueprintCommandlet(FileReference ProjectName, string ManifestPath, string UE4Exe = "UE4Editor-Cmd.exe", string Parameters = "")
+        {
+            Parameters = String.Format("{0} -manifest=\"{1}\"", Parameters, ManifestPath);
+            Parameters.Trim();
+
+            RunCommandlet(ProjectName, UE4Exe, "GenerateNativePluginFromBlueprint", Parameters);
+        }
+
 		/// <summary>
 		/// Runs a commandlet using Engine/Binaries/Win64/UE4Editor-Cmd.exe.
 		/// </summary>
diff --git a/Engine/Source/Programs/AutomationTool/AutomationUtils/ProjectParams.cs b/Engine/Source/Programs/AutomationTool/AutomationUtils/ProjectParams.cs
--- a/Engine/Source/Programs/AutomationTool/AutomationUtils/ProjectParams.cs
+++ b/Engine/Source/Programs/AutomationTool/AutomationUtils/ProjectParams.cs
@@ -990,6 +990,16 @@ namespace AutomationTool
 		[Help("createappbundle", "When archiving for Mac, set this to true to package it in a .app bundle instead of normal loose files")]
 		public bool CreateAppBundle;
 
+        /// <summary>
+        /// Shared: Ref to an auto-generated plugin file that should be incorporated into the project's build
+        /// </summary>
+        public FileReference GeneratedPluginDescFile;
+
+        /// <summary>
+        /// Shared: Ref to a manifest file, detailing the asset files that were converted into native source
+        /// </summary>
+        public FileReference GeneratedPluginManifest;
+
 		#endregion
 
 		#region Build
diff --git a/Engine/Source/Programs/AutomationTool/Scripts/AutomationScripts.Automation.csproj b/Engine/Source/Programs/AutomationTool/Scripts/AutomationScripts.Automation.csproj
--- a/Engine/Source/Programs/AutomationTool/Scripts/AutomationScripts.Automation.csproj
+++ b/Engine/Source/Programs/AutomationTool/Scripts/AutomationScripts.Automation.csproj
@@ -74,6 +74,7 @@
     <Compile Include="FixupRedirects.Automation.cs" />
     <Compile Include="GenerateDSYM.Automation.cs" />
     <Compile Include="RebuildLightMapsCommand.Automation.cs" />
+    <Compile Include="NativizeScriptAssets.Automation.cs" />
     <Compile Include="RocketBuild.Automation.cs" />
     <Compile Include="IPhonePackager.Automation.cs" />
     <Compile Include="LauncherLocalization.Automation.cs" />
diff --git a/Engine/Source/Programs/AutomationTool/Scripts/BuildCookRun.Automation.cs b/Engine/Source/Programs/AutomationTool/Scripts/BuildCookRun.Automation.cs
--- a/Engine/Source/Programs/AutomationTool/Scripts/BuildCookRun.Automation.cs
+++ b/Engine/Source/Programs/AutomationTool/Scripts/BuildCookRun.Automation.cs
@@ -205,7 +205,8 @@ public class BuildCookRun : BuildCommand
 			WorkingCL = P4.CreateChange(P4Env.Client, String.Format("{0} build from changelist {1}", Params.ShortProjectName, P4Env.Changelist));
 		}
 
-		Project.Build(this, Params, WorkingCL);
+        Project.NativizeScriptAssets(this, Params);
+        Project.Build(this, Params, WorkingCL);
 		Project.Cook(Params);
 		Project.CopyBuildToStagingDirectory(Params);
 		Project.Package(Params, WorkingCL);
diff --git a/Engine/Source/Programs/AutomationTool/Scripts/BuildProjectCommand.Automation.cs b/Engine/Source/Programs/AutomationTool/Scripts/BuildProjectCommand.Automation.cs
--- a/Engine/Source/Programs/AutomationTool/Scripts/BuildProjectCommand.Automation.cs
+++ b/Engine/Source/Programs/AutomationTool/Scripts/BuildProjectCommand.Automation.cs
@@ -73,7 +73,16 @@ public partial class Project : CommandUtils
 				foreach (var ClientPlatform in Params.ClientTargetPlatforms)
 				{
 					CrashReportPlatforms.Add(ClientPlatform);
-					Agenda.AddTargets(Params.ClientCookedTargets.ToArray(), ClientPlatform, BuildConfig, Params.CodeBasedUprojectPath);
+
+                    string ExtraUBTArgs = "";
+                    // if we're utilizing an auto-generated code plugin/module 
+                    // (a product of the BuildCookeRun process), make sure to 
+                    // compile it into the project
+                    if (Params.GeneratedPluginDescFile != null && Params.GeneratedPluginDescFile.Exists())
+                    {
+                        ExtraUBTArgs += "-PLUGIN \"" + Params.GeneratedPluginDescFile + "\"";
+                    }
+					Agenda.AddTargets(Params.ClientCookedTargets.ToArray(), ClientPlatform, BuildConfig, Params.CodeBasedUprojectPath, ExtraUBTArgs);
 				}
 			}
 		}
diff --git a/Engine/Source/Programs/AutomationTool/Scripts/CookCommand.Automation.cs b/Engine/Source/Programs/AutomationTool/Scripts/CookCommand.Automation.cs
--- a/Engine/Source/Programs/AutomationTool/Scripts/CookCommand.Automation.cs
+++ b/Engine/Source/Programs/AutomationTool/Scripts/CookCommand.Automation.cs
@@ -215,6 +215,13 @@ public partial class Project : CommandUtils
                 {
                     CommandletParams += " -compressed";
                 }
+                // we provide the option for users to run a conversion on certain (script) assets, translating them 
+                // into native source code... to use those classes over the assets, we need to provide the cooker a
+                // manifest of what all was converted (and what native classes to replace them with)
+                if (Params.GeneratedPluginManifest != null && Params.GeneratedPluginManifest.Exists())
+                {
+                    CommandletParams += " -NativizedAssetManifest=\"" + Params.GeneratedPluginManifest + "\"";
+                }
                 if (Params.HasAdditionalCookerOptions)
                 {
                     string FormatedAdditionalCookerParams = Params.AdditionalCookerOptions.TrimStart(new char[] { '\"', ' ' }).TrimEnd(new char[] { '\"', ' ' });
diff --git a/Engine/Source/Programs/AutomationTool/Scripts/NativizeScriptAssets.Automation.cs b/Engine/Source/Programs/AutomationTool/Scripts/NativizeScriptAssets.Automation.cs
new file mode 100644
--- /dev/null
+++ b/Engine/Source/Programs/AutomationTool/Scripts/NativizeScriptAssets.Automation.cs
@@ -0,0 +1,68 @@
+﻿// Copyright 1998-2015 Epic Games, Inc. All Rights Reserved.
+using System;
+using System.Collections.Generic;
+using System.IO;
+using System.Threading;
+using System.Reflection;
+using System.Linq;
+using AutomationTool;
+using UnrealBuildTool;
+
+/// <summary>
+/// Helper command used for converting script assets (Blueprints) into native (C++) code (aiming to help boost performance of the final executable).
+/// <remarks>
+/// Command line parameters used by this command:
+/// </remarks>
+public partial class Project : CommandUtils
+{
+    #region NativizeScriptAssets Command
+
+    public static void NativizeScriptAssets(BuildCommand Command, ProjectParams Params)
+    {
+        bool bRunNativizeCommand = Command.ParseParam("NativizeAssets");
+        if (!bRunNativizeCommand)
+        {
+            return;
+        }
+        Params.ValidateAndLog();
+
+        Log("********** NATIVIZE-ASSETS COMMAND STARTED **********");
+
+        string UE4EditorExe = HostPlatform.Current.GetUE4ExePath(Params.UE4Exe);
+        if (!FileExists(UE4EditorExe))
+        {
+            throw new AutomationException("Missing " + UE4EditorExe + " executable. Needs to be built first.");
+        }
+
+        string GeneratedPluginName = "NativizedScript";
+        string ProjectDir = Params.RawProjectPath.Directory.ToString();
+        string GeneratedPluginDirectory = Path.GetFullPath(CommandUtils.CombinePaths(ProjectDir, "Intermediate", "Plugins", GeneratedPluginName));
+        string NativizationManifestPath = Path.GetFullPath(CommandUtils.CombinePaths(GeneratedPluginDirectory, GeneratedPluginName + "-Manifest"));
+
+        // delete the product file, so that we can properly detect if the 
+        // commandlet did its job
+        DeleteFile(NativizationManifestPath);
+
+        string CommandletParams = "-output=\"" + GeneratedPluginDirectory + "\" " + "-pluginName=" + GeneratedPluginName;
+        GenerateNativePluginFromBlueprintCommandlet(Params.RawProjectPath, NativizationManifestPath, Params.UE4Exe, CommandletParams);
+
+        string PluginDescFilename  = GeneratedPluginName + ".uplugin";
+        string GeneratedPluginPath = Path.GetFullPath(CommandUtils.CombinePaths(GeneratedPluginDirectory, PluginDescFilename));
+
+        Params.GeneratedPluginDescFile = new FileReference(GeneratedPluginPath);
+        Params.GeneratedPluginManifest = new FileReference(NativizationManifestPath);
+
+        if (!Params.GeneratedPluginDescFile.Exists())
+        {
+            throw new AutomationException("Missing the nativized plugin file (" + GeneratedPluginPath + "). It seems the commandlet didn't generate it as expected.");
+        }
+        if (!Params.GeneratedPluginManifest.Exists())
+        {
+            throw new AutomationException("Missing the nativization manifest file (" + NativizationManifestPath + "). It seems the commandlet didn't generate it as expected.");
+        }
+
+        Log("********** NATIVIZE-ASSETS COMMAND COMPLETED **********");
+    }
+
+    #endregion // NativizeScriptAssets Command
+}
