Merging using UE4-To-UE4-Releases-4.10
Strip passwords from DefaultEngine.ini in pak files
#ue4

diff --git a/Engine/Source/Programs/AutomationTool/Scripts/CopyBuildToStagingDirectory.Automation.cs b/Engine/Source/Programs/AutomationTool/Scripts/CopyBuildToStagingDirectory.Automation.cs
--- a/Engine/Source/Programs/AutomationTool/Scripts/CopyBuildToStagingDirectory.Automation.cs
+++ b/Engine/Source/Programs/AutomationTool/Scripts/CopyBuildToStagingDirectory.Automation.cs
@@ -796,6 +796,15 @@ public partial class Project : CommandUtils
 				}
 			}
 
+			// special case DefaultEngine.ini to strip passwords
+			if (Src.EndsWith("DefaultEngine.ini"))
+			{
+				string NewIniFilename = CombinePaths(SC.ProjectRoot, "Saved", "Temp", SC.PlatformDir, "DefaultEngine.ini");
+				InternalUtils.SafeCreateDirectory(Path.GetDirectoryName(NewIniFilename), true);
+				InternalUtils.SafeCopyFile(Src, NewIniFilename, bFilterSpecialLinesFromIniFiles:true);
+				Src = NewIniFilename;
+			}
+
 			// there can be files that only differ in case only, we don't support that in paks as paks are case-insensitive
 			if (UnrealPakResponseFile.ContainsKey(Src))
 			{
