Remove input dependency on WaitToMakeRocketBuild trigger; we only care about the order of that trigger firing, and don't want to build all its dependencies.

diff --git a/Engine/Source/Programs/AutomationTool/Scripts/RocketBuild.Automation.cs b/Engine/Source/Programs/AutomationTool/Scripts/RocketBuild.Automation.cs
--- a/Engine/Source/Programs/AutomationTool/Scripts/RocketBuild.Automation.cs
+++ b/Engine/Source/Programs/AutomationTool/Scripts/RocketBuild.Automation.cs
@@ -1046,7 +1046,7 @@ namespace Rocket
 
 			AddDependency(FilterRocketNode.StaticGetFullName(HostPlatform));
 			AddDependency(BuildDerivedDataCacheNode.StaticGetFullName(HostPlatform));
-			AddDependency(WaitToMakeRocketBuild.StaticGetFullName());
+			AddPseudodependency(WaitToMakeRocketBuild.StaticGetFullName());
 
 			AgentSharingGroup = "RocketGroup" + StaticGetHostPlatformSuffix(HostPlatform);
 		}
