(WIP) Cleaning up UAT so that it will properly utilizes the Blueprint generated plugin, both for UBT and cook.

diff --git a/Engine/Source/Programs/AutomationTool/AutomationUtils/ProjectParams.cs b/Engine/Source/Programs/AutomationTool/AutomationUtils/ProjectParams.cs
--- a/Engine/Source/Programs/AutomationTool/AutomationUtils/ProjectParams.cs
+++ b/Engine/Source/Programs/AutomationTool/AutomationUtils/ProjectParams.cs
@@ -993,12 +993,15 @@ namespace AutomationTool
         /// <summary>
         /// Shared: Ref to an auto-generated plugin file that should be incorporated into the project's build
         /// </summary>
-        public FileReference GeneratedPluginDescFile;
+        public FileReference NativizedScriptPlugin;
 
         /// <summary>
-        /// Shared: Ref to a manifest file, detailing the asset files that were converted into native source
+        /// Shared: Used to guard against fatal use of the GeneratedScriptPlugin param
         /// </summary>
-        public FileReference GeneratedPluginManifest;
+        public bool UseNativizedScriptPlugin()
+        {
+            return (NativizedScriptPlugin != null) && NativizedScriptPlugin.Exists();
+        }
 
 		#endregion
 
diff --git a/Engine/Source/Programs/AutomationTool/Scripts/BuildProjectCommand.Automation.cs b/Engine/Source/Programs/AutomationTool/Scripts/BuildProjectCommand.Automation.cs
--- a/Engine/Source/Programs/AutomationTool/Scripts/BuildProjectCommand.Automation.cs
+++ b/Engine/Source/Programs/AutomationTool/Scripts/BuildProjectCommand.Automation.cs
@@ -38,6 +38,14 @@ public partial class Project : CommandUtils
 		var Agenda = new UE4Build.BuildAgenda();
 		var CrashReportPlatforms = new HashSet<UnrealTargetPlatform>();
 
+        string ScriptPluginArgs = "";
+        // if we're utilizing an auto-generated code plugin/module (a product of 
+        // the BuildCookeRun process), make sure to compile it along with the targets here
+        if (Params.UseNativizedScriptPlugin())
+        {
+            ScriptPluginArgs = "-PLUGIN \"" + Params.NativizedScriptPlugin + "\"";
+        }
+
 		// Setup editor targets
 		if (Params.HasEditorTargets && !Params.Rocket)
 		{
@@ -46,7 +54,7 @@ public partial class Project : CommandUtils
 			const UnrealTargetConfiguration EditorConfiguration = UnrealTargetConfiguration.Development;
 
 			CrashReportPlatforms.Add(EditorPlatform);
-			Agenda.AddTargets(Params.EditorTargets.ToArray(), EditorPlatform, EditorConfiguration, Params.CodeBasedUprojectPath);
+            Agenda.AddTargets(Params.EditorTargets.ToArray(), EditorPlatform, EditorConfiguration, Params.CodeBasedUprojectPath, ScriptPluginArgs);
 			if (Params.EditorTargets.Contains("UnrealHeaderTool") == false)
 			{
 				Agenda.AddTargets(new string[] { "UnrealHeaderTool" }, EditorPlatform, EditorConfiguration);
@@ -73,16 +81,7 @@ public partial class Project : CommandUtils
 				foreach (var ClientPlatform in Params.ClientTargetPlatforms)
 				{
 					CrashReportPlatforms.Add(ClientPlatform);
-
-                    string ExtraUBTArgs = "";
-                    // if we're utilizing an auto-generated code plugin/module 
-                    // (a product of the BuildCookeRun process), make sure to 
-                    // compile it into the project
-                    if (Params.GeneratedPluginDescFile != null && Params.GeneratedPluginDescFile.Exists())
-                    {
-                        ExtraUBTArgs += "-PLUGIN \"" + Params.GeneratedPluginDescFile + "\"";
-                    }
-					Agenda.AddTargets(Params.ClientCookedTargets.ToArray(), ClientPlatform, BuildConfig, Params.CodeBasedUprojectPath, ExtraUBTArgs);
+                    Agenda.AddTargets(Params.ClientCookedTargets.ToArray(), ClientPlatform, BuildConfig, Params.CodeBasedUprojectPath, ScriptPluginArgs);
 				}
 			}
 		}
@@ -93,7 +92,7 @@ public partial class Project : CommandUtils
 				foreach (var ServerPlatform in Params.ServerTargetPlatforms)
 				{
 					CrashReportPlatforms.Add(ServerPlatform);
-					Agenda.AddTargets(Params.ServerCookedTargets.ToArray(), ServerPlatform, BuildConfig, Params.CodeBasedUprojectPath);
+                    Agenda.AddTargets(Params.ServerCookedTargets.ToArray(), ServerPlatform, BuildConfig, Params.CodeBasedUprojectPath, ScriptPluginArgs);
 				}
 			}
 		}
diff --git a/Engine/Source/Programs/AutomationTool/Scripts/CookCommand.Automation.cs b/Engine/Source/Programs/AutomationTool/Scripts/CookCommand.Automation.cs
--- a/Engine/Source/Programs/AutomationTool/Scripts/CookCommand.Automation.cs
+++ b/Engine/Source/Programs/AutomationTool/Scripts/CookCommand.Automation.cs
@@ -216,11 +216,11 @@ public partial class Project : CommandUtils
                     CommandletParams += " -compressed";
                 }
                 // we provide the option for users to run a conversion on certain (script) assets, translating them 
-                // into native source code... to use those classes over the assets, we need to provide the cooker a
-                // manifest of what all was converted (and what native classes to replace them with)
-                if (Params.GeneratedPluginManifest != null && Params.GeneratedPluginManifest.Exists())
+                // into native source code... to use those classes over the assets, we need to incorporate the 
+                // generated plugin into the cooker
+                if (Params.UseNativizedScriptPlugin())
                 {
-                    CommandletParams += " -NativizedAssetManifest=\"" + Params.GeneratedPluginManifest + "\"";
+                    CommandletParams += " -PLUGIN=\"" + Params.NativizedScriptPlugin + "\"";
                 }
                 if (Params.HasAdditionalCookerOptions)
                 {
diff --git a/Engine/Source/Programs/AutomationTool/Scripts/NativizeScriptAssets.Automation.cs b/Engine/Source/Programs/AutomationTool/Scripts/NativizeScriptAssets.Automation.cs
--- a/Engine/Source/Programs/AutomationTool/Scripts/NativizeScriptAssets.Automation.cs
+++ b/Engine/Source/Programs/AutomationTool/Scripts/NativizeScriptAssets.Automation.cs
@@ -49,17 +49,11 @@ public partial class Project : CommandUtils
         string PluginDescFilename  = GeneratedPluginName + ".uplugin";
         string GeneratedPluginPath = Path.GetFullPath(CommandUtils.CombinePaths(GeneratedPluginDirectory, PluginDescFilename));
 
-        Params.GeneratedPluginDescFile = new FileReference(GeneratedPluginPath);
-        Params.GeneratedPluginManifest = new FileReference(NativizationManifestPath);
-
-        if (!Params.GeneratedPluginDescFile.Exists())
+        Params.NativizedScriptPlugin = new FileReference(GeneratedPluginPath);
+        if (!Params.UseNativizedScriptPlugin())
         {
             throw new AutomationException("Missing the nativized plugin file (" + GeneratedPluginPath + "). It seems the commandlet didn't generate it as expected.");
         }
-        if (!Params.GeneratedPluginManifest.Exists())
-        {
-            throw new AutomationException("Missing the nativization manifest file (" + NativizationManifestPath + "). It seems the commandlet didn't generate it as expected.");
-        }
 
         Log("********** NATIVIZE-ASSETS COMMAND COMPLETED **********");
     }
