fix for merge error

diff --git a/Engine/Source/Programs/AutomationTool/AutomationUtils/ProjectParams.cs b/Engine/Source/Programs/AutomationTool/AutomationUtils/ProjectParams.cs
--- a/Engine/Source/Programs/AutomationTool/AutomationUtils/ProjectParams.cs
+++ b/Engine/Source/Programs/AutomationTool/AutomationUtils/ProjectParams.cs
@@ -992,15 +992,18 @@ namespace AutomationTool
 		[Help("createappbundle", "When archiving for Mac, set this to true to package it in a .app bundle instead of normal loose files")]
 		public bool CreateAppBundle;
 
-        /// <summary>
-        /// Shared: Ref to an auto-generated plugin file that should be incorporated into the project's build
-        /// </summary>
-        public FileReference GeneratedPluginDescFile;
+		/// <summary>
+		/// Shared: Ref to an auto-generated plugin file that should be incorporated into the project's build
+		/// </summary>
+		public FileReference NativizedScriptPlugin;
 
 		/// <summary>
-		/// Shared: Ref to a manifest file, detailing the asset files that were converted into native source
+		/// Shared: Used to guard against fatal use of the GeneratedScriptPlugin param
 		/// </summary>
-		public FileReference GeneratedPluginManifest;
+		public bool UseNativizedScriptPlugin()
+		{
+			return (NativizedScriptPlugin != null) && NativizedScriptPlugin.Exists();
+		}
 
 		#endregion
 
