Merging fix for failure to sign on mac from 4.10

diff --git a/Engine/Source/Programs/AutomationTool/IOS/IOSPlatform.Automation.cs b/Engine/Source/Programs/AutomationTool/IOS/IOSPlatform.Automation.cs
--- a/Engine/Source/Programs/AutomationTool/IOS/IOSPlatform.Automation.cs
+++ b/Engine/Source/Programs/AutomationTool/IOS/IOSPlatform.Automation.cs
@@ -283,11 +283,11 @@ public class IOSPlatform : Platform
 					{
 						IPPArguments += " -iterate";
 					}
-					if (Params.Provision != null)
+					if (!string.IsNullOrEmpty(Params.Provision))
 					{
 						IPPArguments += " -provision \"" + Params.Provision + "\""; 
 					}
-					if (Params.Certificate != null)
+					if (!string.IsNullOrEmpty(Params.Certificate))
 					{
 						IPPArguments += " -certificate \"" + Params.Certificate + "\"";
 					}
@@ -335,12 +335,12 @@ public class IOSPlatform : Platform
 					{
 						IPPArguments.Add(" -iterate");
 					}
-					if (Params.Provision != null)
+					if (!string.IsNullOrEmpty(Params.Provision))
 					{
 						IPPArguments.Add(" -provision");
 						IPPArguments.Add(Params.Provision);
 					}
-					if (Params.Certificate != null)
+					if (!string.IsNullOrEmpty(Params.Certificate))
 					{
 						IPPArguments.Add(" -certificate");
 						IPPArguments.Add(Params.Certificate);
@@ -439,15 +439,15 @@ public class IOSPlatform : Platform
 		Arguments += " - iOS'";
 		Arguments += " -configuration " + TargetConfig.ToString();
 		Arguments += " -sdk iphoneos";
-		if (Certificate != null)
+		if (!string.IsNullOrEmpty(Certificate))
 		{
-			Arguments += " CODE_SIGN_IDENTITY=" + Certificate;
+			Arguments += " CODE_SIGN_IDENTITY=\"" + Certificate + "\"";
 		}
 		else
 		{
 			Arguments += " CODE_SIGN_IDENTITY=" + (Distribution ? "\"iPhone Distribution\"" : "\"iPhone Developer\"");
 		}
-		if (Provision != null)
+		if (!string.IsNullOrEmpty(Provision))
 		{
 			// read the provision to get the UUID
 			if (File.Exists(Environment.GetEnvironmentVariable("HOME") + "/Library/MobileDevice/Provisioning Profiles/" + Provision))
