Merge main stabilization changes back to //depot/UE4

diff --git a/Engine/Build/BatchFiles/GenerateProjectFiles.bat b/Engine/Build/BatchFiles/GenerateProjectFiles.bat
--- a/Engine/Build/BatchFiles/GenerateProjectFiles.bat
+++ b/Engine/Build/BatchFiles/GenerateProjectFiles.bat
@@ -40,28 +40,30 @@ if exist "%ProgramFiles(x86)%\MSBuild\12.0\bin\MSBuild.exe" (
 	goto ReadyToCompile
 )
 
-rem ## Check for Visual Studio 2015
+rem ## Check for Visual Studio 2013
+
+for %%P in (%*) do if "%%P" == "-2015" goto NoVisualStudio2013Environment
 
 pushd %~dp0
-call GetVSComnToolsPath 14
+call GetVSComnToolsPath 12
 popd
 
 if "%VsComnToolsPath%" == "" goto NoVisualStudio2013Environment
 call "%VsComnToolsPath%/../../VC/bin/x86_amd64/vcvarsx86_amd64.bat" >NUL
 goto ReadyToCompile
 
-:NoVisualStudio2015Environment
-rem ## Check for Visual Studio 2013
+:NoVisualStudio2013Environment
+rem ## Check for Visual Studio 2015
 
 pushd %~dp0
-call GetVSComnToolsPath 12
+call GetVSComnToolsPath 14
 popd
 
-if "%VsComnToolsPath%" == "" goto NoVisualStudio2013Environment
+if "%VsComnToolsPath%" == "" goto NoVisualStudio2015Environment
 call "%VsComnToolsPath%/../../VC/bin/x86_amd64/vcvarsx86_amd64.bat" >NUL
 goto ReadyToCompile
 
-:NoVisualStudio2013Environment
+:NoVisualStudio2015Environment
 rem ## Check for Visual Studio 2012
 
 pushd %~dp0
diff --git a/Engine/Build/BatchFiles/RunUAT.bat b/Engine/Build/BatchFiles/RunUAT.bat
--- a/Engine/Build/BatchFiles/RunUAT.bat
+++ b/Engine/Build/BatchFiles/RunUAT.bat
@@ -29,31 +29,33 @@ if not exist Source\Programs\AutomationToolLauncher\AutomationToolLauncher.cspro
 rem ## Check to see if we're already running under a Visual Studio environment shell
 if not "%INCLUDE%" == "" if not "%LIB%" == "" goto ReadyToCompile
 
-rem ## Check for Visual Studio 2015
+rem ## Check for Visual Studio 2013
+
+for %%P in (%*) do if "%%P" == "-2015" goto NoVisualStudio2013Environment
 
 pushd %~dp0
-call GetVSComnToolsPath 14
+call GetVSComnToolsPath 12
 popd
 
-if "%VsComnToolsPath%" == "" goto NoVisualStudio2015Environment
+if "%VsComnToolsPath%" == "" goto NoVisualStudio2013Environment
 call "%VsComnToolsPath%/../../VC/bin/x86_amd64/vcvarsx86_amd64.bat" >NUL
 goto ReadyToCompile
 
 
-rem ## Check for Visual Studio 2013
-:NoVisualStudio2015Environment
+rem ## Check for Visual Studio 2015
+:NoVisualStudio2013Environment
 
 pushd %~dp0
-call GetVSComnToolsPath 12
+call GetVSComnToolsPath 14
 popd
 
-if "%VsComnToolsPath%" == "" goto NoVisualStudio2013Environment
+if "%VsComnToolsPath%" == "" goto NoVisualStudio2015Environment
 call "%VsComnToolsPath%/../../VC/bin/x86_amd64/vcvarsx86_amd64.bat" >NUL
 goto ReadyToCompile
 
 
 rem ## Check for Visual Studio 2012
-:NoVisualStudio2013Environment
+:NoVisualStudio2015Environment
 
 pushd %~dp0
 call GetVSComnToolsPath 11
diff --git a/Engine/Shaders/BasePassPixelShader.usf b/Engine/Shaders/BasePassPixelShader.usf
--- a/Engine/Shaders/BasePassPixelShader.usf
+++ b/Engine/Shaders/BasePassPixelShader.usf
@@ -875,10 +875,7 @@ void Main(
 		GBuffer.IndirectIrradiance = IndirectIrradiance;
 		GBuffer.PrecomputedShadowFactors = PrecomputedShadowFactors;
 		GBuffer.GBufferAO = MaterialAO;
-		GBuffer.Opacity = Opacity;
-		GBuffer.DecalMask = Primitive.DecalReceiverMask;
-		GBuffer.HasDistanceFieldRepresentation = Primitive.HasDistanceFieldRepresentation;
-		GBuffer.HasHeightfieldRepresentation = Primitive.HasHeightfieldRepresentation;
+		GBuffer.PerObjectGBufferData = Primitive.PerObjectGBufferData;
 
 		#if MATERIAL_SHADINGMODEL_UNLIT
 			GBuffer.ShadingModelID = SHADINGMODELID_UNLIT;
@@ -886,43 +883,62 @@ void Main(
 			GBuffer.ShadingModelID = SHADINGMODELID_DEFAULT_LIT;
 		#elif MATERIAL_SHADINGMODEL_SUBSURFACE
 			GBuffer.ShadingModelID = SHADINGMODELID_SUBSURFACE;
-			GBuffer.CustomData = EncodeSubsurfaceColor(SubsurfaceColor);
+		GBuffer.CustomData.rgb = EncodeSubsurfaceColor(SubsurfaceColor);
+		GBuffer.CustomData.a = Opacity;
 		#elif MATERIAL_SHADINGMODEL_PREINTEGRATED_SKIN
 			GBuffer.ShadingModelID = SHADINGMODELID_PREINTEGRATED_SKIN;
-			GBuffer.CustomData = EncodeSubsurfaceColor(SubsurfaceColor);
+		GBuffer.CustomData.rgb = EncodeSubsurfaceColor(SubsurfaceColor);
+		GBuffer.CustomData.a = Opacity;
 		#elif MATERIAL_SHADINGMODEL_SUBSURFACE_PROFILE
 			GBuffer.ShadingModelID = SHADINGMODELID_SUBSURFACE_PROFILE;
-			GBuffer.CustomData = EncodeSubsurfaceProfile(SubsurfaceProfile);
+		GBuffer.CustomData.rgb = EncodeSubsurfaceProfile(SubsurfaceProfile);
+		GBuffer.CustomData.a = Opacity;
 		#elif MATERIAL_SHADINGMODEL_CLEAR_COAT
 			GBuffer.ShadingModelID = SHADINGMODELID_CLEAR_COAT;
 
-			float ClearCoat				= GetMaterialClearCoat(MaterialParameters);
-			float ClearCoatRoughness	= GetMaterialClearCoatRoughness(MaterialParameters);
-			float Film = 1 * ClearCoat;
+			float ClearCoat				= saturate( GetMaterialCustomData0(MaterialParameters) );
+			float ClearCoatRoughness	= saturate( GetMaterialCustomData1(MaterialParameters) );
 			float MetalSpec = 0.9;
 
 		float NoV = saturate( dot( MaterialParameters.WorldNormal, MaterialParameters.CameraVector ) );
 
 			// Approximation of refraction's effect on EnvBRDF
 			float RefractionScale = ( (NoV * 0.5 + 0.5) * NoV - 1 ) * saturate( 1.25 - 1.25 * Roughness ) + 1;
-			RefractionScale = lerp( 1, RefractionScale, ClearCoat );
-
-			float3 AbsorptionColor = (1 - Film) + BaseColor * ( Film * (1 / MetalSpec) );
-			GBuffer.BaseColor = lerp( BaseColor, MetalSpec, Film );
 
 			// Approximation of absorption integral, tuned for Roughness=0.4
+			float3 AbsorptionColor = BaseColor * (1 / MetalSpec);
 			float3 Absorption = AbsorptionColor * ( (NoV - 1) * 0.85 * ( 1 - lerp( AbsorptionColor, Square(AbsorptionColor), -0.78 ) ) + 1 );
-			//float Toe = 1 - exp2( -9 * NoV - 3.5 );
-			//float3 Absorption = Toe * AbsorptionColor * ( (NoV - 1) * 0.85 * ( 1 - lerp( AbsorptionColor, Square(AbsorptionColor), -0.78 ) ) + lerp( 0.93, 0.96, AbsorptionColor * 2 - 1 ) );
 
-			GBuffer.BaseColor *= Absorption * lerp( 1, RefractionScale, Metallic );
-			GBuffer.Specular *= Luminance( Absorption ) * RefractionScale;
+			GBuffer.BaseColor = lerp( BaseColor, lerp( BaseColor, MetalSpec * Absorption, Metallic ) * RefractionScale, ClearCoat );
+			GBuffer.Specular *= lerp( 1, RefractionScale, ClearCoat );
 
 			GBuffer.CustomData.x = ClearCoat;
 			GBuffer.CustomData.y = ClearCoatRoughness;
 		#elif MATERIAL_SHADINGMODEL_TWOSIDED_FOLIAGE
 			GBuffer.ShadingModelID = SHADINGMODELID_TWOSIDED_FOLIAGE;
-			GBuffer.CustomData = EncodeSubsurfaceColor(SubsurfaceColor);
+			GBuffer.CustomData.rgb = EncodeSubsurfaceColor(SubsurfaceColor);
+			GBuffer.CustomData.a = Opacity;
+		#elif MATERIAL_SHADINGMODEL_HAIR
+			GBuffer.ShadingModelID = SHADINGMODELID_HAIR;
+			#if NUM_MATERIAL_OUTPUTS_GETTANGENTOUTPUT > 0
+				float3 Tangent = GetTangentOutput0(MaterialParameters);
+				Tangent = normalize( TransformTangentVectorToWorld( MaterialParameters.TangentToWorld, Tangent ) );
+				GBuffer.WorldNormal = Tangent;
+			#endif
+			GBuffer.CustomData.xy = UnitVectorToOctahedron( MaterialParameters.WorldNormal ) * 0.5 + 0.5;
+			GBuffer.CustomData.z = saturate( GetMaterialCustomData1(MaterialParameters) );	// Roughness2
+			GBuffer.CustomData.w = saturate( GetMaterialCustomData0(MaterialParameters) );	// Opacity
+		#else
+			// missing shading model, compiler should report ShadingModelID is not set
+		#endif
+
+		GBuffer.SelectiveOutputMask = GetSelectiveOutputMask();
+
+	#if WRITES_VELOCITY_TO_GBUFFER
+		{
+			// 2d velocity, includes camera an object motion
+		#if WRITES_VELOCITY_TO_GBUFFER_USE_POS_INTERPOLATOR
+			float2 Velocity = Calculate2DVelocity(BasePassInterpolants.VelocityScreenPosition, BasePassInterpolants.VelocityPrevScreenPosition);
 		#else
 			float2 Velocity = Calculate2DVelocity(MaterialParameters.ScreenPosition, BasePassInterpolants.VelocityPrevScreenPosition);
 		#endif
@@ -934,7 +950,8 @@ void Main(
 		GBuffer.Velocity = 0;
 	#endif
 
-		float QuantizationBias = PseudoRandom( MaterialParameters.SVPosition.xy ) * 0.5 - 0.5;
+		// -0.5 .. 0.5, could be optimzed as lower quality noise would be sufficient
+		float QuantizationBias = PseudoRandom( MaterialParameters.SvPosition.xy ) - 0.5f;
 		EncodeGBuffer(GBuffer, OutGBufferA, OutGBufferB, OutGBufferC, OutGBufferD, OutGBufferE, OutGBufferVelocity, QuantizationBias);
 	#endif 
 
diff --git a/Engine/Source/Editor/HierarchicalLODOutliner/Private/HLODTreeWidgetItem.cpp b/Engine/Source/Editor/HierarchicalLODOutliner/Private/HLODTreeWidgetItem.cpp
--- a/Engine/Source/Editor/HierarchicalLODOutliner/Private/HLODTreeWidgetItem.cpp
+++ b/Engine/Source/Editor/HierarchicalLODOutliner/Private/HLODTreeWidgetItem.cpp
@@ -173,7 +173,7 @@ END_SLATE_FUNCTION_BUILD_OPTIMIZATION
 		{
 			FDragValidationInfo ValidationInfo = FDragValidationInfo::Invalid();
 
-			HLODOutliner::HandleDrop(WeakTableViewBase, DragDropEvent, *TreeItem, ValidationInfo, false);
+			HLODOutliner::HandleDrop(WeakTableViewBase, DragDropEvent, *TreeItem, ValidationInfo, this, false);
 			UpdateOperationDecorator(DragDropEvent, ValidationInfo);
 		}
 	}
diff --git a/Engine/Source/Editor/KismetCompiler/Private/AnimBlueprintCompiler.cpp b/Engine/Source/Editor/KismetCompiler/Private/AnimBlueprintCompiler.cpp
--- a/Engine/Source/Editor/KismetCompiler/Private/AnimBlueprintCompiler.cpp
+++ b/Engine/Source/Editor/KismetCompiler/Private/AnimBlueprintCompiler.cpp
@@ -1950,9 +1950,9 @@ void FAnimBlueprintCompiler::FEvaluationHandlerRecord::RegisterPin(UEdGraphPin*
 
 						if(bLeafNode)
 						{
-							if(UK2Node_VariableGet* VariableGetNode = Cast<UK2Node_VariableGet>(LinkedNode))
+							if(UK2Node_VariableGet* LinkedVariableGetNode = Cast<UK2Node_VariableGet>(LinkedNode))
 							{
-								if(!VariableGetNode->IsNodePure() || !VariableGetNode->VariableReference.IsSelfContext())
+								if(!LinkedVariableGetNode->IsNodePure() || !LinkedVariableGetNode->VariableReference.IsSelfContext())
 								{
 									// only local variable access is allowed for leaf nodes, so 
 									Handler.bHasOnlyMemberAccess = false;
diff --git a/Engine/Source/Editor/KismetCompiler/Private/KismetCompilerMisc.cpp b/Engine/Source/Editor/KismetCompiler/Private/KismetCompilerMisc.cpp
--- a/Engine/Source/Editor/KismetCompiler/Private/KismetCompilerMisc.cpp
+++ b/Engine/Source/Editor/KismetCompiler/Private/KismetCompilerMisc.cpp
@@ -505,10 +505,12 @@ void FKismetCompilerUtilities::ValidateEnumProperties(UObject* DefaultObject, FC
 				{
 					MessageLog.Warning(
 						*FString::Printf(
-							*LOCTEXT("InvalidEnumDefaultValue_Error", "Default Enum value '%s' for class '%s' is invalid in object '%s' ").ToString(),
+							*LOCTEXT("InvalidEnumDefaultValue_Error", "Default Enum value '%s' for class '%s' is invalid in object '%s'. EnumVal: %d. EnumAcceptableMax: %d ").ToString(),
 							*ByteProperty->GetName(),
 							*DefaultObject->GetClass()->GetName(),
-							*DefaultObject->GetName()
+							*DefaultObject->GetName(),
+							EnumValue,
+							Enum->GetMaxEnumValue()
 						)
 					);
 				}
diff --git a/Engine/Source/Editor/UnrealEd/Private/Commandlets/ChunkManifestGenerator.cpp b/Engine/Source/Editor/UnrealEd/Private/Commandlets/ChunkManifestGenerator.cpp
--- a/Engine/Source/Editor/UnrealEd/Private/Commandlets/ChunkManifestGenerator.cpp
+++ b/Engine/Source/Editor/UnrealEd/Private/Commandlets/ChunkManifestGenerator.cpp
@@ -976,9 +976,9 @@ bool FChunkManifestGenerator::GetPackageDependencyChain(FName SourcePackage, FNa
 	return false;
 }
 
-bool FChunkManifestGenerator::GetPackageDependencies(FName PackageName, TArray<FName>& DependentPackageNames, EAssetRegistryDependencyType::Type DependencyType)
+bool FChunkManifestGenerator::GetPackageDependencies(FName PackageName, TArray<FName>& DependentPackageNames, EAssetRegistryDependencyType::Type InDependencyType)
 {	
-	return AssetRegistry.GetDependencies(PackageName, DependentPackageNames, DependencyType);
+	return AssetRegistry.GetDependencies(PackageName, DependentPackageNames, InDependencyType);
 }
 
 bool FChunkManifestGenerator::GatherAllPackageDependencies(FName PackageName, TArray<FName>& DependentPackageNames)
diff --git a/Engine/Source/Editor/UnrealEd/Private/StaticLightingSystem/StaticLightingSystem.cpp b/Engine/Source/Editor/UnrealEd/Private/StaticLightingSystem/StaticLightingSystem.cpp
--- a/Engine/Source/Editor/UnrealEd/Private/StaticLightingSystem/StaticLightingSystem.cpp
+++ b/Engine/Source/Editor/UnrealEd/Private/StaticLightingSystem/StaticLightingSystem.cpp
@@ -768,9 +768,9 @@ void UpdateStaticLightingHLODTreeIndices(TMultiMap<AActor*, FStaticLightingMesh*
 
 	for (AActor* SubActor : LODActor->SubActors)
 	{
-		if (ALODActor* LODActor = Cast<ALODActor>(SubActor))
+		if (ALODActor* LODSubActor = Cast<ALODActor>(SubActor))
 		{
-			UpdateStaticLightingHLODTreeIndices(ActorMeshMap, LODActor, HLODTreeIndex, HLODLeafIndex);
+			UpdateStaticLightingHLODTreeIndices(ActorMeshMap, LODSubActor, HLODTreeIndex, HLODLeafIndex);
 		}
 		else
 		{
diff --git a/Engine/Source/Programs/AutomationTool/AutomationUtils/CommandUtils.cs b/Engine/Source/Programs/AutomationTool/AutomationUtils/CommandUtils.cs
--- a/Engine/Source/Programs/AutomationTool/AutomationUtils/CommandUtils.cs
+++ b/Engine/Source/Programs/AutomationTool/AutomationUtils/CommandUtils.cs
@@ -1910,6 +1910,13 @@ namespace AutomationTool
             return CachedRootBuildStorageDirectory;
         }
 
+		// Included for compatibility during //UE4/Main import
+		[Obsolete]
+		public static string RootSharedTempStorageDirectory()
+		{
+			return RootBuildStorageDirectory();
+		}
+
         public static bool DirectoryExistsAndIsWritable_NoExceptions(string Dir)
         {
             if (!DirectoryExists_NoExceptions(Dir))
diff --git a/Engine/Source/Programs/AutomationTool/AutomationUtils/CommandletUtils.cs b/Engine/Source/Programs/AutomationTool/AutomationUtils/CommandletUtils.cs
--- a/Engine/Source/Programs/AutomationTool/AutomationUtils/CommandletUtils.cs
+++ b/Engine/Source/Programs/AutomationTool/AutomationUtils/CommandletUtils.cs
@@ -13,6 +13,24 @@ namespace AutomationTool
 	{
 		#region Commandlets
 
+		/// <summary>
+		/// Runs Cook commandlet.
+		/// </summary>
+		[Obsolete("CookCommandlet now takes a FileReference for project argument")]
+		public static void CookCommandlet(string ProjectName)
+		{
+			FileReference ProjectFile;
+			if(ProjectName.Contains('/') || ProjectName.Contains('\\'))
+			{
+				ProjectFile = new FileReference(ProjectName);
+			}
+			else if(!UProjectInfo.TryGetProjectFileName(ProjectName, out ProjectFile))
+			{
+				throw new AutomationException("Cannot determine project path for {0}", ProjectName);
+			}
+			CookCommandlet(ProjectFile);
+		}
+
 		/// <summary>
 		/// Runs Cook commandlet.
 		/// </summary>
@@ -182,6 +200,24 @@ namespace AutomationTool
 		/// <param name="UE4Exe">The name of the UE4 Editor executable to use.</param>
 		/// <param name="Commandlet">Commandlet name.</param>
 		/// <param name="Parameters">Command line parameters (without -run=)</param>
+		[Obsolete("RunCommandlet now takes a uproject path as first argument")]
+		public static void RunCommandlet(string ProjectName, string UE4Exe, string Commandlet, string Parameters = null)
+		{
+			FileReference ProjectFile;
+			if(!UProjectInfo.TryGetProjectFileName(ProjectName, out ProjectFile))
+			{
+				throw new AutomationException("Cannot determine project file for {0}", ProjectName);
+			}
+			RunCommandlet(ProjectFile, UE4Exe, Commandlet, Parameters);
+		}
+
+		/// <summary>
+		/// Runs a commandlet using Engine/Binaries/Win64/UE4Editor-Cmd.exe.
+		/// </summary>
+		/// <param name="ProjectFile">Project name.</param>
+		/// <param name="UE4Exe">The name of the UE4 Editor executable to use.</param>
+		/// <param name="Commandlet">Commandlet name.</param>
+		/// <param name="Parameters">Command line parameters (without -run=)</param>
 		public static void RunCommandlet(FileReference ProjectName, string UE4Exe, string Commandlet, string Parameters = null)
 		{
 			Log("Running UE4Editor {0} for project {1}", Commandlet, ProjectName);
diff --git a/Engine/Source/Programs/AutomationTool/AutomationUtils/ProjectParams.cs b/Engine/Source/Programs/AutomationTool/AutomationUtils/ProjectParams.cs
--- a/Engine/Source/Programs/AutomationTool/AutomationUtils/ProjectParams.cs
+++ b/Engine/Source/Programs/AutomationTool/AutomationUtils/ProjectParams.cs
@@ -781,6 +781,220 @@ namespace AutomationTool
 			ValidateAndLog();
 		}
 
+		// Included for compatibility during //UE4/Main import
+		[Obsolete]
+		public ProjectParams(			
+			string RawProjectPath,
+
+			CommandUtils Command = null,
+			string Device = null,			
+			string MapToRun = null,	
+			string AdditionalServerMapParams = null,
+			ParamList<string> Port = null,
+			string RunCommandline = null,						
+			string StageCommandline = null,
+            string BundleName = null,
+            string StageDirectoryParam = null,
+			string UE4Exe = null,
+			string SignPak = null,
+			List<UnrealTargetConfiguration> ClientConfigsToBuild = null,
+			List<UnrealTargetConfiguration> ServerConfigsToBuild = null,
+			ParamList<string> MapsToCook = null,
+			ParamList<string> DirectoriesToCook = null,
+            string InternationalizationPreset = null,
+            ParamList<string> CulturesToCook = null,
+			ParamList<string> ClientCookedTargets = null,
+			ParamList<string> EditorTargets = null,
+			ParamList<string> ServerCookedTargets = null,
+			List<UnrealTargetPlatform> ClientTargetPlatforms = null,
+            Dictionary<UnrealTargetPlatform, UnrealTargetPlatform> ClientDependentPlatformMap = null,
+			List<UnrealTargetPlatform> ServerTargetPlatforms = null,
+            Dictionary<UnrealTargetPlatform, UnrealTargetPlatform> ServerDependentPlatformMap = null,
+			bool? Build = null,
+			bool? Cook = null,
+			string CookFlavor = null,
+			bool? Run = null,
+			bool? SkipServer = null,
+			bool? Clean = null,
+            bool? Compressed = null,
+            bool? UseDebugParamForEditorExe = null,
+            bool? IterativeCooking = null,
+            bool? CookAll = null,
+            bool? CookMapsOnly = null,
+            bool? CookOnTheFly = null,
+            bool? CookOnTheFlyStreaming = null,
+            bool? UnversionedCookedContent = null,
+            bool? SkipCookingEditorContent = null,
+            int? NumCookersToSpawn = null,
+            string AdditionalCookerOptions = null,
+            string BasedOnReleaseVersion = null,
+            string CreateReleaseVersion = null,
+            bool? GeneratePatch = null,
+            string DLCName = null,
+            string DiffCookedContentPath = null,
+            bool? DLCIncludeEngineContent = null,
+            bool? NewCook = null,
+            bool? OldCook = null,
+			bool? CrashReporter = null,
+			bool? DedicatedServer = null,
+			bool? Client = null,
+			bool? Deploy = null,
+			bool? FileServer = null,
+			bool? Foreign = null,
+			bool? ForeignCode = null,
+			bool? LogWindow = null,
+			bool? NoCleanStage = null,
+			bool? NoClient = null,
+			bool? NoDebugInfo = null,
+			bool? NoXGE = null,
+			bool? Package = null,
+			bool? Pak = null,
+			bool? Prereqs = null,
+			bool? NoBootstrapExe = null,
+            bool? SignedPak = null,
+            bool? NullRHI = null,
+            bool? FakeClient = null,
+            bool? EditorTest = null,
+            bool? RunAutomationTests = null,
+            string RunAutomationTest = null,
+            int? CrashIndex = null,
+            bool? Rocket = null,
+			bool? SkipCook = null,
+			bool? SkipCookOnTheFly = null,
+			bool? SkipPak = null,
+			bool? SkipStage = null,
+			bool? Stage = null,
+			bool? Manifests = null,
+            bool? CreateChunkInstall = null,
+			bool? Unattended = null,
+			int? NumClients = null,
+			bool? Archive = null,
+			string ArchiveDirectoryParam = null,
+			bool? ArchiveMetaData = null,
+			bool? CreateAppBundle = null,
+			ParamList<string> ProgramTargets = null,
+			bool? Distribution = null,
+            bool? Prebuilt = null,
+            int? RunTimeoutSeconds = null,
+			string SpecifiedArchitecture = null,
+            bool? IterativeDeploy = null,
+			bool? FastCook = null,
+			bool? IgnoreCookErrors = null,
+			bool? CodeSign = null,
+			bool? UploadSymbols = null,
+			string Provision = null,
+			string Certificate = null,
+			ParamList<string> InMapsToRebuildLightMaps = null,
+			string CreateReleaseVersionBasePath = null, //unused
+			string BasedOnReleaseVersionBasePath = null, //unused
+			bool StageNonMonolithic = false //unused
+			)
+			: this(
+				new FileReference(RawProjectPath),
+
+				Command: Command,
+				Device: Device,			
+				MapToRun: MapToRun,	
+				AdditionalServerMapParams: AdditionalServerMapParams,
+				Port: Port,
+				RunCommandline: RunCommandline,						
+				StageCommandline: StageCommandline,
+				BundleName: BundleName,
+				StageDirectoryParam: StageDirectoryParam,
+				UE4Exe: UE4Exe,
+				SignPak: SignPak,
+				ClientConfigsToBuild: ClientConfigsToBuild,
+				ServerConfigsToBuild: ServerConfigsToBuild,
+				MapsToCook: MapsToCook,
+				DirectoriesToCook: DirectoriesToCook,
+				InternationalizationPreset: InternationalizationPreset,
+				CulturesToCook: CulturesToCook,
+				ClientCookedTargets: ClientCookedTargets,
+				EditorTargets: EditorTargets,
+				ServerCookedTargets: ServerCookedTargets,
+				ClientTargetPlatforms: ClientTargetPlatforms,
+				ClientDependentPlatformMap: ClientDependentPlatformMap,
+				ServerTargetPlatforms: ServerTargetPlatforms,
+				ServerDependentPlatformMap: ServerDependentPlatformMap,
+				Build: Build,
+				Cook: Cook,
+				CookFlavor: CookFlavor,
+				Run: Run,
+				SkipServer: SkipServer,
+				Clean: Clean,
+				Compressed: Compressed,
+				UseDebugParamForEditorExe: UseDebugParamForEditorExe,
+				IterativeCooking: IterativeCooking,
+				CookAll: CookAll,
+				CookMapsOnly: CookMapsOnly,
+				CookOnTheFly: CookOnTheFly,
+				CookOnTheFlyStreaming: CookOnTheFlyStreaming,
+				UnversionedCookedContent: UnversionedCookedContent,
+				SkipCookingEditorContent: SkipCookingEditorContent,
+				NumCookersToSpawn: NumCookersToSpawn,
+				AdditionalCookerOptions: AdditionalCookerOptions,
+				BasedOnReleaseVersion: BasedOnReleaseVersion,
+				CreateReleaseVersion: CreateReleaseVersion,
+				GeneratePatch: GeneratePatch,
+				DLCName: DLCName,
+				DiffCookedContentPath: DiffCookedContentPath,
+				DLCIncludeEngineContent: DLCIncludeEngineContent,
+				NewCook: NewCook,
+				OldCook: OldCook,
+				CrashReporter: CrashReporter,
+				DedicatedServer: DedicatedServer,
+				Client: Client,
+				Deploy: Deploy,
+				FileServer: FileServer,
+				Foreign: Foreign,
+				ForeignCode: ForeignCode,
+				LogWindow: LogWindow,
+				NoCleanStage: NoCleanStage,
+				NoClient: NoClient,
+				NoDebugInfo: NoDebugInfo,
+				NoXGE: NoXGE,
+				Package: Package,
+				Pak: Pak,
+				Prereqs: Prereqs,
+				NoBootstrapExe: NoBootstrapExe,
+				SignedPak: SignedPak,
+				NullRHI: NullRHI,
+				FakeClient: FakeClient,
+				EditorTest: EditorTest,
+				RunAutomationTests: RunAutomationTests,
+				RunAutomationTest: RunAutomationTest,
+				CrashIndex: CrashIndex,
+				Rocket: Rocket,
+				SkipCook: SkipCook,
+				SkipCookOnTheFly: SkipCookOnTheFly,
+				SkipPak: SkipPak,
+				SkipStage: SkipStage,
+				Stage: Stage,
+				Manifests: Manifests,
+				CreateChunkInstall: CreateChunkInstall,
+				Unattended: Unattended,
+				NumClients: NumClients,
+				Archive: Archive,
+				ArchiveDirectoryParam: ArchiveDirectoryParam,
+				ArchiveMetaData: ArchiveMetaData,
+				CreateAppBundle: CreateAppBundle,
+				ProgramTargets: ProgramTargets,
+				Distribution: Distribution,
+				Prebuilt: Prebuilt,
+				RunTimeoutSeconds: RunTimeoutSeconds,
+				SpecifiedArchitecture: SpecifiedArchitecture,
+				IterativeDeploy: IterativeDeploy,
+				FastCook: FastCook,
+				IgnoreCookErrors: IgnoreCookErrors,
+				CodeSign: CodeSign,
+				UploadSymbols: UploadSymbols,
+				Provision: Provision,
+				Certificate: Certificate,
+				InMapsToRebuildLightMaps: InMapsToRebuildLightMaps
+			)
+		{
+		}
+
 		#endregion
 
 		#region Shared
diff --git a/Engine/Source/Programs/AutomationTool/AutomationUtils/ScriptCompiler.cs b/Engine/Source/Programs/AutomationTool/AutomationUtils/ScriptCompiler.cs
--- a/Engine/Source/Programs/AutomationTool/AutomationUtils/ScriptCompiler.cs
+++ b/Engine/Source/Programs/AutomationTool/AutomationUtils/ScriptCompiler.cs
@@ -241,7 +241,7 @@ namespace AutomationTool
 				throw new AutomationException(String.Format("Unable to build Project {0}. Project file not found.", ProjectFile));
 			}
 
-			var CmdLine = String.Format("\"{0}\" /verbosity:quiet /nologo /target:Build /property:Configuration={1} /property:Platform=AnyCPU /p:TreatWarningsAsErrors=true /p:BuildProjectReferences=true",
+			var CmdLine = String.Format("\"{0}\" /verbosity:quiet /nologo /target:Build /property:Configuration={1} /property:Platform=AnyCPU /p:TreatWarningsAsErrors=false /p:NoWarn=\"612,618,672\" /p:BuildProjectReferences=true",
 				ProjectFile, BuildConfig);
 
 			// Compile the project
diff --git a/Engine/Source/Programs/AutomationTool/GUBP/LegacyBranchSetup.cs b/Engine/Source/Programs/AutomationTool/GUBP/LegacyBranchSetup.cs
--- a/Engine/Source/Programs/AutomationTool/GUBP/LegacyBranchSetup.cs
+++ b/Engine/Source/Programs/AutomationTool/GUBP/LegacyBranchSetup.cs
@@ -20,6 +20,13 @@ partial class GUBP
 		public bool bForceIncrementalCompile;
         public JobInfo JobInfo;
 
+		// Included for compatibility with //UE4/Main import
+		[Obsolete]
+		public bool bPreflightBuild
+		{
+			get { return JobInfo.IsPreflight; }
+		}
+
 		public Dictionary<string, GUBPNode> GUBPNodes = new Dictionary<string, GUBPNode>();
 		public Dictionary<string, GUBPAggregateNode> GUBPAggregates = new Dictionary<string, GUBPAggregateNode>();
 
diff --git a/Engine/Source/Programs/AutomationTool/GUBP/LegacyNodes.cs b/Engine/Source/Programs/AutomationTool/GUBP/LegacyNodes.cs
--- a/Engine/Source/Programs/AutomationTool/GUBP/LegacyNodes.cs
+++ b/Engine/Source/Programs/AutomationTool/GUBP/LegacyNodes.cs
@@ -20,6 +20,20 @@ partial class GUBP
         public List<string> AllDependencyBuildProducts = null;
         public string AgentSharingGroup = "";
 
+		// Included for compatibility during //UE4/Main import
+		[Obsolete]
+		public virtual string GameNameIfAnyForTempStorage()
+		{
+			return "";
+		}
+
+		// Included for compatibility during //UE4/Main import
+		[Obsolete]
+		public virtual string RootIfAnyForTempStorage()
+		{
+			return "";
+		}
+
         public virtual string GetFullName()
         {
             throw new AutomationException("Unimplemented GetFullName.");
@@ -1581,7 +1595,7 @@ partial class GUBP
         UnrealTargetPlatform TargetPlatform;
         string CookPlatform;
         bool bIsMassive;
-
+		
 		public String ExtraArgsForCook = "";
 
         public CookNode(GUBPBranchConfig BranchConfig, UnrealTargetPlatform InHostPlatform, BranchInfo.BranchUProject InGameProj, UnrealTargetPlatform InTargetPlatform, string InCookPlatform)
@@ -1676,6 +1690,30 @@ partial class GUBP
         {
             return bIsMassive ? 240 : base.TimeoutInMinutes();
         }
+
+		public string RootForCook()
+		{
+			return CombinePaths(Path.GetDirectoryName(GameProj.FilePath.FullName), "Saved", "Cooked", CookPlatform);
+		}
+
+		//this lets the extra build products through, but filepaths become too long for windows. revert for now.
+#if false
+        public override string RootIfAnyForTempStorage()
+        {
+            return CombinePaths(Path.GetDirectoryName(GameProj.FilePath), "Saved");
+        }
+
+		private string RootForCookedFiles()
+		{
+			return CombinePaths(RootIfAnyForTempStorage(), "Cooked", CookPlatform);
+		}
+
+		private string RootForChunkFiles()
+		{
+			return CombinePaths(RootIfAnyForTempStorage(), "TmpPackaging", CookPlatform);
+		}
+#endif
+
         public override void DoBuild(GUBP bp)
         {
             if (HostPlatform == UnrealTargetPlatform.Mac)
@@ -1687,7 +1725,7 @@ partial class GUBP
 				String CookArgs = String.Format("-Unversioned {0}", ExtraArgsForCook);
 				CommandUtils.CookCommandlet(GameProj.FilePath, "UE4Editor-Cmd.exe", null, null, null, null, CookPlatform, CookArgs);
 			}
-            var CookedPath = CombinePaths(Path.GetDirectoryName(GameProj.FilePath.FullName), "Saved", "Cooked", CookPlatform);
+			var CookedPath = RootForCook();
             var CookedFiles = CommandUtils.FindFiles("*", true, CookedPath);
             if (CookedFiles.GetLength(0) < 1)
             {
@@ -1699,6 +1737,17 @@ partial class GUBP
             {
                 AddBuildProduct(CookedFile);
             }
+
+			//can't add these until we straighten out the storage root.
+#if false
+			//add in any chunk data files if they exist.
+			String TmpPackagingPath = RootForChunkFiles();
+			var PackagingFiles = CommandUtils.FindFiles("*", true, TmpPackagingPath);
+			foreach (var PackageDataFile in PackagingFiles)
+			{
+				AddBuildProduct(PackageDataFile);
+			}
+#endif
         }
     }
 
diff --git a/Engine/Source/Programs/UnrealBuildTool/Configuration/UEBuildPlatform.cs b/Engine/Source/Programs/UnrealBuildTool/Configuration/UEBuildPlatform.cs
--- a/Engine/Source/Programs/UnrealBuildTool/Configuration/UEBuildPlatform.cs
+++ b/Engine/Source/Programs/UnrealBuildTool/Configuration/UEBuildPlatform.cs
@@ -456,6 +456,20 @@ namespace UnrealBuildTool
 			}
 		}
 
+		// Included for compatibility during //UE4/Main import
+		[Obsolete]
+		public CPPTargetPlatform GetCPPTargetPlatform(UnrealTargetPlatform ForPlatform)
+		{
+			if(Platform == ForPlatform)
+			{
+				return DefaultCppPlatform;
+			}
+			else
+			{
+				throw new BuildException("Target platform {0} can not be compiled on {1}", ForPlatform, Platform);
+			}
+		}
+
 		/// <summary>
 		/// Returns the delimiter used to separate paths in the PATH environment variable for the platform we are executing on.
 		/// </summary>
diff --git a/Engine/Source/Programs/UnrealBuildTool/IOS/UEDeployIOS.cs b/Engine/Source/Programs/UnrealBuildTool/IOS/UEDeployIOS.cs
--- a/Engine/Source/Programs/UnrealBuildTool/IOS/UEDeployIOS.cs
+++ b/Engine/Source/Programs/UnrealBuildTool/IOS/UEDeployIOS.cs
@@ -157,6 +157,7 @@ namespace UnrealBuildTool
 			bool bSkipDefaultPNGs = false;
 			string IntermediateDirectory = (bIsUE4Game ? InEngineDir : ProjectDirectory) + "/Intermediate/IOS";
 			string PListFile = IntermediateDirectory + "/" + GameName + "-Info.plist";
+			ProjectName = !String.IsNullOrEmpty(ProjectName) ? ProjectName : GameName;
 			VersionUtilities.BuildDirectory = BuildDirectory;
 			VersionUtilities.GameName = GameName;
 
diff --git a/Engine/Source/Programs/UnrealBuildTool/System/DynamicCompilation.cs b/Engine/Source/Programs/UnrealBuildTool/System/DynamicCompilation.cs
--- a/Engine/Source/Programs/UnrealBuildTool/System/DynamicCompilation.cs
+++ b/Engine/Source/Programs/UnrealBuildTool/System/DynamicCompilation.cs
@@ -141,8 +141,8 @@ namespace UnrealBuildTool
 				// We always want to generate a class library, not an executable
 				CompileParams.GenerateExecutable = false;
 
-				// Always fail compiles for warnings
-				CompileParams.TreatWarningsAsErrors = true;
+				// Never fail compiles for warnings
+				CompileParams.TreatWarningsAsErrors = false;
 
 				// Always generate debug information as it takes minimal time
 				CompileParams.IncludeDebugInformation = true;
@@ -221,15 +221,18 @@ namespace UnrealBuildTool
 				throw new BuildException(Ex, "Failed to launch compiler to compile assembly from source files '{0}' (Exception: {1})", SourceFileNames.ToString(), Ex.Message);
 			}
 
-			// Display compilation errors
+			// Display compilation warnings and errors
 			if (CompileResults.Errors.Count > 0)
 			{
-				Log.TraceInformation("Errors detected while compiling {0}:", OutputAssemblyPath);
+				Log.TraceInformation("Messages while compiling {0}:", OutputAssemblyPath);
 				foreach (var CurError in CompileResults.Errors)
 				{
 					Log.TraceInformation(CurError.ToString());
 				}
-				throw new BuildException("UnrealBuildTool encountered an error while compiling source files");
+				if(CompileResults.Errors.HasErrors || TreatWarningsAsErrors)
+				{
+					throw new BuildException("UnrealBuildTool encountered an error while compiling source files");
+				}
 			}
 
 			// Grab the generated assembly
diff --git a/Engine/Source/Programs/UnrealBuildTool/System/ExternalExecution.cs b/Engine/Source/Programs/UnrealBuildTool/System/ExternalExecution.cs
--- a/Engine/Source/Programs/UnrealBuildTool/System/ExternalExecution.cs
+++ b/Engine/Source/Programs/UnrealBuildTool/System/ExternalExecution.cs
@@ -731,8 +731,12 @@ namespace UnrealBuildTool
 							UBTArguments.Append(" -noxge");
 						}
 
-						// Propagate command-line option to switch back to old 2013 toolchain
-						if (UnrealBuildTool.CommandLineContains("-2013"))
+						// Propagate command-line option
+						if ( UnrealBuildTool.CommandLineContains( "-2015" ) )
+						{
+							UBTArguments.Append( " -2015" );
+						}
+						if ( UnrealBuildTool.CommandLineContains( "-2013" ) )
 						{
 							UBTArguments.Append(" -2013");
 						}
diff --git a/Engine/Source/Programs/UnrealBuildTool/System/RulesCompiler.cs b/Engine/Source/Programs/UnrealBuildTool/System/RulesCompiler.cs
--- a/Engine/Source/Programs/UnrealBuildTool/System/RulesCompiler.cs
+++ b/Engine/Source/Programs/UnrealBuildTool/System/RulesCompiler.cs
@@ -1573,6 +1573,23 @@ namespace UnrealBuildTool
 		/// Map of root folders to a cached list of all UBT-related source files in that folder or any of its sub-folders.
 		/// We cache these file names so we can avoid searching for them later on.
 		static Dictionary<DirectoryReference, RulesFileCache> RootFolderToRulesFileCache = new Dictionary<DirectoryReference, RulesFileCache>();
+		
+		// Included for compatibility during //UE4/Main import
+		[Obsolete]
+		public static string GetModuleFilename(string ModuleName)
+		{
+			foreach(RulesFileCache Cache in RootFolderToRulesFileCache.Values)
+			{
+				foreach(FileReference ModuleFile in Cache.ModuleRules)
+				{
+					if(String.Compare(ModuleFile.GetFileNameWithoutAnyExtensions(), ModuleName, StringComparison.InvariantCultureIgnoreCase) == 0)
+					{
+						return ModuleFile.FullName;
+					}
+				}
+			}
+			return null;
+		}
 
 		public static List<FileReference> FindAllRulesSourceFiles(RulesFileType RulesFileType, List<DirectoryReference> GameFolders, List<FileReference> ForeignPlugins, List<DirectoryReference> AdditionalSearchPaths, bool bIncludeEngine = true)
 		{
diff --git a/Engine/Source/Programs/UnrealBuildTool/System/TargetReceipt.cs b/Engine/Source/Programs/UnrealBuildTool/System/TargetReceipt.cs
--- a/Engine/Source/Programs/UnrealBuildTool/System/TargetReceipt.cs
+++ b/Engine/Source/Programs/UnrealBuildTool/System/TargetReceipt.cs
@@ -455,4 +455,25 @@ namespace UnrealBuildTool
 			}
 		}
 	}
+
+	/// <summary>
+	/// Legacy class for compatibility in //UE4/Main import.
+	/// </summary>
+	[Obsolete]
+	public class BuildReceipt : TargetReceipt
+	{
+		public BuildReceipt(TargetReceipt Other) : base(Other)
+		{
+		}
+
+		public static new BuildReceipt Read(string FileName)
+		{
+			return new BuildReceipt(TargetReceipt.Read(FileName));
+		}
+
+		public void ExpandPathVariables(string EngineDir, string ProjectDir)
+		{
+			base.ExpandPathVariables(new DirectoryReference(EngineDir), new DirectoryReference(ProjectDir));
+		}
+	}
 }
diff --git a/Engine/Source/Programs/UnrealBuildTool/System/XcodeProjectFileGenerator.cs b/Engine/Source/Programs/UnrealBuildTool/System/XcodeProjectFileGenerator.cs
--- a/Engine/Source/Programs/UnrealBuildTool/System/XcodeProjectFileGenerator.cs
+++ b/Engine/Source/Programs/UnrealBuildTool/System/XcodeProjectFileGenerator.cs
@@ -777,7 +777,28 @@ namespace UnrealBuildTool
 		/// <returns>The normalized path</returns>
 		private static string ConvertPath(string InPath)
 		{
-			return InPath.Replace("\\", "/");
+			if (BuildHostPlatform.Current.Platform != UnrealTargetPlatform.Mac)
+			{
+				if (InPath[1] != ':')
+				{
+					throw new BuildException("Can only convert full paths ({0})", InPath);
+				}
+
+				string MacPath = string.Format("{0}{1}/{2}/{3}",
+					RemoteToolChain.UserDevRootMac,
+					Environment.MachineName,
+					InPath[0].ToString().ToUpper(),
+					InPath.Substring(3));
+
+				// clean the path
+				MacPath = MacPath.Replace("\\", "/");
+
+				return MacPath;
+			}
+			else
+			{
+				return InPath.Replace("\\", "/");
+			}
 		}
 
 		/// <summary>
diff --git a/Engine/Source/Programs/UnrealBuildTool/ToolChain/UEToolChain.cs b/Engine/Source/Programs/UnrealBuildTool/ToolChain/UEToolChain.cs
--- a/Engine/Source/Programs/UnrealBuildTool/ToolChain/UEToolChain.cs
+++ b/Engine/Source/Programs/UnrealBuildTool/ToolChain/UEToolChain.cs
@@ -9,7 +9,14 @@ using Microsoft.Win32;
 
 namespace UnrealBuildTool
 {
-	public abstract class UEToolChain
+	// Included for compatibility during //UE4/Main import
+	public interface IUEToolChain
+	{
+		[Obsolete]
+		void StripSymbols(string SourceFile, string TargetFile);
+	}
+
+	public abstract class UEToolChain : IUEToolChain
 	{
 		public readonly CPPTargetPlatform CppPlatform;
 
@@ -18,6 +25,15 @@ namespace UnrealBuildTool
 			CppPlatform = InCppPlatform;
 		}
 
+		// Included for compatibility during //UE4/Main import
+		[Obsolete]
+		public static IUEToolChain GetPlatformToolChain(CPPTargetPlatform Platform)
+		{
+			UEBuildPlatform BuildPlatform = UEBuildPlatform.GetBuildPlatform(UEBuildTarget.CPPTargetPlatformToUnrealTargetPlatform(Platform));
+			UEBuildPlatformContext Context = BuildPlatform.CreateContext(null);
+			return Context.CreateToolChain(BuildPlatform.DefaultCppPlatform);
+		}
+
 		public abstract CPPOutput CompileCPPFiles(UEBuildTarget Target, CPPEnvironment CompileEnvironment, List<FileItem> SourceFiles, string ModuleName);
 
 		public virtual CPPOutput CompileRCFiles(UEBuildTarget Target, CPPEnvironment Environment, List<FileItem> RCFiles)
diff --git a/Engine/Source/Programs/UnrealBuildTool/Windows/UEBuildWindows.cs b/Engine/Source/Programs/UnrealBuildTool/Windows/UEBuildWindows.cs
--- a/Engine/Source/Programs/UnrealBuildTool/Windows/UEBuildWindows.cs
+++ b/Engine/Source/Programs/UnrealBuildTool/Windows/UEBuildWindows.cs
@@ -463,19 +463,19 @@ namespace UnrealBuildTool
 					CachedCompiler = WindowsCompiler.VisualStudio2015;
 				}
 
-				// Second, default based on what's installed, from newest to oldest
- 				else if (!String.IsNullOrEmpty(WindowsPlatform.GetVSComnToolsPath(WindowsCompiler.VisualStudio2015)))
- 				{
- 					CachedCompiler = WindowsCompiler.VisualStudio2015;
- 				}
+				// Second, default based on what's installed, test for 2013 first
 				else if (!String.IsNullOrEmpty(WindowsPlatform.GetVSComnToolsPath(WindowsCompiler.VisualStudio2013)))
 				{
 					CachedCompiler = WindowsCompiler.VisualStudio2013;
 				}
+ 				else if (!String.IsNullOrEmpty(WindowsPlatform.GetVSComnToolsPath(WindowsCompiler.VisualStudio2015)))
+ 				{
+ 					CachedCompiler = WindowsCompiler.VisualStudio2015;
+ 				}
 				else
 				{
-					// Finally assume 2015 is installed to defer errors somewhere else like VCToolChain
-					CachedCompiler = WindowsCompiler.VisualStudio2015;
+					// Finally assume 2013 is installed to defer errors somewhere else like VCToolChain
+					CachedCompiler = WindowsCompiler.VisualStudio2013;
 				}
 
 				return CachedCompiler.Value;
diff --git a/Engine/Source/Runtime/AIModule/Private/Actions/PawnActionsComponent.cpp b/Engine/Source/Runtime/AIModule/Private/Actions/PawnActionsComponent.cpp
--- a/Engine/Source/Runtime/AIModule/Private/Actions/PawnActionsComponent.cpp
+++ b/Engine/Source/Runtime/AIModule/Private/Actions/PawnActionsComponent.cpp
@@ -418,7 +418,7 @@ EPawnActionAbortState::Type UPawnActionsComponent::K2_ForceAbortAction(UPawnActi
 {
 	if (ActionToAbort)
 	{
-		ForceAbortAction(*ActionToAbort);
+		return ForceAbortAction(*ActionToAbort);
 	}
 	return EPawnActionAbortState::NeverStarted;
 }
@@ -553,7 +553,7 @@ bool UPawnActionsComponent::K2_PerformAction(APawn* Pawn, UPawnAction* Action, T
 {
 	if (Pawn && Action)
 	{
-		PerformAction(*Pawn, *Action, Priority);
+		return PerformAction(*Pawn, *Action, Priority);
 	}
 	return false;
 }
diff --git a/Engine/Source/Runtime/Core/Private/IOS/IOSPlatformMisc.cpp b/Engine/Source/Runtime/Core/Private/IOS/IOSPlatformMisc.cpp
--- a/Engine/Source/Runtime/Core/Private/IOS/IOSPlatformMisc.cpp
+++ b/Engine/Source/Runtime/Core/Private/IOS/IOSPlatformMisc.cpp
@@ -476,11 +476,11 @@ FIOSPlatformMisc::EIOSDevice FIOSPlatformMisc::GetIOSDeviceType()
 		{
 			if (Minor == 1)
 			{
-				DeviceType == IOS_IPhone6Plus;
+				DeviceType = IOS_IPhone6Plus;
 			}
 			else if (Minor == 2)
 			{
-				DeviceType == IOS_IPhone6;
+				DeviceType = IOS_IPhone6;
 			}
 		}
 		else if (Major == 8)
@@ -488,11 +488,11 @@ FIOSPlatformMisc::EIOSDevice FIOSPlatformMisc::GetIOSDeviceType()
 			// note that Apple switched the minor order around between 6 and 6S (gotta keep us on our toes!)
 			if (Minor == 1)
 			{
-				DeviceType == IOS_IPhone6;
+				DeviceType = IOS_IPhone6;
 			}
 			else if (Minor == 2)
 			{
-				DeviceType == IOS_IPhone6Plus;
+				DeviceType = IOS_IPhone6Plus;
 			}
 		}
 		else if (Major >= 9)
diff --git a/Engine/Source/Runtime/Core/Private/Linux/LinuxPlatformMemory.cpp b/Engine/Source/Runtime/Core/Private/Linux/LinuxPlatformMemory.cpp
--- a/Engine/Source/Runtime/Core/Private/Linux/LinuxPlatformMemory.cpp
+++ b/Engine/Source/Runtime/Core/Private/Linux/LinuxPlatformMemory.cpp
@@ -41,11 +41,6 @@ class FMalloc* FLinuxPlatformMemory::BaseAllocator()
 		Jemalloc,
 		Binned
 	}
->>>> ORIGINAL //depot/UE4-Orion/Engine/Source/Runtime/Core/Private/Linux/LinuxPlatformMemory.cpp#4
-	AllocatorToUse = FORCE_ANSI_ALLOCATOR ? EAllocatorToUse::Ansi : EAllocatorToUse::Binned;
-==== THEIRS //depot/UE4-Orion/Engine/Source/Runtime/Core/Private/Linux/LinuxPlatformMemory.cpp#5
-	AllocatorToUse = FORCE_ANSI_ALLOCATOR ? EAllocatorToUse::Ansi : EAllocatorToUse::Jemalloc;
-==== YOURS //Marc.Audy_Z2487/Engine/Source/Runtime/Core/Private/Linux/LinuxPlatformMemory.cpp
 	AllocatorToUse = EAllocatorToUse::Binned;
 
 	// Prefer jemalloc for the editor and programs as it saved ~20% RES usage in my (RCL) tests.
@@ -54,7 +49,6 @@ class FMalloc* FLinuxPlatformMemory::BaseAllocator()
 	{
 		AllocatorToUse = EAllocatorToUse::Jemalloc;
 	}
-<<<<
 
 	if (FORCE_ANSI_ALLOCATOR)
 	{
diff --git a/Engine/Source/Runtime/Core/Public/Math/Transform.h b/Engine/Source/Runtime/Core/Public/Math/Transform.h
--- a/Engine/Source/Runtime/Core/Public/Math/Transform.h
+++ b/Engine/Source/Runtime/Core/Public/Math/Transform.h
@@ -90,7 +90,7 @@ public:
 	FORCEINLINE void DiagnosticCheckNaN_Rotate() const {}
 	FORCEINLINE void DiagnosticCheckNaN_Scale3D() const {}
 	FORCEINLINE void DiagnosticCheckNaN_All() const {}
-	FORCEINLINE void DiagnosticCheck_IsValid() const;
+	FORCEINLINE void DiagnosticCheck_IsValid() const {}
 #endif
 
 	/**
diff --git a/Engine/Source/Runtime/Engine/Private/SceneView.cpp b/Engine/Source/Runtime/Engine/Private/SceneView.cpp
--- a/Engine/Source/Runtime/Engine/Private/SceneView.cpp
+++ b/Engine/Source/Runtime/Engine/Private/SceneView.cpp
@@ -46,14 +46,21 @@ FBuiltinSamplersUniformBuffer::FBuiltinSamplersUniformBuffer()
 
 void FBuiltinSamplersUniformBuffer::InitDynamicRHI()
 {
+	static TRefCountPtr<FRHISamplerState> Bilinear =			RHICreateSamplerState(FSamplerStateInitializerRHI(SF_Bilinear, AM_Wrap, AM_Wrap, AM_Wrap));
+	static TRefCountPtr<FRHISamplerState> BilinearClamped =		RHICreateSamplerState(FSamplerStateInitializerRHI(SF_Bilinear, AM_Clamp, AM_Clamp, AM_Clamp));
+	static TRefCountPtr<FRHISamplerState> Point =				RHICreateSamplerState(FSamplerStateInitializerRHI(SF_Point, AM_Wrap, AM_Wrap, AM_Wrap));
+	static TRefCountPtr<FRHISamplerState> PointClamped =		RHICreateSamplerState(FSamplerStateInitializerRHI(SF_Point, AM_Clamp, AM_Clamp, AM_Clamp));
+	static TRefCountPtr<FRHISamplerState> Trilinear =			RHICreateSamplerState(FSamplerStateInitializerRHI(SF_Trilinear, AM_Wrap, AM_Wrap, AM_Wrap));
+	static TRefCountPtr<FRHISamplerState> TrilinearClamped =	RHICreateSamplerState(FSamplerStateInitializerRHI(SF_Trilinear, AM_Clamp, AM_Clamp, AM_Clamp));
+
 	FBuiltinSamplersParameters UB;
 
-	UB.Bilinear =			RHICreateSamplerState(FSamplerStateInitializerRHI(SF_Bilinear,	AM_Wrap, AM_Wrap, AM_Wrap));
-	UB.BilinearClamped =	RHICreateSamplerState(FSamplerStateInitializerRHI(SF_Bilinear,	AM_Clamp, AM_Clamp, AM_Clamp));
-	UB.Point =				RHICreateSamplerState(FSamplerStateInitializerRHI(SF_Point,		AM_Wrap, AM_Wrap, AM_Wrap));
-	UB.PointClamped =		RHICreateSamplerState(FSamplerStateInitializerRHI(SF_Point,		AM_Clamp, AM_Clamp, AM_Clamp));
-	UB.Trilinear =			RHICreateSamplerState(FSamplerStateInitializerRHI(SF_Trilinear,	AM_Wrap, AM_Wrap, AM_Wrap));
-	UB.TrilinearClamped =	RHICreateSamplerState(FSamplerStateInitializerRHI(SF_Trilinear,	AM_Clamp, AM_Clamp, AM_Clamp));
+	UB.Bilinear =			Bilinear;
+	UB.BilinearClamped =	BilinearClamped;
+	UB.Point =				Point;
+	UB.PointClamped =		PointClamped;
+	UB.Trilinear =			Trilinear;
+	UB.TrilinearClamped =	TrilinearClamped;
 	SetContents(UB);
 
 	TUniformBuffer<FBuiltinSamplersParameters>::InitDynamicRHI();
diff --git a/Engine/Source/Runtime/Engine/Public/SceneView.h b/Engine/Source/Runtime/Engine/Public/SceneView.h
--- a/Engine/Source/Runtime/Engine/Public/SceneView.h
+++ b/Engine/Source/Runtime/Engine/Public/SceneView.h
@@ -474,7 +474,7 @@ BEGIN_UNIFORM_BUFFER_STRUCT(FBuiltinSamplersParameters,ENGINE_API)
 	DECLARE_UNIFORM_BUFFER_STRUCT_MEMBER_SAMPLER(SamplerState, TrilinearClamped)
 END_UNIFORM_BUFFER_STRUCT(FBuiltinSamplersParameters)
 
-ENGINE_API class FBuiltinSamplersUniformBuffer : public TUniformBuffer<FBuiltinSamplersParameters>
+class ENGINE_API FBuiltinSamplersUniformBuffer : public TUniformBuffer<FBuiltinSamplersParameters>
 {
 public:
 	FBuiltinSamplersUniformBuffer();
diff --git a/Engine/Source/Runtime/Launch/Private/Android/AndroidJNI.cpp b/Engine/Source/Runtime/Launch/Private/Android/AndroidJNI.cpp
--- a/Engine/Source/Runtime/Launch/Private/Android/AndroidJNI.cpp
+++ b/Engine/Source/Runtime/Launch/Private/Android/AndroidJNI.cpp
@@ -394,7 +394,7 @@ extern "C" void Java_com_epicgames_ue4_GameActivity_nativeVirtualKeyboardResult(
 			{
 				FGraphEventRef SetWidgetText = FFunctionGraphTask::CreateAndDispatchWhenReady([&]()
 				{
-					VirtualKeyboardWidget->SetTextFromVirtualKeyboard(FText::FromString(FString(UTF8_TO_TCHAR(javaChars))), false);
+					VirtualKeyboardWidget->SetTextFromVirtualKeyboard(FText::FromString(FString(UTF8_TO_TCHAR(javaChars))), ESetTextType::Commited, ETextCommit::OnUserMovedFocus);
 				}, TStatId(), NULL, ENamedThreads::GameThread);
 				FTaskGraphInterface::Get().WaitUntilTaskCompletes(SetWidgetText);
 			}
diff --git a/Engine/Source/Runtime/Networking/Public/Common/UdpSocketReceiver.h b/Engine/Source/Runtime/Networking/Public/Common/UdpSocketReceiver.h
--- a/Engine/Source/Runtime/Networking/Public/Common/UdpSocketReceiver.h
+++ b/Engine/Source/Runtime/Networking/Public/Common/UdpSocketReceiver.h
@@ -94,9 +94,12 @@ public:
 				Reader->SetNumUninitialized(FMath::Min(Size, 65507u));
 
 				int32 Read = 0;
-				Socket->RecvFrom(Reader->GetData(), Reader->Num(), Read, *Sender);
-
-				DataReceivedDelegate.ExecuteIfBound(Reader, FIPv4Endpoint(Sender));
+				
+				if (Socket->RecvFrom(Reader->GetData(), Reader->Num(), Read, *Sender))
+				{
+					Reader->RemoveAt(Read, Reader->Num() - Read, false);
+					DataReceivedDelegate.ExecuteIfBound(Reader, FIPv4Endpoint(Sender));
+				}
 			}
 		}
 
diff --git a/Engine/Source/Runtime/RHI/Public/RHICommandList.h b/Engine/Source/Runtime/RHI/Public/RHICommandList.h
--- a/Engine/Source/Runtime/RHI/Public/RHICommandList.h
+++ b/Engine/Source/Runtime/RHI/Public/RHICommandList.h
@@ -1356,7 +1356,7 @@ public:
 		}
 		else
 		{
-			check(Contents && ContentsSize && &Layout);
+			check(Contents && ContentsSize && (&Layout == nullptr));
 			auto* Cmd = new (AllocCommand<FRHICommandBuildLocalUniformBuffer>()) FRHICommandBuildLocalUniformBuffer(this, Contents, ContentsSize, Layout);
 			Result.WorkArea = &Cmd->WorkArea;
 		}
diff --git a/Engine/Source/Runtime/Slate/Private/Framework/Text/IOS/IOSPlatformTextField.cpp b/Engine/Source/Runtime/Slate/Private/Framework/Text/IOS/IOSPlatformTextField.cpp
--- a/Engine/Source/Runtime/Slate/Private/Framework/Text/IOS/IOSPlatformTextField.cpp
+++ b/Engine/Source/Runtime/Slate/Private/Framework/Text/IOS/IOSPlatformTextField.cpp
@@ -59,7 +59,7 @@ void FIOSPlatformTextField::ShowVirtualKeyboard(bool bShow, int32 UserIndex, TSh
 											FIOSAsyncTask* AsyncTask = [[FIOSAsyncTask alloc] init];
 											AsyncTask.GameThreadCallback = ^ bool(void)
 											{
-												TextWidget->SetTextFromVirtualKeyboard(TextEntry, false);
+												TextWidget->SetTextFromVirtualKeyboard(TextEntry, ESetTextType::Commited, ETextCommit::OnUserMovedFocus);
 
 												// clear the TextWidget
 												TextWidget = nullptr;
@@ -193,11 +193,11 @@ void FIOSPlatformTextField::ShowVirtualKeyboard(bool bShow, int32 UserIndex, TSh
 	FIOSAsyncTask* AsyncTask = [[FIOSAsyncTask alloc] init];
     AsyncTask.GameThreadCallback = ^ bool(void)
     {
-	// index 1 is the OK button
-	if(buttonIndex == 1)
-	{
-			TextWidget->SetTextFromVirtualKeyboard(TextEntry, false);
-	}
+		// index 1 is the OK button
+		if(buttonIndex == 1)
+		{
+			TextWidget->SetTextFromVirtualKeyboard(TextEntry, ESetTextType::Commited, ETextCommit::OnUserMovedFocus);
+		}
     
         // clear the TextWidget
         TextWidget = nullptr;
diff --git a/Engine/Source/Runtime/Sockets/Private/BSDSockets/SocketsBSD.cpp b/Engine/Source/Runtime/Sockets/Private/BSDSockets/SocketsBSD.cpp
--- a/Engine/Source/Runtime/Sockets/Private/BSDSockets/SocketsBSD.cpp
+++ b/Engine/Source/Runtime/Sockets/Private/BSDSockets/SocketsBSD.cpp
@@ -167,13 +167,16 @@ bool FSocketBSD::RecvFrom(uint8* Data, int32 BufferSize, int32& BytesRead, FInte
 	BytesRead = recvfrom(Socket, (char*)Data, BufferSize, TranslatedFlags, &Addr, &Size);
 //	NETWORK_PROFILER(FSocket::RecvFrom(Data,BufferSize,BytesRead,Source));
 
-	bool Result = BytesRead >= 0;
-	if (Result)
+	if (BytesRead < 0)
 	{
-		LastActivityTime = FDateTime::UtcNow();
+		BytesRead = 0;
+
+		return false;
 	}
 
-	return Result;
+	LastActivityTime = FDateTime::UtcNow();
+
+	return true;
 }
 
 
