Update the new index before writing it to the CIS counter file, otherwise we're just going to run full CIS on the *next* build.

diff --git a/Engine/Source/Programs/AutomationTool/GUBP/GUBP.cs b/Engine/Source/Programs/AutomationTool/GUBP/GUBP.cs
--- a/Engine/Source/Programs/AutomationTool/GUBP/GUBP.cs
+++ b/Engine/Source/Programs/AutomationTool/GUBP/GUBP.cs
@@ -1153,8 +1153,13 @@ public partial class GUBP : BuildCommand
 					}
 				}
 			}
+			if(ParseParam("ResetCounter"))
 			{
-				string Line = ParseParam("ResetCounter")? "-1 0" : String.Format("{0} {1}", NewIndex, NowMinutes);
+				NewIndex = 0;
+				NowMinutes = 0;
+			}
+			{
+				string Line = String.Format("{0} {1}", NewIndex, NowMinutes);
 				LogVerbose("Attempting to write {0} with {1}", P4IndexFileP4, Line);
 				int WorkingCL = -1;
 				try
