remove temporary ProjectParams constructor and update out of date files

diff --git a/Engine/Source/Programs/AutomationTool/AutomationUtils/ProjectParams.cs b/Engine/Source/Programs/AutomationTool/AutomationUtils/ProjectParams.cs
--- a/Engine/Source/Programs/AutomationTool/AutomationUtils/ProjectParams.cs
+++ b/Engine/Source/Programs/AutomationTool/AutomationUtils/ProjectParams.cs
@@ -781,220 +781,6 @@ namespace AutomationTool
 			ValidateAndLog();
 		}
 
-		// Included for compatibility during //UE4/Main import
-		[Obsolete]
-		public ProjectParams(			
-			string RawProjectPath,
-
-			CommandUtils Command = null,
-			string Device = null,			
-			string MapToRun = null,	
-			string AdditionalServerMapParams = null,
-			ParamList<string> Port = null,
-			string RunCommandline = null,						
-			string StageCommandline = null,
-            string BundleName = null,
-            string StageDirectoryParam = null,
-			string UE4Exe = null,
-			string SignPak = null,
-			List<UnrealTargetConfiguration> ClientConfigsToBuild = null,
-			List<UnrealTargetConfiguration> ServerConfigsToBuild = null,
-			ParamList<string> MapsToCook = null,
-			ParamList<string> DirectoriesToCook = null,
-            string InternationalizationPreset = null,
-            ParamList<string> CulturesToCook = null,
-			ParamList<string> ClientCookedTargets = null,
-			ParamList<string> EditorTargets = null,
-			ParamList<string> ServerCookedTargets = null,
-			List<UnrealTargetPlatform> ClientTargetPlatforms = null,
-            Dictionary<UnrealTargetPlatform, UnrealTargetPlatform> ClientDependentPlatformMap = null,
-			List<UnrealTargetPlatform> ServerTargetPlatforms = null,
-            Dictionary<UnrealTargetPlatform, UnrealTargetPlatform> ServerDependentPlatformMap = null,
-			bool? Build = null,
-			bool? Cook = null,
-			string CookFlavor = null,
-			bool? Run = null,
-			bool? SkipServer = null,
-			bool? Clean = null,
-            bool? Compressed = null,
-            bool? UseDebugParamForEditorExe = null,
-            bool? IterativeCooking = null,
-            bool? CookAll = null,
-            bool? CookMapsOnly = null,
-            bool? CookOnTheFly = null,
-            bool? CookOnTheFlyStreaming = null,
-            bool? UnversionedCookedContent = null,
-            bool? SkipCookingEditorContent = null,
-            int? NumCookersToSpawn = null,
-            string AdditionalCookerOptions = null,
-            string BasedOnReleaseVersion = null,
-            string CreateReleaseVersion = null,
-            bool? GeneratePatch = null,
-            string DLCName = null,
-            string DiffCookedContentPath = null,
-            bool? DLCIncludeEngineContent = null,
-            bool? NewCook = null,
-            bool? OldCook = null,
-			bool? CrashReporter = null,
-			bool? DedicatedServer = null,
-			bool? Client = null,
-			bool? Deploy = null,
-			bool? FileServer = null,
-			bool? Foreign = null,
-			bool? ForeignCode = null,
-			bool? LogWindow = null,
-			bool? NoCleanStage = null,
-			bool? NoClient = null,
-			bool? NoDebugInfo = null,
-			bool? NoXGE = null,
-			bool? Package = null,
-			bool? Pak = null,
-			bool? Prereqs = null,
-			bool? NoBootstrapExe = null,
-            bool? SignedPak = null,
-            bool? NullRHI = null,
-            bool? FakeClient = null,
-            bool? EditorTest = null,
-            bool? RunAutomationTests = null,
-            string RunAutomationTest = null,
-            int? CrashIndex = null,
-            bool? Rocket = null,
-			bool? SkipCook = null,
-			bool? SkipCookOnTheFly = null,
-			bool? SkipPak = null,
-			bool? SkipStage = null,
-			bool? Stage = null,
-			bool? Manifests = null,
-            bool? CreateChunkInstall = null,
-			bool? Unattended = null,
-			int? NumClients = null,
-			bool? Archive = null,
-			string ArchiveDirectoryParam = null,
-			bool? ArchiveMetaData = null,
-			bool? CreateAppBundle = null,
-			ParamList<string> ProgramTargets = null,
-			bool? Distribution = null,
-            bool? Prebuilt = null,
-            int? RunTimeoutSeconds = null,
-			string SpecifiedArchitecture = null,
-            bool? IterativeDeploy = null,
-			bool? FastCook = null,
-			bool? IgnoreCookErrors = null,
-			bool? CodeSign = null,
-			bool? UploadSymbols = null,
-			string Provision = null,
-			string Certificate = null,
-			ParamList<string> InMapsToRebuildLightMaps = null,
-			string CreateReleaseVersionBasePath = null, //unused
-			string BasedOnReleaseVersionBasePath = null, //unused
-			bool StageNonMonolithic = false //unused
-			)
-			: this(
-				new FileReference(RawProjectPath),
-
-				Command: Command,
-				Device: Device,			
-				MapToRun: MapToRun,	
-				AdditionalServerMapParams: AdditionalServerMapParams,
-				Port: Port,
-				RunCommandline: RunCommandline,						
-				StageCommandline: StageCommandline,
-				BundleName: BundleName,
-				StageDirectoryParam: StageDirectoryParam,
-				UE4Exe: UE4Exe,
-				SignPak: SignPak,
-				ClientConfigsToBuild: ClientConfigsToBuild,
-				ServerConfigsToBuild: ServerConfigsToBuild,
-				MapsToCook: MapsToCook,
-				DirectoriesToCook: DirectoriesToCook,
-				InternationalizationPreset: InternationalizationPreset,
-				CulturesToCook: CulturesToCook,
-				ClientCookedTargets: ClientCookedTargets,
-				EditorTargets: EditorTargets,
-				ServerCookedTargets: ServerCookedTargets,
-				ClientTargetPlatforms: ClientTargetPlatforms,
-				ClientDependentPlatformMap: ClientDependentPlatformMap,
-				ServerTargetPlatforms: ServerTargetPlatforms,
-				ServerDependentPlatformMap: ServerDependentPlatformMap,
-				Build: Build,
-				Cook: Cook,
-				CookFlavor: CookFlavor,
-				Run: Run,
-				SkipServer: SkipServer,
-				Clean: Clean,
-				Compressed: Compressed,
-				UseDebugParamForEditorExe: UseDebugParamForEditorExe,
-				IterativeCooking: IterativeCooking,
-				CookAll: CookAll,
-				CookMapsOnly: CookMapsOnly,
-				CookOnTheFly: CookOnTheFly,
-				CookOnTheFlyStreaming: CookOnTheFlyStreaming,
-				UnversionedCookedContent: UnversionedCookedContent,
-				SkipCookingEditorContent: SkipCookingEditorContent,
-				NumCookersToSpawn: NumCookersToSpawn,
-				AdditionalCookerOptions: AdditionalCookerOptions,
-				BasedOnReleaseVersion: BasedOnReleaseVersion,
-				CreateReleaseVersion: CreateReleaseVersion,
-				GeneratePatch: GeneratePatch,
-				DLCName: DLCName,
-				DiffCookedContentPath: DiffCookedContentPath,
-				DLCIncludeEngineContent: DLCIncludeEngineContent,
-				NewCook: NewCook,
-				OldCook: OldCook,
-				CrashReporter: CrashReporter,
-				DedicatedServer: DedicatedServer,
-				Client: Client,
-				Deploy: Deploy,
-				FileServer: FileServer,
-				Foreign: Foreign,
-				ForeignCode: ForeignCode,
-				LogWindow: LogWindow,
-				NoCleanStage: NoCleanStage,
-				NoClient: NoClient,
-				NoDebugInfo: NoDebugInfo,
-				NoXGE: NoXGE,
-				Package: Package,
-				Pak: Pak,
-				Prereqs: Prereqs,
-				NoBootstrapExe: NoBootstrapExe,
-				SignedPak: SignedPak,
-				NullRHI: NullRHI,
-				FakeClient: FakeClient,
-				EditorTest: EditorTest,
-				RunAutomationTests: RunAutomationTests,
-				RunAutomationTest: RunAutomationTest,
-				CrashIndex: CrashIndex,
-				Rocket: Rocket,
-				SkipCook: SkipCook,
-				SkipCookOnTheFly: SkipCookOnTheFly,
-				SkipPak: SkipPak,
-				SkipStage: SkipStage,
-				Stage: Stage,
-				Manifests: Manifests,
-				CreateChunkInstall: CreateChunkInstall,
-				Unattended: Unattended,
-				NumClients: NumClients,
-				Archive: Archive,
-				ArchiveDirectoryParam: ArchiveDirectoryParam,
-				ArchiveMetaData: ArchiveMetaData,
-				CreateAppBundle: CreateAppBundle,
-				ProgramTargets: ProgramTargets,
-				Distribution: Distribution,
-				Prebuilt: Prebuilt,
-				RunTimeoutSeconds: RunTimeoutSeconds,
-				SpecifiedArchitecture: SpecifiedArchitecture,
-				IterativeDeploy: IterativeDeploy,
-				FastCook: FastCook,
-				IgnoreCookErrors: IgnoreCookErrors,
-				CodeSign: CodeSign,
-				UploadSymbols: UploadSymbols,
-				Provision: Provision,
-				Certificate: Certificate,
-				InMapsToRebuildLightMaps: InMapsToRebuildLightMaps
-			)
-		{
-		}
-
 		#endregion
 
 		#region Shared
