#UE4 Removed architecture exe decoration code that I recently added since it was not necessary and breaks linux (or more rather, continues to break linux even after removing the change to ProjectParams.cs).

diff --git a/Engine/Source/Programs/AutomationTool/Scripts/CopyBuildToStagingDirectory.Automation.cs b/Engine/Source/Programs/AutomationTool/Scripts/CopyBuildToStagingDirectory.Automation.cs
--- a/Engine/Source/Programs/AutomationTool/Scripts/CopyBuildToStagingDirectory.Automation.cs
+++ b/Engine/Source/Programs/AutomationTool/Scripts/CopyBuildToStagingDirectory.Automation.cs
@@ -1602,16 +1602,6 @@ public partial class Project : CommandUtils
 
 			string PlatformName = StagePlatform.ToString();
 			string StageArchitecture = !String.IsNullOrEmpty(Params.SpecifiedArchitecture) ? Params.SpecifiedArchitecture : "";
-			if (string.IsNullOrEmpty(StageArchitecture))
-			{
-				StageArchitecture = "";
-				var BuildPlatform = UEBuildPlatform.GetBuildPlatform(StagePlatform, true);
-				if (BuildPlatform != null)
-				{
-					StageArchitecture = BuildPlatform.CreateContext(Params.RawProjectPath).GetActiveArchitecture();
-				}
-			}
-
 			foreach (var Target in ListToProcess)
 			{
 				foreach (var Config in ConfigsToProcess)
