#UE4 Preserving subpath on DefaultEngine when copying to temp directory when preparing pak file.

diff --git a/Engine/Source/Programs/AutomationTool/Scripts/CopyBuildToStagingDirectory.Automation.cs b/Engine/Source/Programs/AutomationTool/Scripts/CopyBuildToStagingDirectory.Automation.cs
--- a/Engine/Source/Programs/AutomationTool/Scripts/CopyBuildToStagingDirectory.Automation.cs
+++ b/Engine/Source/Programs/AutomationTool/Scripts/CopyBuildToStagingDirectory.Automation.cs
@@ -797,9 +797,14 @@ public partial class Project : CommandUtils
 			}
 
 			// special case DefaultEngine.ini to strip passwords
-			if (Src.EndsWith("DefaultEngine.ini"))
+			if (Path.GetFileName(Src).Equals("DefaultEngine.ini"))
 			{
-				string NewIniFilename = CombinePaths(SC.ProjectRoot, "Saved", "Temp", SC.PlatformDir, "DefaultEngine.ini");
+				string SrcDirectoryPath = Path.GetDirectoryName(Src);
+				string ConfigDirectory = "Config";
+				int ConfigRootIdx = SrcDirectoryPath.LastIndexOf(ConfigDirectory);
+				string SubpathUnderConfig = (ConfigRootIdx != -1) ? SrcDirectoryPath.Substring(ConfigRootIdx + ConfigDirectory.Length) : "Unknown";
+
+				string NewIniFilename = CombinePaths(SC.ProjectRoot, "Saved", "Temp", SC.PlatformDir, SubpathUnderConfig, "DefaultEngine.ini");
 				InternalUtils.SafeCreateDirectory(Path.GetDirectoryName(NewIniFilename), true);
 				InternalUtils.SafeCopyFile(Src, NewIniFilename, bFilterSpecialLinesFromIniFiles:true);
 				Src = NewIniFilename;
