Fixed issue with custom icon not in windows task bar.

diff --git a/Engine/Source/Programs/AutomationTool/Win/WinPlatform.Automation.cs b/Engine/Source/Programs/AutomationTool/Win/WinPlatform.Automation.cs
--- a/Engine/Source/Programs/AutomationTool/Win/WinPlatform.Automation.cs
+++ b/Engine/Source/Programs/AutomationTool/Win/WinPlatform.Automation.cs
@@ -61,6 +61,24 @@ public abstract class BaseWinPlatform : Platform
 							BootstrapArguments = String.Format("..\\..\\..\\{0}\\{0}.uproject", SC.ShortProjectName);
 						}
 
+
+                        // set the icon on the original exe this will be used in the task bar when the bootstrap exe runs
+                        if (InternalUtils.SafeFileExists(CombinePaths(SC.ProjectRoot, "Build/Windows/Application.ico")))
+                        {
+                            GroupIconResource GroupIcon = null;
+                            GroupIcon = GroupIconResource.FromIco(CombinePaths(SC.ProjectRoot, "Build/Windows/Application.ico"));
+
+                            // Update the icon on the original exe because this will be used when the game is running in the task bar
+                            using (ModuleResourceUpdate Update = new ModuleResourceUpdate(Executable.Path, false))
+                            {
+                                const int IconResourceId = 101;
+                                if (GroupIcon != null)
+                                {
+                                    Update.SetIcons(IconResourceId, GroupIcon);
+                                }
+                            }
+                        }
+
 						string BootstrapExeName;
 						if(SC.StageTargetConfigurations.Count > 1)
 						{
