Fixed ue4game.exe in binaries directory having it's icon set when packaging.

diff --git a/Engine/Source/Programs/AutomationTool/Win/WinPlatform.Automation.cs b/Engine/Source/Programs/AutomationTool/Win/WinPlatform.Automation.cs
--- a/Engine/Source/Programs/AutomationTool/Win/WinPlatform.Automation.cs
+++ b/Engine/Source/Programs/AutomationTool/Win/WinPlatform.Automation.cs
@@ -62,19 +62,22 @@ public abstract class BaseWinPlatform : Platform
 						}
 
 
-                        // set the icon on the original exe this will be used in the task bar when the bootstrap exe runs
-                        if (InternalUtils.SafeFileExists(CombinePaths(SC.ProjectRoot, "Build/Windows/Application.ico")))
+                        foreach (string StageExePath in GetExecutableNames(SC))
                         {
-                            GroupIconResource GroupIcon = null;
-                            GroupIcon = GroupIconResource.FromIco(CombinePaths(SC.ProjectRoot, "Build/Windows/Application.ico"));
-
-                            // Update the icon on the original exe because this will be used when the game is running in the task bar
-                            using (ModuleResourceUpdate Update = new ModuleResourceUpdate(Executable.Path, false))
+                            // set the icon on the original exe this will be used in the task bar when the bootstrap exe runs
+                            if (InternalUtils.SafeFileExists(CombinePaths(SC.ProjectRoot, "Build/Windows/Application.ico")))
                             {
-                                const int IconResourceId = 101;
-                                if (GroupIcon != null)
+                                GroupIconResource GroupIcon = null;
+                                GroupIcon = GroupIconResource.FromIco(CombinePaths(SC.ProjectRoot, "Build/Windows/Application.ico"));
+
+                                // Update the icon on the original exe because this will be used when the game is running in the task bar
+                                using (ModuleResourceUpdate Update = new ModuleResourceUpdate(StageExePath, false))
                                 {
-                                    Update.SetIcons(IconResourceId, GroupIcon);
+                                    const int IconResourceId = 101;
+                                    if (GroupIcon != null)
+                                    {
+                                        Update.SetIcons(IconResourceId, GroupIcon);
+                                    }
                                 }
                             }
                         }
