TempStorage no longer tries to verify that each file path length is short enough for the shared location since it never actually creates the file there, it just zips them all up now.

We do still verify before we unzip that the file name is short enough, but this is much less likely to be a problem because it should exactly match the file path length of the original location that created it, which was by definition a valid file length!

diff --git a/Engine/Source/Programs/AutomationTool/GUBP/TempStorage.cs b/Engine/Source/Programs/AutomationTool/GUBP/TempStorage.cs
--- a/Engine/Source/Programs/AutomationTool/GUBP/TempStorage.cs
+++ b/Engine/Source/Programs/AutomationTool/GUBP/TempStorage.cs
@@ -410,19 +410,10 @@ namespace AutomationTool
             public List<string> GetFiles(string RootDir)
             {
                 // flatten the list of directories, pull the files out, set their root path, and ensure the path is not too long.
-                return Directories.SelectMany(Dir=>Dir.Value).Select(FileInfo =>
-                    {
-                        var NewFilePath = CommandUtils.CombinePaths(RootDir, FileInfo.Name);
-                        // create a FileInfo using the file, which will help us catch path too long exceptions early.
-                        try
-                        {
-                            return new FileInfo(NewFilePath).FullName;
-                        }
-                        catch (PathTooLongException Ex)
-                        {
-                            throw new AutomationException(Ex, "Path too long ... failed to create FileInfo for {0}", NewFilePath);
-                        }
-                    }).ToList();
+                return Directories
+                    .SelectMany(Dir => Dir.Value)
+                    .Select(FileInfo => CommandUtils.CombinePaths(RootDir, FileInfo.Name))
+                    .ToList();
             }
         }
 
@@ -1100,7 +1091,19 @@ namespace AutomationTool
 
                 // We know the source files exist and are under RootDir because we created the manifest, which verifies it.
                 // Now create the list of target files
-                var DestFiles = SharedFiles.Select(Filename => CommandUtils.MakeRerootedFilePath(Filename, SharedStorageNodeDir, RootDir)).ToList();
+                var DestFiles = SharedFiles.Select(Filename =>
+                {
+                    var DestFile = CommandUtils.MakeRerootedFilePath(Filename, SharedStorageNodeDir, RootDir);
+                    // create a FileInfo using the file, which will help us catch path too long exceptions early and propagate a friendly error containing the full path.
+                    try
+                    {
+                        return new FileInfo(DestFile).FullName;
+                    }
+                    catch (PathTooLongException Ex)
+                    {
+                        throw new AutomationException(Ex, "Path too long ... failed to create FileInfo for {0}", DestFile);
+                    }
+                }).ToList();
 
                 var ZipBasename = Path.GetFileNameWithoutExtension(LocalManifest);
 
