remove the read-only file-attribute before copy .htaccess

diff --git a/Engine/Source/Programs/AutomationTool/HTML5/HTML5Platform.Automation.cs b/Engine/Source/Programs/AutomationTool/HTML5/HTML5Platform.Automation.cs
--- a/Engine/Source/Programs/AutomationTool/HTML5/HTML5Platform.Automation.cs
+++ b/Engine/Source/Programs/AutomationTool/HTML5/HTML5Platform.Automation.cs
@@ -149,7 +149,17 @@ public class HTML5Platform : Platform
 		GenerateMacCommandFromTemplate(MacBashTemplateFile, MacBashOutputFile, MonoPath);
 
 		string htaccessTemplate = Path.Combine(CombinePaths(CmdEnv.LocalRoot, "Engine"), "Build", "HTML5", "htaccess.template");
-		File.Copy(htaccessTemplate, Path.Combine(PackagePath, ".htaccess"), true);
+		string htaccesspath = Path.Combine(PackagePath, ".htaccess");
+		if ( File.Exists(htaccesspath) )
+		{
+			FileAttributes attributes = File.GetAttributes(htaccesspath);
+			if ((attributes & FileAttributes.ReadOnly) == FileAttributes.ReadOnly)
+			{
+				attributes &= ~FileAttributes.ReadOnly;
+				File.SetAttributes(htaccesspath, attributes);
+			}
+		}
+		File.Copy(htaccessTemplate, htaccesspath, true);
 
 		string JSDir = Path.Combine(CombinePaths(CmdEnv.LocalRoot, "Engine"), "Build", "HTML5");
 		string OutDir = PackagePath;
@@ -395,7 +405,7 @@ public class HTML5Platform : Platform
 		SC.ArchiveFiles(PackagePath, Path.GetFileName(Path.Combine(PackagePath, "RunMacHTML5LaunchHelper.command")));
 		SC.ArchiveFiles(PackagePath, Path.GetFileName(Path.Combine(PackagePath, ".htaccess")));
 
-		if (HTMLPakAutomation.CanCreateMapPaks(Params))
+        if (HTMLPakAutomation.CanCreateMapPaks(Params))
 		{
 		// find all paks.
 			string[] Files = Directory.GetFiles(Path.Combine(PackagePath, Params.ShortProjectName), "*",SearchOption.AllDirectories);
