Fixed issue with diffcooked build failing compare of file and crashing in orion.

diff --git a/Engine/Source/Programs/AutomationTool/Scripts/CookCommand.Automation.cs b/Engine/Source/Programs/AutomationTool/Scripts/CookCommand.Automation.cs
--- a/Engine/Source/Programs/AutomationTool/Scripts/CookCommand.Automation.cs
+++ b/Engine/Source/Programs/AutomationTool/Scripts/CookCommand.Automation.cs
@@ -402,9 +402,13 @@ public partial class Project : CommandUtils
                         Log("Diff cooked content failed on file " + SourceFilename + " when comparing against " + DestFilename + " because " + Ex.ToString());
                     }
 
-                    if (SourceFile.LongLength == DestFile.LongLength)
+                    if (SourceFile == null || DestFile == null)
                     {
-                        for (long Index = 0; Index < SourceFile.LongLength; ++Index)
+                        Log("Diff cooked content failed on file " + SourceFilename + " when comparing against " + DestFilename + " " + (SourceFile==null?SourceFilename:DestFilename) + " file is missing");
+                    }
+                    else if (SourceFile.LongLength == DestFile.LongLength)
+                    {
+                        for (long Index = 0; Index < SourceFile.LongLength; ++Index) 
                         {
                             if (SourceFile[Index] != DestFile[Index])
                             {
