Stopped log output of exception because our log parser on build machines thinks that it's an error even though it's handled.

diff --git a/Engine/Source/Programs/AutomationTool/Scripts/CookCommand.Automation.cs b/Engine/Source/Programs/AutomationTool/Scripts/CookCommand.Automation.cs
--- a/Engine/Source/Programs/AutomationTool/Scripts/CookCommand.Automation.cs
+++ b/Engine/Source/Programs/AutomationTool/Scripts/CookCommand.Automation.cs
@@ -387,9 +387,9 @@ public partial class Project : CommandUtils
                     {
                         SourceFile = File.ReadAllBytes(SourceFilename);
                     }
-                    catch (Exception Ex)
+                    catch (Exception)
                     {
-                        Log("Diff cooked content failed on file " + SourceFilename + " when comparing against " + DestFilename + " because " + Ex.ToString());
+                        Log("Diff cooked content failed to load file " + SourceFilename);
                     }
 
                     byte[] DestFile = null;
@@ -397,9 +397,9 @@ public partial class Project : CommandUtils
                     {
                         DestFile = File.ReadAllBytes(DestFilename);
                     }
-                    catch (Exception Ex)
+                    catch (Exception)
                     {
-                        Log("Diff cooked content failed on file " + SourceFilename + " when comparing against " + DestFilename + " because " + Ex.ToString());
+                        Log("Diff cooked content failed to load file " + DestFilename );
                     }
 
                     if (SourceFile == null || DestFile == null)
