HTML5: output "shareable link" when packaging has uploaded successfully to amazon's s3 servers

tested with phosphor's AWS S3 KeyID & Secret Access Key

diff --git a/Engine/Source/Programs/AutomationTool/HTML5/HTML5Platform.Automation.cs b/Engine/Source/Programs/AutomationTool/HTML5/HTML5Platform.Automation.cs
--- a/Engine/Source/Programs/AutomationTool/HTML5/HTML5Platform.Automation.cs
+++ b/Engine/Source/Programs/AutomationTool/HTML5/HTML5Platform.Automation.cs
@@ -405,7 +405,7 @@ public class HTML5Platform : Platform
 		SC.ArchiveFiles(PackagePath, Path.GetFileName(Path.Combine(PackagePath, "RunMacHTML5LaunchHelper.command")));
 		SC.ArchiveFiles(PackagePath, Path.GetFileName(Path.Combine(PackagePath, ".htaccess")));
 
-        if (HTMLPakAutomation.CanCreateMapPaks(Params))
+		if (HTMLPakAutomation.CanCreateMapPaks(Params))
 		{
 		// find all paks.
 			string[] Files = Directory.GetFiles(Path.Combine(PackagePath, Params.ShortProjectName), "*",SearchOption.AllDirectories);
@@ -584,6 +584,10 @@ public class HTML5Platform : Platform
 		}
 
 		Task.WaitAll(UploadTasks.ToArray());
+
+		string URL = "http://" + BucketName + ".s3.amazonaws.com/" + FolderName + "/" + SC.ShortProjectName + ".html";
+		Log("Your project's shareable link is: " + URL);
+
 		Log("Upload Tasks finished.");
 	}
 
