Add script for doing UGS builds, which can handle submitting archived binaries to a stream.

diff --git a/Engine/Source/Programs/AutomationTool/AutomationUtils/CommandUtils.cs b/Engine/Source/Programs/AutomationTool/AutomationUtils/CommandUtils.cs
--- a/Engine/Source/Programs/AutomationTool/AutomationUtils/CommandUtils.cs
+++ b/Engine/Source/Programs/AutomationTool/AutomationUtils/CommandUtils.cs
@@ -1798,6 +1798,43 @@ namespace AutomationTool
 			return ParseParamValue(Params, Param, Default);
 		}
 
+		/// <summary>
+		/// Parses the argument list for any number of parameters.
+		/// </summary>
+		/// <param name="ArgList">Argument list.</param>
+		/// <param name="Param">Param to read its value.</param>
+		/// <returns>Returns an array of values for this parameter (or an empty array if one was not found.</returns>
+		public static string[] ParseParamValues(object[] ArgList, string Param)
+		{
+			string ParamEquals = Param;
+			if (!ParamEquals.EndsWith("="))
+			{
+				ParamEquals += "=";
+			}
+
+			List<string> Values = new List<string>();
+			foreach (object Arg in ArgList)
+			{
+				string ArgStr = Arg.ToString();
+				if (ArgStr.StartsWith(ParamEquals, StringComparison.InvariantCultureIgnoreCase))
+				{
+					Values.Add(ArgStr.Substring(ParamEquals.Length));
+				}
+			}
+			return Values.ToArray();
+		}
+
+		/// <summary>
+		/// Parses the argument list for any number of parameters.
+		/// </summary>
+		/// <param name="ArgList">Argument list.</param>
+		/// <param name="Param">Param to read its value.</param>
+		/// <returns>Returns an array of values for this parameter (or an empty array if one was not found.</returns>
+		public string[] ParseParamValues(string Param)
+		{
+			return ParseParamValues(Params, Param);
+		}
+
 		/// <summary>
 		/// Parses the command's Params list for a parameter and reads its value. 
 		/// Ex. ParseParamValue(Args, "map=")
diff --git a/Engine/Source/Programs/AutomationTool/AutomationUtils/P4Utils.cs b/Engine/Source/Programs/AutomationTool/AutomationUtils/P4Utils.cs
--- a/Engine/Source/Programs/AutomationTool/AutomationUtils/P4Utils.cs
+++ b/Engine/Source/Programs/AutomationTool/AutomationUtils/P4Utils.cs
@@ -2595,11 +2595,18 @@ namespace AutomationTool
             SpecInput += "Options: " + ClientSpec.Options.ToString().ToLowerInvariant().Replace(",", "") + Environment.NewLine;
             SpecInput += "SubmitOptions: " + ClientSpec.SubmitOptions.ToString().ToLowerInvariant().Replace(",", "") + Environment.NewLine;
             SpecInput += "LineEnd: " + ClientSpec.LineEnd.ToString().ToLowerInvariant() + Environment.NewLine;
-            SpecInput += "View:" + Environment.NewLine;
-            foreach (var Mapping in ClientSpec.View)
-            {
-                SpecInput += "\t" + Mapping.Key + " //" + ClientSpec.Name + Mapping.Value + Environment.NewLine;
-            }
+			if(ClientSpec.Stream != null)
+			{
+				SpecInput += "Stream: " + ClientSpec.Stream + Environment.NewLine;
+			}
+			else
+			{
+				SpecInput += "View:" + Environment.NewLine;
+				foreach (var Mapping in ClientSpec.View)
+				{
+					SpecInput += "\t" + Mapping.Key + " //" + ClientSpec.Name + Mapping.Value + Environment.NewLine;
+				}
+			}
 			if (AllowSpew) CommandUtils.LogLog(SpecInput);
             LogP4("client -i", SpecInput, AllowSpew: AllowSpew, WithClient: false);
             return ClientSpec;
