Fix temporary P4 clients used in marketplace and Rocket promotion scripts to clear stream when cloning the builder's default client and manually setting view.

diff --git a/Engine/Source/Programs/AutomationTool/Scripts/Tests.Automation.cs b/Engine/Source/Programs/AutomationTool/Scripts/Tests.Automation.cs
--- a/Engine/Source/Programs/AutomationTool/Scripts/Tests.Automation.cs
+++ b/Engine/Source/Programs/AutomationTool/Scripts/Tests.Automation.cs
@@ -266,6 +266,7 @@ class GitPullRequest : BuildCommand
 		NewClient.Name = TestClient;
 		NewClient.View = new List<KeyValuePair<string, string>>();
 		NewClient.View.Add(new KeyValuePair<string, string>(Depot + "/...", "/..."));
+		NewClient.Stream = null;
 		if (!P4.DoesClientExist(TestClient))
 		{
 			P4.CreateClient(NewClient);
