Added a dependency on Tools to BuildEngineLocalizationNode. It could potentially compile shaders.

diff --git a/Engine/Source/Programs/AutomationTool/GUBP/LegacyNodes.cs b/Engine/Source/Programs/AutomationTool/GUBP/LegacyNodes.cs
--- a/Engine/Source/Programs/AutomationTool/GUBP/LegacyNodes.cs
+++ b/Engine/Source/Programs/AutomationTool/GUBP/LegacyNodes.cs
@@ -2654,6 +2654,7 @@ partial class GUBP
 			//LocalizationBranchSuffix = InLocalizationBranchSuffix;
 
 			AddDependency(RootEditorNode.StaticGetFullName(HostPlatform));
+            AddDependency(ToolsNode.StaticGetFullName(HostPlatform));
 		}
 
 		public override void DoBuild(GUBP bp)
