Change launcher promotion script so that we ship new bootstrap launcher executable, and have the ability to promote 32-bit launcher builds.

diff --git a/Engine/Source/Programs/AutomationTool/AutomationUtils/MCPPublic.cs b/Engine/Source/Programs/AutomationTool/AutomationUtils/MCPPublic.cs
--- a/Engine/Source/Programs/AutomationTool/AutomationUtils/MCPPublic.cs
+++ b/Engine/Source/Programs/AutomationTool/AutomationUtils/MCPPublic.cs
@@ -62,26 +62,31 @@ namespace EpicGames.MCP.Automation
         }
     }
 
-    /// <summary>
-    /// Enum that defines the MCP backend-compatible platform
-    /// </summary>
-    public enum MCPPlatform
-    {
-        /// <summary>
-        /// MCP doesn't care about Win32 vs. Win64
-        /// </summary>
-        Windows,
+	/// <summary>
+	/// Enum that defines the MCP backend-compatible platform
+	/// </summary>
+	public enum MCPPlatform
+	{
+		/// <summary>
+		/// MCP uses Windows for Win64
+		/// </summary>
+		Windows,
 
-        /// <summary>
-        /// Mac platform.
-        /// </summary>
-        Mac,
+		/// <summary>
+		/// 32 bit Windows
+		/// </summary>
+		Win32,
+
+		/// <summary>
+		/// Mac platform.
+		/// </summary>
+		Mac,
 
 		/// <summary>
 		/// Linux platform.
 		/// </summary>
 		Linux
-    }
+}
 
     /// <summary>
     /// Enum that defines CDN types
@@ -168,48 +173,56 @@ namespace EpicGames.MCP.Automation
 		/// </summary>
 		private readonly string _ManifestFilename;
 
-        /// <summary>
-        /// Determine the platform name (Win32/64 becomes Windows, Mac is Mac, the rest we don't currently understand)
-        /// </summary>
-        static public MCPPlatform ToMCPPlatform(UnrealTargetPlatform TargetPlatform)
-        {
-            if (TargetPlatform != UnrealTargetPlatform.Win64 && TargetPlatform != UnrealTargetPlatform.Win32 && TargetPlatform != UnrealTargetPlatform.Mac && TargetPlatform != UnrealTargetPlatform.Linux)
-            {
-                throw new AutomationException("Platform {0} is not properly supported by the MCP backend yet", TargetPlatform);
-            }
+		/// <summary>
+		/// Determine the platform name
+		/// </summary>
+		static public MCPPlatform ToMCPPlatform(UnrealTargetPlatform TargetPlatform)
+		{
+			if (TargetPlatform != UnrealTargetPlatform.Win64 && TargetPlatform != UnrealTargetPlatform.Win32 && TargetPlatform != UnrealTargetPlatform.Mac && TargetPlatform != UnrealTargetPlatform.Linux)
+			{
+				throw new AutomationException("Platform {0} is not properly supported by the MCP backend yet", TargetPlatform);
+			}
 
-			if (TargetPlatform == UnrealTargetPlatform.Win64 || TargetPlatform == UnrealTargetPlatform.Win32)
+			if (TargetPlatform == UnrealTargetPlatform.Win64)
 			{
 				return MCPPlatform.Windows;
 			}
+			else if (TargetPlatform == UnrealTargetPlatform.Win32)
+			{
+				return MCPPlatform.Win32;
+			}
 			else if (TargetPlatform == UnrealTargetPlatform.Mac)
 			{
 				return MCPPlatform.Mac;
 			}
 
 			return MCPPlatform.Linux;
-        }
-        /// <summary>
-        /// Determine the platform name (Win32/64 becomes Windows, Mac is Mac, the rest we don't currently understand)
-        /// </summary>
-        static public UnrealTargetPlatform FromMCPPlatform(MCPPlatform TargetPlatform)
-        {
-            if (TargetPlatform != MCPPlatform.Windows && TargetPlatform != MCPPlatform.Mac && TargetPlatform != MCPPlatform.Linux)
-            {
-                throw new AutomationException("Platform {0} is not properly supported by the MCP backend yet", TargetPlatform);
-            }
+		}
+		/// <summary>
+		/// Determine the platform name
+		/// </summary>
+		static public UnrealTargetPlatform FromMCPPlatform(MCPPlatform TargetPlatform)
+		{
+			if (TargetPlatform != MCPPlatform.Windows && TargetPlatform != MCPPlatform.Win32 && TargetPlatform != MCPPlatform.Mac && TargetPlatform != MCPPlatform.Linux)
+			{
+				throw new AutomationException("Platform {0} is not properly supported by the MCP backend yet", TargetPlatform);
+			}
 
 			if (TargetPlatform == MCPPlatform.Windows)
 			{
 				return UnrealTargetPlatform.Win64;
 			}
+			else if (TargetPlatform == MCPPlatform.Win32)
+			{
+				return UnrealTargetPlatform.Win32;
+			}
 			else if (TargetPlatform == MCPPlatform.Mac)
 			{
 				return UnrealTargetPlatform.Mac;
 			}
 
 			return UnrealTargetPlatform.Linux;
-        }
+		}
         /// <summary>
         /// Returns the build root path (P:\Builds on build machines usually)
         /// </summary>
@@ -347,7 +360,7 @@ namespace EpicGames.MCP.Automation
             /// </summary>
             public ChunkType AppChunkType;
             /// <summary>
-            /// Matches the corresponding BuildPatchTool command line argument.
+            /// Used as part of manifest filename and build version strings.
             /// </summary>
             public MCPPlatform Platform;
             /// <summary>
