Fix for 'launch on' not working due to stale code from other projects preventing UAT from compiling. Editor passes -ScriptsForProject= argument to UAT to limit the scripts which are compiled and loaded, but a stale AutomationTool.csproj.References file can bring in references to other projects. Disable the include using a property passed on the command line from RunUAT.bat.

diff --git a/Engine/Build/BatchFiles/RunUAT.bat b/Engine/Build/BatchFiles/RunUAT.bat
--- a/Engine/Build/BatchFiles/RunUAT.bat
+++ b/Engine/Build/BatchFiles/RunUAT.bat
@@ -80,7 +80,7 @@ goto DoRunUAT
 :ReadyToCompile
 msbuild /nologo /verbosity:quiet Source\Programs\AutomationToolLauncher\AutomationToolLauncher.csproj /property:Configuration=Development /property:Platform=AnyCPU
 if not %ERRORLEVEL% == 0 goto Error_UATCompileFailed
-msbuild /nologo /verbosity:quiet Source\Programs\AutomationTool\AutomationTool.csproj /property:Configuration=Development /property:Platform=AnyCPU
+msbuild /nologo /verbosity:quiet Source\Programs\AutomationTool\AutomationTool.csproj /property:Configuration=Development /property:Platform=AnyCPU /property:AutomationToolProjectOnly=true
 if not %ERRORLEVEL% == 0 goto Error_UATCompileFailed
 
 
diff --git a/Engine/Source/Programs/AutomationTool/AutomationTool.csproj b/Engine/Source/Programs/AutomationTool/AutomationTool.csproj
--- a/Engine/Source/Programs/AutomationTool/AutomationTool.csproj
+++ b/Engine/Source/Programs/AutomationTool/AutomationTool.csproj
@@ -128,7 +128,7 @@
     <WCFMetadata Include="Service References\" />
   </ItemGroup>
   <Import Project="$(MSBuildToolsPath)\Microsoft.CSharp.targets" />
-  <Import Project="AutomationTool.csproj.References" Condition="Exists('AutomationTool.csproj.References')" />
+  <Import Project="AutomationTool.csproj.References" Condition="Exists('AutomationTool.csproj.References') And '$(AutomationToolProjectOnly)' == ''" />
   <!-- To modify your build process, add your task inside one of the targets below and uncomment it. 
        Other similar extension points exist, see Microsoft.Common.targets.
   <Target Name="BeforeBuild">
