Copying from //UE4/Orion-Staging (originated from //Orion/Main @ CL-2766367)

diff --git a/Engine/Config/BaseEngine.ini b/Engine/Config/BaseEngine.ini
--- a/Engine/Config/BaseEngine.ini
+++ b/Engine/Config/BaseEngine.ini
@@ -1374,6 +1374,9 @@ bUseStreamingPause=false
 +KnownMissingPackages=(PackageName="/Engine/EngineAnimNotifies/AnimNotify_PlayParticleEffect")
 +KnownMissingPackages=(PackageName="/Engine/EngineAnimNotifies/AnimNotify_PlaySound")
 
+;Sequencer Redirects
++ActiveClassRedirects=(OldClassName="MovieSceneMaterialParameterSection",NewClassName="/Script/MovieSceneTracks.MovieSceneParameterSection")
+
 [CoreUObject.Metadata]
 ; Note: UnrealHeaderTool should be rerun after making changes to MetadataRedirects to catch any keys specified in class headers
 +MetadataRedirects=(OldKey="K2Protected", NewKey="BlueprintProtected")
diff --git a/Engine/Shaders/ShadingModels.usf b/Engine/Shaders/ShadingModels.usf
--- a/Engine/Shaders/ShadingModels.usf
+++ b/Engine/Shaders/ShadingModels.usf
@@ -705,7 +705,7 @@ float3 EyeSubsurfaceShading( FGBufferData GBuffer, float3 L, float3 V, half3 N )
 
 	float NoL = saturate( dot( IrisNormal, L ) );
 	float Power = lerp( 12, 1, NoL );
-	float Caustic = 0.5 + 0.5 * ( Power + 1 ) * pow( saturate( dot( CausticNormal, L ) ), Power );
+	float Caustic = 0.6 + 0.2 * ( Power + 1 ) * pow( saturate( dot( CausticNormal, L ) ), Power );
 	float Iris = NoL * Caustic;
 
 	// http://blog.stevemcauley.com/2011/12/03/energy-conserving-wrapped-diffuse/
diff --git a/Engine/Source/Editor/BlueprintGraph/Classes/K2Node_FunctionEntry.h b/Engine/Source/Editor/BlueprintGraph/Classes/K2Node_FunctionEntry.h
--- a/Engine/Source/Editor/BlueprintGraph/Classes/K2Node_FunctionEntry.h
+++ b/Engine/Source/Editor/BlueprintGraph/Classes/K2Node_FunctionEntry.h
@@ -72,7 +72,7 @@ class UK2Node_FunctionEntry : public UK2Node_FunctionTerminator
 	/** Set the extra flags on this node */
 	BLUEPRINTGRAPH_API void SetExtraFlags(int32 InFlags)
 	{
-		ExtraFlags = InFlags;
+		ExtraFlags = (InFlags & ~FUNC_Native);
 	}
 
 protected:
diff --git a/Engine/Source/Editor/BlueprintGraph/Private/K2Node_FunctionEntry.cpp b/Engine/Source/Editor/BlueprintGraph/Private/K2Node_FunctionEntry.cpp
--- a/Engine/Source/Editor/BlueprintGraph/Private/K2Node_FunctionEntry.cpp
+++ b/Engine/Source/Editor/BlueprintGraph/Private/K2Node_FunctionEntry.cpp
@@ -126,6 +126,9 @@ void UK2Node_FunctionEntry::Serialize(FArchive& Ar)
 			// Allow legacy implementations to violate const-correctness
 			bEnforceConstCorrectness = false;
 		}
+
+		// @TODO: gate this with a version check once it makes its way into main
+		ExtraFlags &= ~FUNC_Native;
 	}
 }
 
diff --git a/Engine/Source/Editor/MeshPaint/Private/MeshPaintStaticMeshAdapter.h b/Engine/Source/Editor/MeshPaint/Private/MeshPaintStaticMeshAdapter.h
--- a/Engine/Source/Editor/MeshPaint/Private/MeshPaintStaticMeshAdapter.h
+++ b/Engine/Source/Editor/MeshPaint/Private/MeshPaintStaticMeshAdapter.h
@@ -20,7 +20,7 @@ public:
 	virtual int32 GetNumTexCoords() const override;
 	virtual void GetTriangleInfo(int32 TriIndex, struct FTexturePaintTriangleInfo& OutTriInfo) const override;
 	virtual bool SupportsTexturePaint() const override { return true; }
-	virtual bool SupportsVertexPaint() const override { return true; }
+	virtual bool SupportsVertexPaint() const override { return StaticMeshComponent && !StaticMeshComponent->bDisallowMeshPaintPerInstance; }
 	virtual bool LineTraceComponent(struct FHitResult& OutHit, const FVector Start, const FVector End, const struct FCollisionQueryParams& Params) const override;
 	virtual void SphereIntersectTriangles(TArray<int32>& OutTriangles, const float ComponentSpaceSquaredBrushRadius, const FVector& ComponentSpaceBrushPosition) const override;
 	virtual void QueryPaintableTextures(int32 MaterialIndex, int32& OutDefaultIndex, TArray<struct FPaintableTexture>& InOutTextureList) override;
diff --git a/Engine/Source/Editor/MovieSceneTools/Private/CurveKeyAreas/FloatCurveKeyArea.cpp b/Engine/Source/Editor/MovieSceneTools/Private/CurveKeyAreas/FloatCurveKeyArea.cpp
--- a/Engine/Source/Editor/MovieSceneTools/Private/CurveKeyAreas/FloatCurveKeyArea.cpp
+++ b/Engine/Source/Editor/MovieSceneTools/Private/CurveKeyAreas/FloatCurveKeyArea.cpp
@@ -62,6 +62,21 @@ TArray<FKeyHandle> FFloatCurveKeyArea::AddKeyUnique(float Time, EMovieSceneKeyIn
 	return AddedKeyHandles;
 }
 
+TOptional<FKeyHandle> FFloatCurveKeyArea::DuplicateKey(FKeyHandle KeyToDuplicate)
+{
+	if (!Curve->IsKeyHandleValid(KeyToDuplicate))
+	{
+		return TOptional<FKeyHandle>();
+	}
+
+	FRichCurveKey ThisKey = Curve->GetKey(KeyToDuplicate);
+
+	FKeyHandle KeyHandle = Curve->AddKey(GetKeyTime(KeyToDuplicate), ThisKey.Value);
+	// Ensure the rest of the key properties are added
+	Curve->GetKey(KeyHandle) = ThisKey;
+
+	return KeyHandle;
+}
 
 bool FFloatCurveKeyArea::CanCreateKeyEditor()
 {
diff --git a/Engine/Source/Editor/MovieSceneTools/Private/CurveKeyAreas/IntegralKeyArea.cpp b/Engine/Source/Editor/MovieSceneTools/Private/CurveKeyAreas/IntegralKeyArea.cpp
--- a/Engine/Source/Editor/MovieSceneTools/Private/CurveKeyAreas/IntegralKeyArea.cpp
+++ b/Engine/Source/Editor/MovieSceneTools/Private/CurveKeyAreas/IntegralKeyArea.cpp
@@ -31,6 +31,15 @@ TArray<FKeyHandle> FIntegralCurveKeyAreaBase::AddKeyUnique(float Time, EMovieSce
 	return AddedKeyHandles;
 }
 
+TOptional<FKeyHandle> FIntegralCurveKeyAreaBase::DuplicateKey(FKeyHandle KeyToDuplicate)
+{
+	if (!Curve.IsKeyHandleValid(KeyToDuplicate))
+	{
+		return TOptional<FKeyHandle>();
+	}
+
+	return Curve.AddKey(GetKeyTime(KeyToDuplicate), Curve.GetKey(KeyToDuplicate).Value);
+}
 
 void FIntegralCurveKeyAreaBase::DeleteKey(FKeyHandle KeyHandle)
 {
diff --git a/Engine/Source/Editor/MovieSceneTools/Private/CurveKeyAreas/NameCurveKeyArea.cpp b/Engine/Source/Editor/MovieSceneTools/Private/CurveKeyAreas/NameCurveKeyArea.cpp
--- a/Engine/Source/Editor/MovieSceneTools/Private/CurveKeyAreas/NameCurveKeyArea.cpp
+++ b/Engine/Source/Editor/MovieSceneTools/Private/CurveKeyAreas/NameCurveKeyArea.cpp
@@ -31,6 +31,15 @@ TArray<FKeyHandle> FNameCurveKeyArea::AddKeyUnique(float Time, EMovieSceneKeyInt
 	return AddedKeyHandles;
 }
 
+TOptional<FKeyHandle> FNameCurveKeyArea::DuplicateKey(FKeyHandle KeyToDuplicate)
+{
+	if (!Curve.IsKeyHandleValid(KeyToDuplicate))
+	{
+		return TOptional<FKeyHandle>();
+	}
+
+	return Curve.AddKey(GetKeyTime(KeyToDuplicate), Curve.GetKey(KeyToDuplicate).Value);
+}
 
 bool FNameCurveKeyArea::CanCreateKeyEditor()
 {
diff --git a/Engine/Source/Editor/MovieSceneTools/Private/MovieSceneToolsModule.cpp b/Engine/Source/Editor/MovieSceneTools/Private/MovieSceneToolsModule.cpp
--- a/Engine/Source/Editor/MovieSceneTools/Private/MovieSceneToolsModule.cpp
+++ b/Engine/Source/Editor/MovieSceneTools/Private/MovieSceneToolsModule.cpp
@@ -24,6 +24,7 @@
 #include "AudioTrackEditor.h"
 #include "SkeletalAnimationTrackEditor.h"
 #include "ParticleTrackEditor.h"
+#include "ParticleParameterTrackEditor.h"
 #include "AttachTrackEditor.h"
 #include "EventTrackEditor.h"
 #include "PathTrackEditor.h"
@@ -65,6 +66,7 @@ public:
 		AudioTrackCreateEditorHandle = SequencerModule.RegisterTrackEditor_Handle( FOnCreateTrackEditor::CreateStatic( &FAudioTrackEditor::CreateTrackEditor ) );
 		EventTrackCreateEditorHandle = SequencerModule.RegisterTrackEditor_Handle( FOnCreateTrackEditor::CreateStatic( &FEventTrackEditor::CreateTrackEditor ) );
 		ParticleTrackCreateEditorHandle = SequencerModule.RegisterTrackEditor_Handle( FOnCreateTrackEditor::CreateStatic( &FParticleTrackEditor::CreateTrackEditor ) );
+		ParticleParameterTrackCreateEditorHandle = SequencerModule.RegisterTrackEditor_Handle( FOnCreateTrackEditor::CreateStatic( &FParticleParameterTrackEditor::CreateTrackEditor ) );
 		PathTrackCreateEditorHandle = SequencerModule.RegisterTrackEditor_Handle( FOnCreateTrackEditor::CreateStatic( &F3DPathTrackEditor::CreateTrackEditor ) );
 		ShotTrackCreateEditorHandle = SequencerModule.RegisterTrackEditor_Handle( FOnCreateTrackEditor::CreateStatic( &FShotTrackEditor::CreateTrackEditor ) );
 		SlomoTrackCreateEditorHandle = SequencerModule.RegisterTrackEditor_Handle( FOnCreateTrackEditor::CreateStatic( &FSlomoTrackEditor::CreateTrackEditor ) );
@@ -100,6 +102,7 @@ public:
 		SequencerModule.UnRegisterTrackEditor_Handle( AudioTrackCreateEditorHandle );
 		SequencerModule.UnRegisterTrackEditor_Handle( EventTrackCreateEditorHandle );
 		SequencerModule.UnRegisterTrackEditor_Handle( ParticleTrackCreateEditorHandle );
+		SequencerModule.UnRegisterTrackEditor_Handle( ParticleParameterTrackCreateEditorHandle );
 		SequencerModule.UnRegisterTrackEditor_Handle( PathTrackCreateEditorHandle );
 		SequencerModule.UnRegisterTrackEditor_Handle( ShotTrackCreateEditorHandle );
 		SequencerModule.UnRegisterTrackEditor_Handle( SlomoTrackCreateEditorHandle );
@@ -159,6 +162,7 @@ private:
 	FDelegateHandle AudioTrackCreateEditorHandle;
 	FDelegateHandle EventTrackCreateEditorHandle;
 	FDelegateHandle ParticleTrackCreateEditorHandle;
+	FDelegateHandle ParticleParameterTrackCreateEditorHandle;
 	FDelegateHandle PathTrackCreateEditorHandle;
 	FDelegateHandle ShotTrackCreateEditorHandle;
 	FDelegateHandle SlomoTrackCreateEditorHandle;
diff --git a/Engine/Source/Editor/MovieSceneTools/Private/Sections/MaterialParameterSection.cpp b/Engine/Source/Editor/MovieSceneTools/Private/Sections/MaterialParameterSection.cpp
deleted file mode 100644
diff --git a/Engine/Source/Editor/MovieSceneTools/Private/Sections/MaterialParameterSection.h b/Engine/Source/Editor/MovieSceneTools/Private/Sections/MaterialParameterSection.h
deleted file mode 100644
diff --git a/Engine/Source/Editor/MovieSceneTools/Private/Sections/ParameterSection.cpp b/Engine/Source/Editor/MovieSceneTools/Private/Sections/ParameterSection.cpp
new file mode 100644
--- /dev/null
+++ b/Engine/Source/Editor/MovieSceneTools/Private/Sections/ParameterSection.cpp
@@ -0,0 +1,141 @@
+// Copyright 1998-2015 Epic Games, Inc. All Rights Reserved.
+
+#include "MovieSceneToolsPrivatePCH.h"
+#include "ParameterSection.h"
+#include "MovieSceneParameterSection.h"
+#include "FloatCurveKeyArea.h"
+
+
+#define LOCTEXT_NAMESPACE "ParameterSection"
+
+DECLARE_DELEGATE(FLayoutParameterDelegate);
+
+
+struct FIndexAndLayoutParameterDelegate
+{
+	int32 Index;
+	FLayoutParameterDelegate LayoutParameterDelegate;
+
+	bool operator<(FIndexAndLayoutParameterDelegate const& Other) const
+	{
+		return Index < Other.Index;
+	}
+};
+
+
+void LayoutScalarParameter( ISectionLayoutBuilder* LayoutBuilder, FScalarParameterNameAndCurve* ScalarParameterNameAndCurve, UMovieSceneParameterSection* ParameterSection )
+{
+	LayoutBuilder->AddKeyArea(
+		ScalarParameterNameAndCurve->ParameterName,
+		FText::FromName( ScalarParameterNameAndCurve->ParameterName ),
+		MakeShareable( new FFloatCurveKeyArea( &ScalarParameterNameAndCurve->ParameterCurve, ParameterSection ) ) );
+}
+
+
+void LayoutVectorParameter( ISectionLayoutBuilder* LayoutBuilder, FVectorParameterNameAndCurves* VectorParameterNameAndCurves, UMovieSceneParameterSection* ParameterSection )
+{
+	LayoutBuilder->PushCategory( VectorParameterNameAndCurves->ParameterName, FText::FromName( VectorParameterNameAndCurves->ParameterName ) );
+	{
+		LayoutBuilder->AddKeyArea( "X", NSLOCTEXT( "ParameterSection", "XArea", "X" ),
+			MakeShareable( new FFloatCurveKeyArea( &VectorParameterNameAndCurves->XCurve, ParameterSection ) ) );
+		LayoutBuilder->AddKeyArea( "Y", NSLOCTEXT( "ParameterSection", "YArea", "Y" ), 
+			MakeShareable( new FFloatCurveKeyArea( &VectorParameterNameAndCurves->YCurve, ParameterSection ) ) );
+		LayoutBuilder->AddKeyArea( "Z", NSLOCTEXT( "ParameterSection", "ZArea", "Z" ), 
+			MakeShareable( new FFloatCurveKeyArea( &VectorParameterNameAndCurves->ZCurve, ParameterSection ) ) );
+	}
+	LayoutBuilder->PopCategory();
+}
+
+
+void LayoutColorParameter( ISectionLayoutBuilder* LayoutBuilder, FColorParameterNameAndCurves* ColorParameterNameAndCurves, UMovieSceneParameterSection* ParameterSection )
+{
+	LayoutBuilder->PushCategory( ColorParameterNameAndCurves->ParameterName, FText::FromName( ColorParameterNameAndCurves->ParameterName ) );
+	{
+		LayoutBuilder->AddKeyArea( "R", NSLOCTEXT( "ParameterSection", "RedArea", "Red" ),
+			MakeShareable( new FFloatCurveKeyArea( &ColorParameterNameAndCurves->RedCurve, ParameterSection ) ) );
+		LayoutBuilder->AddKeyArea( "G", NSLOCTEXT( "ParameterSection", "GreenArea", "Green" ),
+			MakeShareable( new FFloatCurveKeyArea( &ColorParameterNameAndCurves->GreenCurve, ParameterSection ) ) );
+		LayoutBuilder->AddKeyArea( "B", NSLOCTEXT( "ParameterSection", "BlueArea", "Blue" ),
+			MakeShareable( new FFloatCurveKeyArea( &ColorParameterNameAndCurves->BlueCurve, ParameterSection ) ) );
+		LayoutBuilder->AddKeyArea( "A", NSLOCTEXT( "ParameterSection", "OpacityArea", "Opacity" ),
+			MakeShareable( new FFloatCurveKeyArea( &ColorParameterNameAndCurves->AlphaCurve, ParameterSection ) ) );
+	}
+	LayoutBuilder->PopCategory();
+}
+
+void FParameterSection::GenerateSectionLayout( class ISectionLayoutBuilder& LayoutBuilder ) const
+{
+	// TODO - This seems a bit convoluted here to have to use delegates to get the layout to happen in the right order based on the indices.
+	UMovieSceneParameterSection* ParameterSection = Cast<UMovieSceneParameterSection>( &SectionObject );
+	TArray<FIndexAndLayoutParameterDelegate> IndexAndLayoutDelegates;
+
+	// Collect scalar parameters.
+	TArray<FScalarParameterNameAndCurve>* ScalarParameterNamesAndCurves = ParameterSection->GetScalarParameterNamesAndCurves();
+	for ( int32 i = 0; i < ScalarParameterNamesAndCurves->Num(); i++)
+	{
+		FIndexAndLayoutParameterDelegate IndexAndLayoutDelegate;
+		IndexAndLayoutDelegate.Index = (*ScalarParameterNamesAndCurves)[i].Index;
+		IndexAndLayoutDelegate.LayoutParameterDelegate.BindStatic(&LayoutScalarParameter, &LayoutBuilder, &(*ScalarParameterNamesAndCurves)[i], ParameterSection);
+		IndexAndLayoutDelegates.Add(IndexAndLayoutDelegate);
+	}
+
+	// Collect vector parameters
+	TArray<FVectorParameterNameAndCurves>* VectorParameterNamesAndCurves = ParameterSection->GetVectorParameterNamesAndCurves();
+	for ( int32 i = 0; i < VectorParameterNamesAndCurves->Num(); i++ )
+	{
+		FIndexAndLayoutParameterDelegate IndexAndLayoutDelegate;
+		IndexAndLayoutDelegate.Index = ( *VectorParameterNamesAndCurves )[i].Index;
+		IndexAndLayoutDelegate.LayoutParameterDelegate.BindStatic( &LayoutVectorParameter, &LayoutBuilder, &( *VectorParameterNamesAndCurves )[i], ParameterSection );
+		IndexAndLayoutDelegates.Add( IndexAndLayoutDelegate );
+	}
+
+	// Collect Color parameters
+	TArray<FColorParameterNameAndCurves>* ColorParameterNamesAndCurves = ParameterSection->GetColorParameterNamesAndCurves();
+	for ( int32 i = 0; i < ColorParameterNamesAndCurves->Num(); i++ )
+	{
+		FIndexAndLayoutParameterDelegate IndexAndLayoutDelegate;
+		IndexAndLayoutDelegate.Index = (*ColorParameterNamesAndCurves)[i].Index;
+		IndexAndLayoutDelegate.LayoutParameterDelegate.BindStatic( &LayoutColorParameter, &LayoutBuilder, &(*ColorParameterNamesAndCurves)[i], ParameterSection );
+		IndexAndLayoutDelegates.Add( IndexAndLayoutDelegate );
+	}
+
+	// Sort and perform layout
+	IndexAndLayoutDelegates.Sort();
+	for ( FIndexAndLayoutParameterDelegate& IndexAndLayoutDelegate : IndexAndLayoutDelegates )
+	{
+		IndexAndLayoutDelegate.LayoutParameterDelegate.Execute();
+	}
+}
+
+
+bool FParameterSection::RequestDeleteCategory( const TArray<FName>& CategoryNamePath )
+{
+	const FScopedTransaction Transaction( LOCTEXT( "DeleteVectorOrColorParameter", "Delete vector or color parameter" ) );
+	UMovieSceneParameterSection* ParameterSection = Cast<UMovieSceneParameterSection>( &SectionObject );
+	if( ParameterSection->Modify() )
+	{
+		bool bVectorParameterDeleted = ParameterSection->RemoveVectorParameter( CategoryNamePath[0] );
+		bool bColorParameterDeleted = ParameterSection->RemoveColorParameter( CategoryNamePath[0] );
+		return bVectorParameterDeleted || bColorParameterDeleted;
+	}
+	return false;
+}
+
+
+bool FParameterSection::RequestDeleteKeyArea( const TArray<FName>& KeyAreaNamePath )
+{
+	// Only handle paths with a single name, in all other cases the user is deleting a component of a vector parameter.
+	if ( KeyAreaNamePath.Num() == 1)
+	{
+		const FScopedTransaction Transaction( LOCTEXT( "DeleteScalarParameter", "Delete scalar parameter" ) );
+		UMovieSceneParameterSection* ParameterSection = Cast<UMovieSceneParameterSection>( &SectionObject );
+		if (ParameterSection->TryModify())
+		{
+			return ParameterSection->RemoveScalarParameter( KeyAreaNamePath[0] );
+		}
+	}
+	return false;
+}
+
+
+#undef LOCTEXT_NAMESPACE
diff --git a/Engine/Source/Editor/MovieSceneTools/Private/Sections/ParameterSection.h b/Engine/Source/Editor/MovieSceneTools/Private/Sections/ParameterSection.h
new file mode 100644
--- /dev/null
+++ b/Engine/Source/Editor/MovieSceneTools/Private/Sections/ParameterSection.h
@@ -0,0 +1,31 @@
+// Copyright 1998-2015 Epic Games, Inc. All Rights Reserved.
+
+#pragma once
+
+
+/**
+ * A movie scene section for material parameters.
+ */
+class FParameterSection
+	: public FPropertySection
+{
+public:
+
+	FParameterSection(UMovieSceneSection& InSectionObject, const FText& SectionName)
+		: FPropertySection(InSectionObject, SectionName)
+	{ }
+
+public:
+
+	// FPropertySection interface
+	virtual void SetIntermediateValue(FPropertyChangedParams PropertyChangedParams) override { }
+	virtual void ClearIntermediateValue() override { }
+
+public:
+
+	// ISequencerSection interface
+
+	virtual void GenerateSectionLayout(class ISectionLayoutBuilder& LayoutBuilder) const override;
+	virtual bool RequestDeleteCategory(const TArray<FName>& CategoryNamePath) override;
+	virtual bool RequestDeleteKeyArea(const TArray<FName>& KeyAreaNamePath) override;
+};
\ No newline at end of file
diff --git a/Engine/Source/Editor/MovieSceneTools/Private/Sections/ShotSequencerSection.cpp b/Engine/Source/Editor/MovieSceneTools/Private/Sections/ShotSequencerSection.cpp
--- a/Engine/Source/Editor/MovieSceneTools/Private/Sections/ShotSequencerSection.cpp
+++ b/Engine/Source/Editor/MovieSceneTools/Private/Sections/ShotSequencerSection.cpp
@@ -155,6 +155,11 @@ void FShotSequencerSection::RegenerateViewportThumbnails(const FIntPoint& Size)
 /* ISequencerSection interface
  *****************************************************************************/
 
+bool FShotSequencerSection::ShouldDrawKeyAreaBackground() const
+{
+	return false;
+}
+
 bool FShotSequencerSection::AreSectionsConnected() const
 {
 	return true;
diff --git a/Engine/Source/Editor/MovieSceneTools/Private/Sections/ShotSequencerSection.h b/Engine/Source/Editor/MovieSceneTools/Private/Sections/ShotSequencerSection.h
--- a/Engine/Source/Editor/MovieSceneTools/Private/Sections/ShotSequencerSection.h
+++ b/Engine/Source/Editor/MovieSceneTools/Private/Sections/ShotSequencerSection.h
@@ -52,6 +52,7 @@ public:
 
 	// ISequencerSection interface
 
+	virtual bool ShouldDrawKeyAreaBackground() const override;
 	virtual bool AreSectionsConnected() const override;
 	virtual void BuildSectionContextMenu(FMenuBuilder& MenuBuilder) override;
 	virtual void GenerateSectionLayout(ISectionLayoutBuilder& LayoutBuilder) const override { }
diff --git a/Engine/Source/Editor/MovieSceneTools/Private/TrackEditors/AudioTrackEditor.cpp b/Engine/Source/Editor/MovieSceneTools/Private/TrackEditors/AudioTrackEditor.cpp
--- a/Engine/Source/Editor/MovieSceneTools/Private/TrackEditors/AudioTrackEditor.cpp
+++ b/Engine/Source/Editor/MovieSceneTools/Private/TrackEditors/AudioTrackEditor.cpp
@@ -344,6 +344,10 @@ UMovieSceneSection* FAudioSection::GetSectionObject()
 	return &Section;
 }
 
+bool FAudioSection::ShouldDrawKeyAreaBackground() const
+{
+	return false;
+}
 
 FText FAudioSection::GetDisplayName() const
 {
diff --git a/Engine/Source/Editor/MovieSceneTools/Private/TrackEditors/AudioTrackEditor.h b/Engine/Source/Editor/MovieSceneTools/Private/TrackEditors/AudioTrackEditor.h
--- a/Engine/Source/Editor/MovieSceneTools/Private/TrackEditors/AudioTrackEditor.h
+++ b/Engine/Source/Editor/MovieSceneTools/Private/TrackEditors/AudioTrackEditor.h
@@ -73,6 +73,7 @@ public:
 	// ISequencerSection interface
 
 	virtual UMovieSceneSection* GetSectionObject() override;
+	virtual bool ShouldDrawKeyAreaBackground() const override;
 	virtual FText GetDisplayName() const override;
 	virtual FText GetSectionTitle() const override;
 	virtual float GetSectionHeight() const override;
diff --git a/Engine/Source/Editor/MovieSceneTools/Private/TrackEditors/MaterialTrackEditor.cpp b/Engine/Source/Editor/MovieSceneTools/Private/TrackEditors/MaterialTrackEditor.cpp
--- a/Engine/Source/Editor/MovieSceneTools/Private/TrackEditors/MaterialTrackEditor.cpp
+++ b/Engine/Source/Editor/MovieSceneTools/Private/TrackEditors/MaterialTrackEditor.cpp
@@ -3,8 +3,8 @@
 #include "MovieSceneToolsPrivatePCH.h"
 #include "MaterialTrackEditor.h"
 #include "MovieSceneMaterialTrack.h"
-#include "MaterialParameterSection.h"
-#include "MovieSceneMaterialParameterSection.h"
+#include "ParameterSection.h"
+#include "MovieSceneParameterSection.h"
 
 
 #define LOCTEXT_NAMESPACE "MaterialTrackEditor"
@@ -18,10 +18,10 @@ FMaterialTrackEditor::FMaterialTrackEditor( TSharedRef<ISequencer> InSequencer )
 
 TSharedRef<ISequencerSection> FMaterialTrackEditor::MakeSectionInterface( UMovieSceneSection& SectionObject, UMovieSceneTrack& Track )
 {
-	UMovieSceneMaterialParameterSection* ParameterSection = Cast<UMovieSceneMaterialParameterSection>(&SectionObject);
+	UMovieSceneParameterSection* ParameterSection = Cast<UMovieSceneParameterSection>(&SectionObject);
 	checkf( ParameterSection != nullptr, TEXT("Unsupported section type.") );
 
-	return MakeShareable(new FMaterialParameterSection( *ParameterSection, FText::FromName(ParameterSection->GetFName())));
+	return MakeShareable(new FParameterSection( *ParameterSection, FText::FromName(ParameterSection->GetFName())));
 }
 
 
@@ -96,14 +96,14 @@ TSharedRef<SWidget> FMaterialTrackEditor::OnGetAddParameterMenuContent( FGuid Ob
 			ParameterNamesAndActions.Add(NameAndAction);
 		}
 
-		// Collect vector parameters.
-		TArray<FName> VectorParameterNames;
-		TArray<FGuid> VectorParmeterGuids;
-		Material->GetAllVectorParameterNames( VectorParameterNames, VectorParmeterGuids );
-		for ( const FName& VectorParameterName : VectorParameterNames )
+		// Collect color parameters.
+		TArray<FName> ColorParameterNames;
+		TArray<FGuid> ColorParmeterGuids;
+		Material->GetAllVectorParameterNames( ColorParameterNames, ColorParmeterGuids );
+		for ( const FName& ColorParameterName : ColorParameterNames )
 		{
-			FUIAction AddParameterMenuAction( FExecuteAction::CreateSP( this, &FMaterialTrackEditor::AddVectorParameter, ObjectBinding, MaterialTrack, VectorParameterName ) );
-			FParameterNameAndAction NameAndAction( VectorParameterName, AddParameterMenuAction );
+			FUIAction AddParameterMenuAction( FExecuteAction::CreateSP( this, &FMaterialTrackEditor::AddColorParameter, ObjectBinding, MaterialTrack, ColorParameterName ) );
+			FParameterNameAndAction NameAndAction( ColorParameterName, AddParameterMenuAction );
 			ParameterNamesAndActions.Add( NameAndAction );
 		}
 
@@ -160,7 +160,7 @@ void FMaterialTrackEditor::AddScalarParameter( FGuid ObjectBinding, UMovieSceneM
 }
 
 
-void FMaterialTrackEditor::AddVectorParameter( FGuid ObjectBinding, UMovieSceneMaterialTrack* MaterialTrack, FName ParameterName )
+void FMaterialTrackEditor::AddColorParameter( FGuid ObjectBinding, UMovieSceneMaterialTrack* MaterialTrack, FName ParameterName )
 {
 	UMovieSceneSequence* MovieSceneSequence = GetMovieSceneSequence();
 	float KeyTime = GetTimeForKey( MovieSceneSequence );
@@ -172,7 +172,7 @@ void FMaterialTrackEditor::AddVectorParameter( FGuid ObjectBinding, UMovieSceneM
 		FLinearColor ParameterValue;
 		Material->GetVectorParameterValue( ParameterName, ParameterValue );
 		MaterialTrack->Modify();
-		MaterialTrack->AddVectorParameterKey( ParameterName, KeyTime, ParameterValue );
+		MaterialTrack->AddColorParameterKey( ParameterName, KeyTime, ParameterValue );
 	}
 	NotifyMovieSceneDataChanged();
 }
diff --git a/Engine/Source/Editor/MovieSceneTools/Private/TrackEditors/MaterialTrackEditor.h b/Engine/Source/Editor/MovieSceneTools/Private/TrackEditors/MaterialTrackEditor.h
--- a/Engine/Source/Editor/MovieSceneTools/Private/TrackEditors/MaterialTrackEditor.h
+++ b/Engine/Source/Editor/MovieSceneTools/Private/TrackEditors/MaterialTrackEditor.h
@@ -47,12 +47,12 @@ private:
 	 */
 	void AddScalarParameter( FGuid ObjectBinding, UMovieSceneMaterialTrack* MaterialTrack, FName ParameterName );
 
-	/** Adds a vector parameter and initial key to a material track.
+	/** Adds a color parameter and initial key to a material track.
 	* @param ObjectBinding The object binding which owns the material track.
 	* @param MaterialTrack The track to Add the section to.
 	* @param ParameterName The name of the parameter to add an initial key for.
 	*/
-	void AddVectorParameter( FGuid ObjectBinding, UMovieSceneMaterialTrack* MaterialTrack, FName ParameterName );
+	void AddColorParameter( FGuid ObjectBinding, UMovieSceneMaterialTrack* MaterialTrack, FName ParameterName );
 };
 
 
diff --git a/Engine/Source/Editor/MovieSceneTools/Private/TrackEditors/ParticleParameterTrackEditor.cpp b/Engine/Source/Editor/MovieSceneTools/Private/TrackEditors/ParticleParameterTrackEditor.cpp
new file mode 100644
--- /dev/null
+++ b/Engine/Source/Editor/MovieSceneTools/Private/TrackEditors/ParticleParameterTrackEditor.cpp
@@ -0,0 +1,231 @@
+// Copyright 1998-2015 Epic Games, Inc. All Rights Reserved.
+
+#include "MovieSceneToolsPrivatePCH.h"
+#include "ParticleParameterTrackEditor.h"
+#include "MovieSceneParticleParameterTrack.h"
+#include "ParameterSection.h"
+#include "MovieSceneParameterSection.h"
+
+
+#define LOCTEXT_NAMESPACE "ParticleParameterTrackEditor"
+
+FName FParticleParameterTrackEditor::TrackName( "ParticleParameter" );
+
+FParticleParameterTrackEditor::FParticleParameterTrackEditor( TSharedRef<ISequencer> InSequencer )
+	: FMovieSceneTrackEditor( InSequencer )
+{
+}
+
+
+TSharedRef<ISequencerTrackEditor> FParticleParameterTrackEditor::CreateTrackEditor( TSharedRef<ISequencer> OwningSequencer )
+{
+	return MakeShareable( new FParticleParameterTrackEditor( OwningSequencer ) );
+}
+
+
+TSharedRef<ISequencerSection> FParticleParameterTrackEditor::MakeSectionInterface( UMovieSceneSection& SectionObject, UMovieSceneTrack& Track )
+{
+	UMovieSceneParameterSection* ParameterSection = Cast<UMovieSceneParameterSection>(&SectionObject);
+	checkf( ParameterSection != nullptr, TEXT("Unsupported section type.") );
+
+	return MakeShareable(new FParameterSection( *ParameterSection, FText::FromName(ParameterSection->GetFName())));
+}
+
+
+TSharedPtr<SWidget> FParticleParameterTrackEditor::BuildOutlinerEditWidget( const FGuid& ObjectBinding, UMovieSceneTrack* Track )
+{
+	UMovieSceneParticleParameterTrack* ParticleParameterTrack = Cast<UMovieSceneParticleParameterTrack>(Track);
+	return
+		SNew( SComboButton )
+		.ButtonStyle( FEditorStyle::Get(), "FlatButton.Light" )
+		.OnGetMenuContent( this, &FParticleParameterTrackEditor::OnGetAddParameterMenuContent, ObjectBinding, ParticleParameterTrack )
+		.ContentPadding( FMargin( 2, 0 ) )
+		.ButtonContent()
+		[
+			SNew( SHorizontalBox )
+
+			+ SHorizontalBox::Slot()
+			.VAlign( VAlign_Center )
+			.AutoWidth()
+			[
+				SNew( STextBlock )
+				.Font( FEditorStyle::Get().GetFontStyle( "FontAwesome.8" ) )
+				.Text( FText::FromString( FString( TEXT( "\xf067" ) ) ) /*fa-plus*/ )
+			]
+
+			+ SHorizontalBox::Slot()
+			.AutoWidth()
+			.Padding( 4, 0, 0, 0 )
+			[
+				SNew( STextBlock )
+				.Font( FEditorStyle::GetFontStyle( "Sequencer.AnimationOutliner.RegularFont" ) )
+				.Text( LOCTEXT( "AddParameterButton", "Parameter" ) )
+			]
+		];
+}
+
+
+void FParticleParameterTrackEditor::BuildObjectBindingTrackMenu( FMenuBuilder& MenuBuilder, const FGuid& ObjectBinding, const UClass* ObjectClass )
+{
+	if ( ObjectClass->IsChildOf( AEmitter::StaticClass() ) )
+	{
+		const TSharedPtr<ISequencer> ParentSequencer = GetSequencer();
+
+		MenuBuilder.AddMenuEntry(
+			LOCTEXT( "AddParticleParameterTrack", "Particle Parameter Track" ),
+			LOCTEXT( "AddParticleParameterTrackTooltip", "Adds a track for controlling particle parameter values." ),
+			FSlateIcon(),
+			FUIAction
+			(
+				FExecuteAction::CreateSP( this, &FParticleParameterTrackEditor::AddParticleParameterTrack, ObjectBinding ),
+				FCanExecuteAction::CreateSP( this, &FParticleParameterTrackEditor::CanAddParticleParameterTrack, ObjectBinding )
+			));
+	}
+}
+
+
+bool FParticleParameterTrackEditor::SupportsType( TSubclassOf<UMovieSceneTrack> Type ) const
+{
+	return Type == UMovieSceneParticleParameterTrack::StaticClass();
+}
+
+
+TSharedRef<SWidget> FParticleParameterTrackEditor::OnGetAddParameterMenuContent( FGuid ObjectBinding, UMovieSceneParticleParameterTrack* ParticleParameterTrack )
+{
+	FMenuBuilder AddParameterMenuBuilder( true, nullptr );
+
+	AEmitter* Emitter = Cast<AEmitter>( GetSequencer()->GetFocusedMovieSceneSequence()->FindObject( ObjectBinding ) );
+	if ( Emitter != nullptr )
+	{
+		UParticleSystemComponent* ParticleSystemComponent = Emitter->GetParticleSystemComponent();
+		if ( ParticleSystemComponent != nullptr )
+		{
+			TArray<FParameterNameAndAction> ParameterNamesAndActions;
+			const TArray<FParticleSysParam> InstanceParameters = ParticleSystemComponent->GetAsyncInstanceParameters();
+			for ( const FParticleSysParam& ParticleSystemParameter : InstanceParameters )
+			{
+				switch ( ParticleSystemParameter.ParamType )
+				{
+				case PSPT_Scalar:
+				{
+					FUIAction AddParameterMenuAction( FExecuteAction::CreateSP( this, &FParticleParameterTrackEditor::AddScalarParameter, ObjectBinding, ParticleParameterTrack, ParticleSystemParameter.Name ) );
+					FParameterNameAndAction NameAndAction( ParticleSystemParameter.Name, AddParameterMenuAction );
+					ParameterNamesAndActions.Add( NameAndAction );
+					break;
+				}
+				case PSPT_Vector:
+				{
+					FUIAction AddParameterMenuAction( FExecuteAction::CreateSP( this, &FParticleParameterTrackEditor::AddVectorParameter, ObjectBinding, ParticleParameterTrack, ParticleSystemParameter.Name ) );
+					FParameterNameAndAction NameAndAction( ParticleSystemParameter.Name, AddParameterMenuAction );
+					ParameterNamesAndActions.Add( NameAndAction );
+					break;
+				}
+				case PSPT_Color:
+				{
+					FUIAction AddParameterMenuAction( FExecuteAction::CreateSP( this, &FParticleParameterTrackEditor::AddColorParameter, ObjectBinding, ParticleParameterTrack, ParticleSystemParameter.Name ) );
+					FParameterNameAndAction NameAndAction( ParticleSystemParameter.Name, AddParameterMenuAction );
+					ParameterNamesAndActions.Add( NameAndAction );
+					break;
+				}
+				}
+			}
+
+			// Sort and generate menu.
+			ParameterNamesAndActions.Sort();
+			for ( FParameterNameAndAction NameAndAction : ParameterNamesAndActions )
+			{
+				AddParameterMenuBuilder.AddMenuEntry( FText::FromName( NameAndAction.ParameterName ), FText(), FSlateIcon(), NameAndAction.Action );
+			}
+		}
+	}
+
+	return AddParameterMenuBuilder.MakeWidget();
+}
+
+
+bool FParticleParameterTrackEditor::CanAddParticleParameterTrack( FGuid ObjectBinding )
+{
+	return GetSequencer()->GetFocusedMovieSceneSequence()->GetMovieScene()->FindTrack( UMovieSceneParticleParameterTrack::StaticClass(), ObjectBinding, TrackName ) == nullptr;
+}
+
+
+void FParticleParameterTrackEditor::AddParticleParameterTrack( FGuid ObjectBinding )
+{
+	FindOrCreateTrackForObject( ObjectBinding, UMovieSceneParticleParameterTrack::StaticClass(), TrackName, true);
+	NotifyMovieSceneDataChanged();
+}
+
+
+void FParticleParameterTrackEditor::AddScalarParameter( FGuid ObjectBinding, UMovieSceneParticleParameterTrack* ParticleParameterTrack, FName ParameterName )
+{
+	UMovieSceneSequence* MovieSceneSequence = GetMovieSceneSequence();
+	float KeyTime = GetTimeForKey( MovieSceneSequence );
+
+	UObject* Object = GetSequencer()->GetFocusedMovieSceneSequence()->FindObject( ObjectBinding );
+	AEmitter* Emitter = Cast<AEmitter>(Object);
+	if ( Emitter != nullptr )
+	{
+		UParticleSystemComponent* ParticleSystemComponent = Emitter->GetParticleSystemComponent();
+		if ( ParticleSystemComponent != nullptr )
+		{
+			float Value;
+			ParticleSystemComponent->GetFloatParameter( ParameterName, Value );
+
+			const FScopedTransaction Transaction( LOCTEXT( "AddScalarParameter", "Add scalar parameter" ) );
+			ParticleParameterTrack->Modify();
+			ParticleParameterTrack->AddScalarParameterKey( ParameterName, KeyTime, Value );
+		}
+	}
+	NotifyMovieSceneDataChanged();
+}
+
+
+void FParticleParameterTrackEditor::AddVectorParameter( FGuid ObjectBinding, UMovieSceneParticleParameterTrack* ParticleParameterTrack, FName ParameterName )
+{
+	UMovieSceneSequence* MovieSceneSequence = GetMovieSceneSequence();
+	float KeyTime = GetTimeForKey( MovieSceneSequence );
+
+	UObject* Object = GetSequencer()->GetFocusedMovieSceneSequence()->FindObject( ObjectBinding );
+	AEmitter* Emitter = Cast<AEmitter>( Object );
+	if ( Emitter != nullptr )
+	{
+		UParticleSystemComponent* ParticleSystemComponent = Emitter->GetParticleSystemComponent();
+		if ( ParticleSystemComponent != nullptr )
+		{
+			FVector Value;
+			ParticleSystemComponent->GetVectorParameter( ParameterName, Value );
+
+			const FScopedTransaction Transaction( LOCTEXT( "AddVectorParameter", "Add vector parameter" ) );
+			ParticleParameterTrack->Modify();
+			ParticleParameterTrack->AddVectorParameterKey( ParameterName, KeyTime, Value );
+		}
+	}
+	NotifyMovieSceneDataChanged();
+}
+
+
+void FParticleParameterTrackEditor::AddColorParameter( FGuid ObjectBinding, UMovieSceneParticleParameterTrack* ParticleParameterTrack, FName ParameterName )
+{
+	UMovieSceneSequence* MovieSceneSequence = GetMovieSceneSequence();
+	float KeyTime = GetTimeForKey( MovieSceneSequence );
+
+	UObject* Object = GetSequencer()->GetFocusedMovieSceneSequence()->FindObject( ObjectBinding );
+	AEmitter* Emitter = Cast<AEmitter>( Object );
+	if ( Emitter != nullptr )
+	{
+		UParticleSystemComponent* ParticleSystemComponent = Emitter->GetParticleSystemComponent();
+		if ( ParticleSystemComponent != nullptr )
+		{
+			FLinearColor Value;
+			ParticleSystemComponent->GetColorParameter( ParameterName, Value );
+
+			const FScopedTransaction Transaction( LOCTEXT( "AddColorParameter", "Add color parameter" ) );
+			ParticleParameterTrack->Modify();
+			ParticleParameterTrack->AddColorParameterKey( ParameterName, KeyTime, Value );
+		}
+	}
+	NotifyMovieSceneDataChanged();
+}
+
+
+#undef LOCTEXT_NAMESPACE
diff --git a/Engine/Source/Editor/MovieSceneTools/Private/TrackEditors/ParticleParameterTrackEditor.h b/Engine/Source/Editor/MovieSceneTools/Private/TrackEditors/ParticleParameterTrackEditor.h
new file mode 100644
--- /dev/null
+++ b/Engine/Source/Editor/MovieSceneTools/Private/TrackEditors/ParticleParameterTrackEditor.h
@@ -0,0 +1,86 @@
+// Copyright 1998-2015 Epic Games, Inc. All Rights Reserved.
+
+#pragma once
+
+class UMovieSceneParticleParameterTrack;
+
+/**
+ * Track editor for material parameters.
+ */
+class FParticleParameterTrackEditor
+	: public FMovieSceneTrackEditor
+{
+public:
+
+	/** Constructor. */
+	FParticleParameterTrackEditor( TSharedRef<ISequencer> InSequencer );
+
+	/** Virtual destructor. */
+	virtual ~FParticleParameterTrackEditor() { }
+
+	/**
+	* Creates an instance of this class.  Called by a sequencer.
+	*
+	* @param OwningSequencer The sequencer instance to be used by this tool.
+	* @return The new instance of this class.
+	*/
+	static TSharedRef<ISequencerTrackEditor> CreateTrackEditor( TSharedRef<ISequencer> OwningSequencer );
+
+	// ISequencerTrackEditor interface
+
+	virtual TSharedPtr<SWidget> BuildOutlinerEditWidget( const FGuid& ObjectBinding, UMovieSceneTrack* Track ) override;
+	virtual TSharedRef<ISequencerSection> MakeSectionInterface( UMovieSceneSection& SectionObject, UMovieSceneTrack& Track ) override;
+	virtual bool SupportsType( TSubclassOf<UMovieSceneTrack> Type ) const override;
+
+private:
+
+	static FName TrackName;
+
+	// Struct used for building the parameter menu.
+	struct FParameterNameAndAction
+	{
+		FName ParameterName;
+		FUIAction Action;
+
+		FParameterNameAndAction( FName InParameterName, FUIAction InAction )
+		{
+			ParameterName = InParameterName;
+			Action = InAction;
+		}
+
+		bool operator<( FParameterNameAndAction const& Other ) const
+		{
+			return ParameterName < Other.ParameterName;
+		}
+	};
+
+	void BuildObjectBindingTrackMenu( FMenuBuilder& MenuBuilder, const FGuid& ObjectBinding, const UClass* ObjectClass );
+
+	/** Provides the contents of the add parameter menu. */
+	TSharedRef<SWidget> OnGetAddParameterMenuContent( FGuid ObjectBinding, UMovieSceneParticleParameterTrack* ParticleParameterTrack );
+
+	bool CanAddParticleParameterTrack (FGuid ObjectBinding );
+
+	void AddParticleParameterTrack( FGuid ObjectBinding );
+
+	/** Adds a scalar parameter and initial key to a material track.
+	 * @param ObjectBinding The object binding which owns the material track.
+	 * @param ParticleParameterTrack The track to Add the section to.
+	 * @param ParameterName The name of the parameter to add an initial key for.
+	 */
+	void AddScalarParameter( FGuid ObjectBinding, UMovieSceneParticleParameterTrack* ParticleParameterTrack, FName ParameterName );
+
+	/** Adds a vector parameter and initial key to a material track.
+	* @param ObjectBinding The object binding which owns the material track.
+	* @param ParticleParameterTrack The track to Add the section to.
+	* @param ParameterName The name of the parameter to add an initial key for.
+	*/
+	void AddVectorParameter( FGuid ObjectBinding, UMovieSceneParticleParameterTrack* ParticleParameterTrack, FName ParameterName );
+
+	/** Adds a Color parameter and initial key to a material track.
+	* @param ObjectBinding The object binding which owns the material track.
+	* @param ParticleParameterTrack The track to Add the section to.
+	* @param ParameterName The name of the parameter to add an initial key for.
+	*/
+	void AddColorParameter( FGuid ObjectBinding, UMovieSceneParticleParameterTrack* ParticleParameterTrack, FName ParameterName );
+};
diff --git a/Engine/Source/Editor/MovieSceneTools/Private/TrackEditors/ParticleTrackEditor.cpp b/Engine/Source/Editor/MovieSceneTools/Private/TrackEditors/ParticleTrackEditor.cpp
--- a/Engine/Source/Editor/MovieSceneTools/Private/TrackEditors/ParticleTrackEditor.cpp
+++ b/Engine/Source/Editor/MovieSceneTools/Private/TrackEditors/ParticleTrackEditor.cpp
@@ -254,7 +254,7 @@ void FParticleTrackEditor::BuildObjectBindingTrackMenu(FMenuBuilder& MenuBuilder
 		const TSharedPtr<ISequencer> ParentSequencer = GetSequencer();
 
 		MenuBuilder.AddMenuEntry(
-			LOCTEXT("AddParticleTrack", "Particle Track"),
+			LOCTEXT("AddParticleTrack", "Particle Toggle Track"),
 			LOCTEXT("TriggerParticlesTooltip", "Adds a track for controlling particle emitter state."),
 			FSlateIcon(),
 			FUIAction(FExecuteAction::CreateSP(this, &FParticleTrackEditor::AddParticleKey, ObjectBinding))
diff --git a/Engine/Source/Editor/MovieSceneTools/Private/TrackEditors/SkeletalAnimationTrackEditor.cpp b/Engine/Source/Editor/MovieSceneTools/Private/TrackEditors/SkeletalAnimationTrackEditor.cpp
--- a/Engine/Source/Editor/MovieSceneTools/Private/TrackEditors/SkeletalAnimationTrackEditor.cpp
+++ b/Engine/Source/Editor/MovieSceneTools/Private/TrackEditors/SkeletalAnimationTrackEditor.cpp
@@ -44,6 +44,11 @@ UMovieSceneSection* FSkeletalAnimationSection::GetSectionObject()
 }
 
 
+bool FSkeletalAnimationSection::ShouldDrawKeyAreaBackground() const
+{
+	return false;
+}
+
 FText FSkeletalAnimationSection::GetDisplayName() const
 {
 	return LOCTEXT("AnimationSection", "Animation");
diff --git a/Engine/Source/Editor/MovieSceneTools/Private/TrackEditors/SkeletalAnimationTrackEditor.h b/Engine/Source/Editor/MovieSceneTools/Private/TrackEditors/SkeletalAnimationTrackEditor.h
--- a/Engine/Source/Editor/MovieSceneTools/Private/TrackEditors/SkeletalAnimationTrackEditor.h
+++ b/Engine/Source/Editor/MovieSceneTools/Private/TrackEditors/SkeletalAnimationTrackEditor.h
@@ -72,6 +72,7 @@ public:
 	// ISequencerSection interface
 
 	virtual UMovieSceneSection* GetSectionObject() override;
+	virtual bool ShouldDrawKeyAreaBackground() const override;
 	virtual FText GetDisplayName() const override;
 	virtual FText GetSectionTitle() const override;
 	virtual float GetSectionHeight() const override;
diff --git a/Engine/Source/Editor/MovieSceneTools/Private/TrackEditors/SubTrackEditor.cpp b/Engine/Source/Editor/MovieSceneTools/Private/TrackEditors/SubTrackEditor.cpp
--- a/Engine/Source/Editor/MovieSceneTools/Private/TrackEditors/SubTrackEditor.cpp
+++ b/Engine/Source/Editor/MovieSceneTools/Private/TrackEditors/SubTrackEditor.cpp
@@ -41,6 +41,11 @@ public:
 		// do nothing
 	}
 
+	virtual bool ShouldDrawKeyAreaBackground() const override
+	{
+		return false;
+	}
+
 	virtual FText GetDisplayName() const override
 	{
 		return GetSectionTitle();
diff --git a/Engine/Source/Editor/MovieSceneTools/Public/FloatCurveKeyArea.h b/Engine/Source/Editor/MovieSceneTools/Public/FloatCurveKeyArea.h
--- a/Engine/Source/Editor/MovieSceneTools/Public/FloatCurveKeyArea.h
+++ b/Engine/Source/Editor/MovieSceneTools/Public/FloatCurveKeyArea.h
@@ -37,6 +37,7 @@ public:
 	// IKeyArea interface
 
 	virtual TArray<FKeyHandle> AddKeyUnique(float Time, EMovieSceneKeyInterpolation InKeyInterpolation, float TimeToCopyFrom = FLT_MAX) override;
+	virtual TOptional<FKeyHandle> DuplicateKey(FKeyHandle KeyToDuplicate) override;
 	virtual bool CanCreateKeyEditor() override;
 	virtual TSharedRef<SWidget> CreateKeyEditor(ISequencer* Sequencer) override;
 	virtual void DeleteKey(FKeyHandle KeyHandle) override;
diff --git a/Engine/Source/Editor/MovieSceneTools/Public/IntegralKeyArea.h b/Engine/Source/Editor/MovieSceneTools/Public/IntegralKeyArea.h
--- a/Engine/Source/Editor/MovieSceneTools/Public/IntegralKeyArea.h
+++ b/Engine/Source/Editor/MovieSceneTools/Public/IntegralKeyArea.h
@@ -25,6 +25,7 @@ public:
 	// IKeyArea interface
 
 	virtual TArray<FKeyHandle> AddKeyUnique( float Time, EMovieSceneKeyInterpolation InKeyInterpolation, float TimeToCopyFrom = FLT_MAX ) override;
+	virtual TOptional<FKeyHandle> DuplicateKey(FKeyHandle KeyToDuplicate) override;
 	virtual void DeleteKey(FKeyHandle KeyHandle) override;
 	virtual FLinearColor GetColor() override;
 	virtual ERichCurveExtrapolation GetExtrapolationMode(bool bPreInfinity) const override;
diff --git a/Engine/Source/Editor/MovieSceneTools/Public/NameCurveKeyArea.h b/Engine/Source/Editor/MovieSceneTools/Public/NameCurveKeyArea.h
--- a/Engine/Source/Editor/MovieSceneTools/Public/NameCurveKeyArea.h
+++ b/Engine/Source/Editor/MovieSceneTools/Public/NameCurveKeyArea.h
@@ -29,6 +29,7 @@ public:
 	// IKeyArea interface
 
 	virtual TArray<FKeyHandle> AddKeyUnique(float Time, EMovieSceneKeyInterpolation InKeyInterpolation, float TimeToCopyFrom = FLT_MAX) override;
+	virtual TOptional<FKeyHandle> DuplicateKey(FKeyHandle KeyToDuplicate) override;
 	virtual bool CanCreateKeyEditor() override;
 	virtual TSharedRef<SWidget> CreateKeyEditor(ISequencer* Sequencer) override;
 	virtual void DeleteKey(FKeyHandle KeyHandle) override;
diff --git a/Engine/Source/Editor/Sequencer/Private/DisplayNodes/SequencerDisplayNode.cpp b/Engine/Source/Editor/Sequencer/Private/DisplayNodes/SequencerDisplayNode.cpp
--- a/Engine/Source/Editor/Sequencer/Private/DisplayNodes/SequencerDisplayNode.cpp
+++ b/Engine/Source/Editor/Sequencer/Private/DisplayNodes/SequencerDisplayNode.cpp
@@ -20,6 +20,9 @@
 #define LOCTEXT_NAMESPACE "SequencerDisplayNode"
 
 
+/** When 0, regeneration of dynamic key groups is enabled, when non-zero, such behaviour is disabled */
+FThreadSafeCounter KeyGroupRegenerationLock;
+
 namespace SequencerNodeConstants
 {
 	extern const float CommonPadding;
@@ -524,7 +527,7 @@ TSharedRef<FGroupedKeyArea> FSequencerDisplayNode::UpdateKeyGrouping(int32 InSec
 	{
 		KeyGroup = MakeShareable(new FGroupedKeyArea(*this, InSectionIndex));
 	}
-	else
+	else if (KeyGroupRegenerationLock.GetValue() == 0)
 	{
 		*KeyGroup = FGroupedKeyArea(*this, InSectionIndex);
 	}
@@ -550,4 +553,15 @@ bool FSequencerDisplayNode::HandleContextMenuRenameNodeCanExecute(TSharedRef<FSe
 }
 
 
+void FSequencerDisplayNode::DisableKeyGoupingRegeneration()
+{
+	KeyGroupRegenerationLock.Increment();
+}
+
+
+void FSequencerDisplayNode::EnableKeyGoupingRegeneration()
+{
+	KeyGroupRegenerationLock.Decrement();
+}
+
 #undef LOCTEXT_NAMESPACE
diff --git a/Engine/Source/Editor/Sequencer/Private/DisplayNodes/SequencerDisplayNode.h b/Engine/Source/Editor/Sequencer/Private/DisplayNodes/SequencerDisplayNode.h
--- a/Engine/Source/Editor/Sequencer/Private/DisplayNodes/SequencerDisplayNode.h
+++ b/Engine/Source/Editor/Sequencer/Private/DisplayNodes/SequencerDisplayNode.h
@@ -314,6 +314,14 @@ public:
 	/** Get the key grouping for the specified section index */
 	TSharedRef<FGroupedKeyArea> GetKeyGrouping(int32 InSectionIndex);
 
+public:
+
+	/** Temporarily disable dynamic regeneration of key groupings. This prevents overlapping key groups from being amalgamated during drags. Key times will continue to update correctly. */
+	static void DisableKeyGoupingRegeneration();
+
+	/** Re-enable dynamic regeneration of key groupings */
+	static void EnableKeyGoupingRegeneration();
+
 private:
 
 	/** Callback for executing a "Rename Node" context menu action. */
diff --git a/Engine/Source/Editor/Sequencer/Private/GroupedKeyArea.cpp b/Engine/Source/Editor/Sequencer/Private/GroupedKeyArea.cpp
--- a/Engine/Source/Editor/Sequencer/Private/GroupedKeyArea.cpp
+++ b/Engine/Source/Editor/Sequencer/Private/GroupedKeyArea.cpp
@@ -373,6 +373,44 @@ TArray<FKeyHandle> FGroupedKeyArea::AddKeyUnique(float Time, EMovieSceneKeyInter
 	return AddedKeyHandles;
 }
 
+TOptional<FKeyHandle> FGroupedKeyArea::DuplicateKey(FKeyHandle KeyToDuplicate)
+{
+	FKeyGrouping* Group = FindGroup(KeyToDuplicate);
+	if (!Group)
+	{
+		return TOptional<FKeyHandle>();
+	}
+
+	const float Time = Group->RepresentativeTime;
+
+	const int32 NewGroupIndex = Groups.Num();
+	Groups.Emplace(Time);
+	
+	for (const FKeyGrouping::FKeyIndex& Key : Group->Keys)
+	{
+		TOptional<FKeyHandle> NewKeyHandle = KeyAreas[Key.AreaIndex]->DuplicateKey(Key.KeyHandle);
+		if (NewKeyHandle.IsSet())
+		{
+			Groups[NewGroupIndex].Keys.Emplace(Key.AreaIndex, NewKeyHandle.GetValue());
+		}
+	}
+
+	// Update the global index with our new key
+	FIndexEntry* IndexEntry = GlobalIndex.Find(IndexKey);
+	if (IndexEntry)
+	{
+		FKeyHandle ThisGroupKeyHandle;
+
+		IndexEntry->GroupHandles.Add(ThisGroupKeyHandle);
+		IndexEntry->HandleToGroup.Add(ThisGroupKeyHandle, NewGroupIndex);
+		IndexEntry->RepresentativeTimes.Add(Time);
+
+		return ThisGroupKeyHandle;
+	}
+
+	return TOptional<FKeyHandle>();
+}
+
 FRichCurve* FGroupedKeyArea::GetRichCurve()
 {
 	return nullptr;
diff --git a/Engine/Source/Editor/Sequencer/Private/GroupedKeyArea.h b/Engine/Source/Editor/Sequencer/Private/GroupedKeyArea.h
--- a/Engine/Source/Editor/Sequencer/Private/GroupedKeyArea.h
+++ b/Engine/Source/Editor/Sequencer/Private/GroupedKeyArea.h
@@ -45,6 +45,12 @@ struct FKeyGrouping
 		FKeyHandle KeyHandle;
 	};
 
+	/** Construct an empty key group */
+	FKeyGrouping(float InRepresentativeTime)
+		: RepresentativeTime(InRepresentativeTime)
+	{
+	}
+
 	/** Construct this group with a single key handle */
 	FKeyGrouping(float InRepresentativeTime, int32 AreaIndex, FKeyHandle KeyHandle)
 		: RepresentativeTime(InRepresentativeTime)
@@ -82,6 +88,7 @@ public:
 	virtual void                    SetExtrapolationMode(ERichCurveExtrapolation ExtrapMode, bool bPreInfinity) override;
 	virtual ERichCurveExtrapolation GetExtrapolationMode(bool bPreInfinity) const override;
 	virtual TArray<FKeyHandle>		AddKeyUnique(float Time, EMovieSceneKeyInterpolation InKeyInterpolation, float TimeToCopyFrom = FLT_MAX) override;
+	virtual TOptional<FKeyHandle>				DuplicateKey(FKeyHandle KeyToDuplicate) override;
 	virtual FRichCurve*				GetRichCurve() override;
 	virtual UMovieSceneSection*		GetOwningSection() override;
 	virtual bool					CanCreateKeyEditor() override;
diff --git a/Engine/Source/Editor/Sequencer/Private/SSequencer.cpp b/Engine/Source/Editor/Sequencer/Private/SSequencer.cpp
--- a/Engine/Source/Editor/Sequencer/Private/SSequencer.cpp
+++ b/Engine/Source/Editor/Sequencer/Private/SSequencer.cpp
@@ -758,8 +758,10 @@ TSharedRef<SWidget> SSequencer::MakeGeneralMenu()
 			MenuBuilder.AddMenuEntry( FSequencerCommands::Get().FindInContentBrowser );
 		}
 
-		MenuBuilder.AddMenuEntry( FSequencerCommands::Get().ExpandNodesAndDescendants );
-		MenuBuilder.AddMenuEntry( FSequencerCommands::Get().CollapseNodesAndDescendants );
+		MenuBuilder.AddMenuEntry( FSequencerCommands::Get().ToggleExpandCollapseNodes );
+		MenuBuilder.AddMenuEntry( FSequencerCommands::Get().ToggleExpandCollapseNodesAndDescendants );
+		MenuBuilder.AddMenuEntry( FSequencerCommands::Get().ExpandAllNodesAndDescendants );
+		MenuBuilder.AddMenuEntry( FSequencerCommands::Get().CollapseAllNodesAndDescendants );
 	}
 	MenuBuilder.EndSection();
 
@@ -832,6 +834,8 @@ TSharedRef<SWidget> SSequencer::MakeGeneralMenu()
 				})
 			],
 			LOCTEXT("PlaybackStartEnd", "End"));
+
+		MenuBuilder.AddMenuEntry( FSequencerCommands::Get().ToggleKeepCursorInPlaybackRange );
 	}
 
 
diff --git a/Engine/Source/Editor/Sequencer/Private/SSequencerSection.cpp b/Engine/Source/Editor/Sequencer/Private/SSequencerSection.cpp
--- a/Engine/Source/Editor/Sequencer/Private/SSequencerSection.cpp
+++ b/Engine/Source/Editor/Sequencer/Private/SSequencerSection.cpp
@@ -14,21 +14,6 @@
 
 double SSequencerSection::SelectionThrobEndTime = 0;
 
-/** When 0, regeneration of dynamic key layouts is enabled, when non-zero, such behaviour is disabled */
-FThreadSafeCounter LayoutRegenerationLock;
-
-
-void SSequencerSection::DisableLayoutRegeneration()
-{
-	LayoutRegenerationLock.Increment();
-}
-
-
-void SSequencerSection::EnableLayoutRegeneration()
-{
-	LayoutRegenerationLock.Decrement();
-}
-
 
 void SSequencerSection::Construct( const FArguments& InArgs, TSharedRef<FSequencerTrackNode> SectionNode, int32 InSectionIndex )
 {
@@ -346,12 +331,7 @@ void SSequencerSection::PaintKeys( const FGeometry& AllottedGeometry, const FSla
 	{
 		// get key handles
 		TSharedPtr<IKeyArea> KeyArea = LayoutElement.GetKeyArea();
-		TArray<FKeyHandle> KeyHandles = KeyArea->GetUnsortedKeyHandles();
 
-		if (!KeyHandles.Num())
-		{
-			continue;
-		}
 
 		FGeometry KeyAreaGeometry = GetKeyAreaGeometry(LayoutElement, AllottedGeometry);
 
@@ -360,15 +340,25 @@ void SSequencerSection::PaintKeys( const FGeometry& AllottedGeometry, const FSla
 			: FTimeToPixel(KeyAreaGeometry, TRange<float>(SectionObject.GetStartTime(), SectionObject.GetEndTime()));
 
 		// draw a box for the key area 
-		FSlateDrawElement::MakeBox( 
-			OutDrawElements,
-			LayerId,
-			KeyAreaGeometry.ToPaintGeometry(),
-			BackgroundBrush,
-			MyClippingRect,
-			DrawEffects,
-			KeyArea->GetColor()
-		); 
+		if (SectionInterface->ShouldDrawKeyAreaBackground())
+		{
+			FSlateDrawElement::MakeBox( 
+				OutDrawElements,
+				LayerId + 1,
+				KeyAreaGeometry.ToPaintGeometry(),
+				BackgroundBrush,
+				MyClippingRect,
+				DrawEffects,
+				KeyArea->GetColor()
+			); 
+		}
+
+		TArray<FKeyHandle> KeyHandles = KeyArea->GetUnsortedKeyHandles();
+
+		if (!KeyHandles.Num())
+		{
+			continue;
+		}
 
 		const int32 KeyLayer = LayerId + 1;
 
@@ -637,10 +627,7 @@ void SSequencerSection::Tick( const FGeometry& AllottedGeometry, const double In
 {	
 	if( GetVisibility() == EVisibility::Visible )
 	{
-		if (LayoutRegenerationLock.GetValue() == 0)
-		{
-			Layout = FKeyAreaLayout(*ParentSectionArea, SectionIndex);
-		}
+		Layout = FKeyAreaLayout(*ParentSectionArea, SectionIndex);
 		FGeometry SectionGeometry = MakeSectionGeometryWithoutHandles( AllottedGeometry, SectionInterface );
 
 		SectionInterface->Tick(SectionGeometry, ParentGeometry, InCurrentTime, InDeltaTime);
diff --git a/Engine/Source/Editor/Sequencer/Private/SSequencerSection.h b/Engine/Source/Editor/Sequencer/Private/SSequencerSection.h
--- a/Engine/Source/Editor/Sequencer/Private/SSequencerSection.h
+++ b/Engine/Source/Editor/Sequencer/Private/SSequencerSection.h
@@ -20,12 +20,6 @@ public:
 	void CacheParentGeometry(const FGeometry& InParentGeometry) {ParentGeometry = InParentGeometry;}
 
 	virtual FVector2D ComputeDesiredSize(float) const override;
-	
-	/** Temporarily disable dynamic layout regeneration. This prevents overlapping key groups from being amalgamated during drags. Key times will continue to update correctly. */
-	static void DisableLayoutRegeneration();
-
-	/** Re-enable dynamic layout regeneration */
-	static void EnableLayoutRegeneration();
 
 private:
 
diff --git a/Engine/Source/Editor/Sequencer/Private/SSequencerTreeView.cpp b/Engine/Source/Editor/Sequencer/Private/SSequencerTreeView.cpp
--- a/Engine/Source/Editor/Sequencer/Private/SSequencerTreeView.cpp
+++ b/Engine/Source/Editor/Sequencer/Private/SSequencerTreeView.cpp
@@ -359,53 +359,53 @@ bool ShouldExpand(const T& InContainer, ETreeRecursion Recursion)
 	return !bAllExpanded;
 }
 
-void SSequencerTreeView::ToggleExpandCollapseNodes(ETreeRecursion Recursion)
+void SSequencerTreeView::ToggleExpandCollapseNodes(ETreeRecursion Recursion, bool bExpandAll)
 {
-	FSequencer& Sequencer = SequencerNodeTree->GetSequencer();
-
-	const TSet< FDisplayNodeRef >& SelectedNodes = Sequencer.GetSelection().GetSelectedOutlinerNodes();
-
 	bool bExpand = false;
-	if (SelectedNodes.Num() != 0)
+	if (bExpandAll)
 	{
-		bExpand = ShouldExpand(SelectedNodes, Recursion);
+		bExpand = ShouldExpand(SequencerNodeTree->GetRootNodes(), Recursion);
 	}
 	else
 	{
-		bExpand = ShouldExpand(SequencerNodeTree->GetRootNodes(), Recursion);
+		FSequencer& Sequencer = SequencerNodeTree->GetSequencer();
+
+		const TSet< FDisplayNodeRef >& SelectedNodes = Sequencer.GetSelection().GetSelectedOutlinerNodes();
+
+		bExpand = ShouldExpand(SelectedNodes, Recursion);
 	}
 
-	ExpandOrCollapseNodes(Recursion, bExpand);
+	ExpandOrCollapseNodes(Recursion, bExpandAll, bExpand);
 }
 
-void SSequencerTreeView::ExpandNodes(ETreeRecursion Recursion)
+void SSequencerTreeView::ExpandNodes(ETreeRecursion Recursion, bool bExpandAll)
 {
 	const bool bExpand = true;
-	ExpandOrCollapseNodes(Recursion, bExpand);
+	ExpandOrCollapseNodes(Recursion, bExpandAll, bExpand);
 }
 
-void SSequencerTreeView::CollapseNodes(ETreeRecursion Recursion)
+void SSequencerTreeView::CollapseNodes(ETreeRecursion Recursion, bool bExpandAll)
 {
 	const bool bExpand = false;
-	ExpandOrCollapseNodes(Recursion, bExpand);
+	ExpandOrCollapseNodes(Recursion, bExpandAll, bExpand);
 }
 
-void SSequencerTreeView::ExpandOrCollapseNodes(ETreeRecursion Recursion, bool bExpand)
+void SSequencerTreeView::ExpandOrCollapseNodes(ETreeRecursion Recursion, bool bExpandAll, bool bExpand)
 {
 	FSequencer& Sequencer = SequencerNodeTree->GetSequencer();
 
-	const TSet< FDisplayNodeRef >& SelectedNodes = Sequencer.GetSelection().GetSelectedOutlinerNodes();
-
-	if (SelectedNodes.Num() != 0)
+	if (bExpandAll)
 	{
-		for (auto& Item : SelectedNodes)
+		for (auto& Item : SequencerNodeTree->GetRootNodes())
 		{
 			ExpandCollapseNode(Item, bExpand, Recursion);
-		}
+		}	
 	}
 	else
 	{
-		for (auto& Item : SequencerNodeTree->GetRootNodes())
+		const TSet< FDisplayNodeRef >& SelectedNodes = Sequencer.GetSelection().GetSelectedOutlinerNodes();
+
+		for (auto& Item : SelectedNodes)
 		{
 			ExpandCollapseNode(Item, bExpand, Recursion);
 		}
diff --git a/Engine/Source/Editor/Sequencer/Private/SSequencerTreeView.h b/Engine/Source/Editor/Sequencer/Private/SSequencerTreeView.h
--- a/Engine/Source/Editor/Sequencer/Private/SSequencerTreeView.h
+++ b/Engine/Source/Editor/Sequencer/Private/SSequencerTreeView.h
@@ -67,16 +67,13 @@ public:
 	void Refresh();
 
 	/** Expand or collapse nodes */
-	void ToggleExpandCollapseNodes(ETreeRecursion Recursion = ETreeRecursion::Recursive);
+	void ToggleExpandCollapseNodes(ETreeRecursion Recursion = ETreeRecursion::Recursive, bool bExpandAll = false);
 
 	/** Expand nodes */
-	void ExpandNodes(ETreeRecursion Recursion = ETreeRecursion::Recursive);
+	void ExpandNodes(ETreeRecursion Recursion = ETreeRecursion::Recursive, bool bExpandAll = false);
 
 	/** Collapse nodes */
-	void CollapseNodes(ETreeRecursion Recursion = ETreeRecursion::Recursive);
-
-	/** Expand or collapse selected nodes */
-	void ExpandSelectedNodes(ETreeRecursion Recursion);
+	void CollapseNodes(ETreeRecursion Recursion = ETreeRecursion::Recursive, bool bExpandAll = false);
 
 	/** Scroll this tree view by the specified number of slate units */
 	void ScrollByDelta(float DeltaInSlateUnits);
@@ -84,7 +81,7 @@ public:
 protected:
 
 	/** Expand or collapse nodes */
-	void ExpandOrCollapseNodes(ETreeRecursion Recursion, bool bExpand);
+	void ExpandOrCollapseNodes(ETreeRecursion Recursion, bool bExpandAll, bool bExpand);
 
 	/** Set the item's expansion state, including all of its children */
 	void ExpandCollapseNode(const FDisplayNodeRef& InNode, bool bExpansionState, ETreeRecursion Recursion);
diff --git a/Engine/Source/Editor/Sequencer/Private/Sequencer.cpp b/Engine/Source/Editor/Sequencer/Private/Sequencer.cpp
--- a/Engine/Source/Editor/Sequencer/Private/Sequencer.cpp
+++ b/Engine/Source/Editor/Sequencer/Private/Sequencer.cpp
@@ -1098,20 +1098,34 @@ void FSequencer::SetGlobalTimeLooped(float InTime)
 	}
 	else
 	{
-
 		TRange<float> TimeBounds = GetTimeBounds();
 		TRange<float> WorkingRange = GetClampRange();
 
+		const bool bReachedEnd = GetGlobalTime() < TimeBounds.GetUpperBoundValue() && InTime >= TimeBounds.GetUpperBoundValue();
+
 		// Stop if we hit the playback range end
-		if (GetGlobalTime() < TimeBounds.GetUpperBoundValue() && InTime >= TimeBounds.GetUpperBoundValue())
+		if (bReachedEnd)
 		{
 			InTime = TimeBounds.GetUpperBoundValue();
 			PlaybackState = EMovieScenePlayerStatus::Stopped;
 		}
-		// Jump to the start of the working range if necessary
-		else if (InTime < WorkingRange.GetLowerBoundValue())
+		// Constrain to the play range if necessary
+		else if (Settings->ShouldKeepCursorInPlayRange())
+		{
+			// Clamp to lower bound
+			InTime = FMath::Max(InTime, TimeBounds.GetLowerBoundValue());
+
+			// Jump back if necessary
+			if (InTime >= TimeBounds.GetUpperBoundValue())
+			{
+				InTime = TimeBounds.GetLowerBoundValue();
+			}
+		}
+		// Ensure the time is within the working range
+		else if (!WorkingRange.Contains(InTime))
 		{
-			InTime = WorkingRange.GetLowerBoundValue();
+			InTime = FMath::Clamp(InTime, WorkingRange.GetLowerBoundValue(), WorkingRange.GetUpperBoundValue());
+			PlaybackState = EMovieScenePlayerStatus::Stopped;
 		}
 	}
 
@@ -2225,14 +2239,16 @@ void FSequencer::StepToPreviousCameraKey()
 	SequencerWidget->StepToPreviousCameraKey();
 }
 
-void FSequencer::ExpandNodesAndDescendants()
+void FSequencer::ExpandAllNodesAndDescendants()
 {
-	SequencerWidget->GetTreeView()->ExpandNodes(ETreeRecursion::Recursive);
+	const bool bExpandAll = true;
+	SequencerWidget->GetTreeView()->ExpandNodes(ETreeRecursion::Recursive, bExpandAll);
 }
 
-void FSequencer::CollapseNodesAndDescendants()
+void FSequencer::CollapseAllNodesAndDescendants()
 {
-	SequencerWidget->GetTreeView()->CollapseNodes(ETreeRecursion::Recursive);
+	const bool bExpandAll = true;
+	SequencerWidget->GetTreeView()->CollapseNodes(ETreeRecursion::Recursive, bExpandAll);
 }
 
 void FSequencer::ToggleExpandCollapseNodes()
@@ -2259,9 +2275,10 @@ void FSequencer::SetKey()
 	
 		for (auto KeyArea : KeyAreas)
 		{
-			KeyArea->GetOwningSection()->Modify();
-	
-			KeyArea->AddKeyUnique(GetGlobalTime(), GetKeyInterpolation());
+			if (KeyArea->GetOwningSection()->TryModify())
+			{
+				KeyArea->AddKeyUnique(GetGlobalTime(), GetKeyInterpolation());
+			}
 		}
 	}
 }
@@ -2328,6 +2345,11 @@ TArray<TSharedPtr<FMovieSceneClipboard>> GClipboardStack;
 
 void FSequencer::CopySelectedKeys()
 {
+	if (Selection.GetSelectedKeys().Num() == 0)
+	{
+		return;
+	}
+
 	TOptional<float> CopyRelativeTo;
 	
 	// Copy relative to the current key hotspot, if applicable
@@ -2364,11 +2386,14 @@ void FSequencer::CopySelectedKeys()
 
 	TSharedRef<FMovieSceneClipboard> Clipboard = MakeShareable( new FMovieSceneClipboard(Builder.Commit(CopyRelativeTo)) );
 	
-	GClipboardStack.Push(Clipboard);
-
-	if (GClipboardStack.Num() > 10)
+	if (Clipboard->GetKeyTrackGroups().Num())
 	{
-		GClipboardStack.RemoveAt(0, 1);
+		GClipboardStack.Push(Clipboard);
+
+		if (GClipboardStack.Num() > 10)
+		{
+			GClipboardStack.RemoveAt(0, 1);
+		}
 	}
 }
 
@@ -2386,7 +2411,7 @@ const TArray<TSharedPtr<FMovieSceneClipboard>>& FSequencer::GetClipboardStack()
 
 void FSequencer::OnClipboardUsed(TSharedPtr<FMovieSceneClipboard> Clipboard)
 {
-	Clipboard->GetEnvironment().DateTime = FDateTime::Now();
+	Clipboard->GetEnvironment().DateTime = FDateTime::UtcNow();
 
 	// Last entry in the stack should be the most up-to-date
 	GClipboardStack.Sort([](const TSharedPtr<FMovieSceneClipboard>& A, const TSharedPtr<FMovieSceneClipboard>& B){
@@ -2506,12 +2531,12 @@ void FSequencer::BindSequencerCommands()
 		FExecuteAction::CreateSP( this, &FSequencer::SetEndPlaybackRange ) );
 
 	SequencerCommandBindings->MapAction(
-		Commands.ExpandNodesAndDescendants,
-		FExecuteAction::CreateSP(this, &FSequencer::ExpandNodesAndDescendants));
+		Commands.ExpandAllNodesAndDescendants,
+		FExecuteAction::CreateSP(this, &FSequencer::ExpandAllNodesAndDescendants));
 
 	SequencerCommandBindings->MapAction(
-		Commands.CollapseNodesAndDescendants,
-		FExecuteAction::CreateSP(this, &FSequencer::CollapseNodesAndDescendants));
+		Commands.CollapseAllNodesAndDescendants,
+		FExecuteAction::CreateSP(this, &FSequencer::CollapseAllNodesAndDescendants));
 
 	SequencerCommandBindings->MapAction(
 		Commands.ToggleExpandCollapseNodes,
@@ -2681,7 +2706,14 @@ void FSequencer::BindSequencerCommands()
 		FCanExecuteAction::CreateLambda(CanCutOrCopy)
 	);
 
-	SequencerCommandBindings->MapAction( Commands.RenderMovie,
+	SequencerCommandBindings->MapAction(
+		Commands.ToggleKeepCursorInPlaybackRange,
+		FExecuteAction::CreateLambda( [this]{ Settings->SetKeepCursorInPlayRange( !Settings->ShouldKeepCursorInPlayRange() ); } ),
+		FCanExecuteAction::CreateLambda( []{ return true; } ),
+		FIsActionChecked::CreateLambda( [this]{ return Settings->ShouldKeepCursorInPlayRange(); } ) );
+
+	SequencerCommandBindings->MapAction(
+		Commands.RenderMovie,
 		FExecuteAction::CreateLambda([this]{
 			FLevelEditorModule& LevelEditorModule = FModuleManager::GetModuleChecked<FLevelEditorModule>(TEXT("LevelEditor"));
 
diff --git a/Engine/Source/Editor/Sequencer/Private/Sequencer.h b/Engine/Source/Editor/Sequencer/Private/Sequencer.h
--- a/Engine/Source/Editor/Sequencer/Private/Sequencer.h
+++ b/Engine/Source/Editor/Sequencer/Private/Sequencer.h
@@ -448,8 +448,8 @@ protected:
 	void StepToNextCameraKey();
 	void StepToPreviousCameraKey();
 
-	void ExpandNodesAndDescendants();
-	void CollapseNodesAndDescendants();
+	void ExpandAllNodesAndDescendants();
+	void CollapseAllNodesAndDescendants();
 
 	/** Expand or collapse selected nodes */
 	void ToggleExpandCollapseNodes();
diff --git a/Engine/Source/Editor/Sequencer/Private/SequencerCommands.cpp b/Engine/Source/Editor/Sequencer/Private/SequencerCommands.cpp
--- a/Engine/Source/Editor/Sequencer/Private/SequencerCommands.cpp
+++ b/Engine/Source/Editor/Sequencer/Private/SequencerCommands.cpp
@@ -18,11 +18,12 @@ void FSequencerCommands::RegisterCommands()
 	UI_COMMAND( StepToPreviousCameraKey, "Step to Previous Camera Key", "Step to the previous camera key", EUserInterfaceActionType::Button, FInputChord(EModifierKey::Alt, EKeys::Comma) );
 	UI_COMMAND( SetStartPlaybackRange, "Set Start Playback Range", "Set the start playback range", EUserInterfaceActionType::Button, FInputChord(EKeys::LeftBracket) );
 	UI_COMMAND( SetEndPlaybackRange, "Set End Playback Range", "Set the end playback range", EUserInterfaceActionType::Button, FInputChord(EKeys::RightBracket) );
+	UI_COMMAND( ToggleKeepCursorInPlaybackRange, "Keep Cursor In Playback Range", "When checked, the cursor will be constrained to the current playback range during playback", EUserInterfaceActionType::ToggleButton, FInputChord() );
 
-	UI_COMMAND( ExpandNodesAndDescendants, "Expand Nodes", "Expand selected nodes", EUserInterfaceActionType::Button, FInputChord() );
-	UI_COMMAND( CollapseNodesAndDescendants, "Collapse Nodes", "Collapse selected nodes", EUserInterfaceActionType::Button, FInputChord() );
-	UI_COMMAND( ToggleExpandCollapseNodes, "Toggle Expand/Collapse Nodes", "Toggle expand or collapse selected nodes", EUserInterfaceActionType::Button, FInputChord(EKeys::O) );
-	UI_COMMAND( ToggleExpandCollapseNodesAndDescendants, "Toggle Expand/Collapse Nodes and Descendants", "Toggle expand or collapse selected nodes and their descendants", EUserInterfaceActionType::Button, FInputChord(EModifierKey::Shift, EKeys::O) );
+	UI_COMMAND( ExpandAllNodesAndDescendants, "Expand All Nodes", "Expand all nodes", EUserInterfaceActionType::Button, FInputChord() );
+	UI_COMMAND( CollapseAllNodesAndDescendants, "Collapse All Nodes", "Collapse all selected nodes", EUserInterfaceActionType::Button, FInputChord() );
+	UI_COMMAND( ToggleExpandCollapseNodes, "Expand/Collapse Nodes", "Toggle expand or collapse selected nodes", EUserInterfaceActionType::Button, FInputChord(EKeys::O) );
+	UI_COMMAND( ToggleExpandCollapseNodesAndDescendants, "Expand/Collapse Nodes and Descendants", "Toggle expand or collapse selected nodes and descendants", EUserInterfaceActionType::Button, FInputChord(EModifierKey::Shift, EKeys::O) );
 
 	UI_COMMAND( SetKey, "Set Key", "Sets a key on the selected tracks", EUserInterfaceActionType::Button, FInputChord(EKeys::Enter) );
 	UI_COMMAND( SetInterpolationCubicAuto, "Set Key Auto", "Cubic interpolation - Automatic tangents", EUserInterfaceActionType::Button, FInputChord(EKeys::One));
diff --git a/Engine/Source/Editor/Sequencer/Private/SequencerCommands.h b/Engine/Source/Editor/Sequencer/Private/SequencerCommands.h
--- a/Engine/Source/Editor/Sequencer/Private/SequencerCommands.h
+++ b/Engine/Source/Editor/Sequencer/Private/SequencerCommands.h
@@ -48,11 +48,14 @@ public:
 	/** Set end playback range */
 	TSharedPtr< FUICommandInfo > SetEndPlaybackRange;
 
-	/** Expand nodes and descendants */
-	TSharedPtr< FUICommandInfo > ExpandNodesAndDescendants;
+	/** Toggle constraining the time cursor to the playback range */
+	TSharedPtr< FUICommandInfo > ToggleKeepCursorInPlaybackRange;
 
-	/** Collapse nodes and descendants */
-	TSharedPtr< FUICommandInfo > CollapseNodesAndDescendants;
+	/** Expand all nodes and descendants */
+	TSharedPtr< FUICommandInfo > ExpandAllNodesAndDescendants;
+
+	/** Collapse all nodes and descendants */
+	TSharedPtr< FUICommandInfo > CollapseAllNodesAndDescendants;
 
 	/** Expand/collapse nodes */
 	TSharedPtr< FUICommandInfo > ToggleExpandCollapseNodes;
diff --git a/Engine/Source/Editor/Sequencer/Private/SequencerSettings.cpp b/Engine/Source/Editor/Sequencer/Private/SequencerSettings.cpp
--- a/Engine/Source/Editor/Sequencer/Private/SequencerSettings.cpp
+++ b/Engine/Source/Editor/Sequencer/Private/SequencerSettings.cpp
@@ -27,6 +27,7 @@ USequencerSettings::USequencerSettings( const FObjectInitializer& ObjectInitiali
 	bShowCurveEditor = false;
 	bShowCurveEditorCurveToolTips = true;
 	bLooping = false;
+	bKeepCursorInPlayRange = true;
 }
 
 bool USequencerSettings::GetAutoKeyEnabled() const
@@ -325,6 +326,20 @@ void USequencerSettings::SetLooping(bool bInLooping)
 	}
 }
 
+bool USequencerSettings::ShouldKeepCursorInPlayRange() const
+{
+	return bKeepCursorInPlayRange;
+}
+
+void USequencerSettings::SetKeepCursorInPlayRange(bool bInKeepCursorInPlayRange)
+{
+	if (bKeepCursorInPlayRange != bInKeepCursorInPlayRange)
+	{
+		bKeepCursorInPlayRange = bInKeepCursorInPlayRange;
+		SaveConfig();
+	}
+}
+
 bool USequencerSettings::GetShowCurveEditorCurveToolTips() const
 {
 	return bShowCurveEditorCurveToolTips;
diff --git a/Engine/Source/Editor/Sequencer/Private/SequencerSettings.h b/Engine/Source/Editor/Sequencer/Private/SequencerSettings.h
--- a/Engine/Source/Editor/Sequencer/Private/SequencerSettings.h
+++ b/Engine/Source/Editor/Sequencer/Private/SequencerSettings.h
@@ -164,6 +164,11 @@ public:
 	/** Sets the loop mode. */
 	void SetLooping(bool bInLooping);
 
+	/** @return true if the cursor should be kept within the playback range during playback in sequencer, false otherwise */
+	bool ShouldKeepCursorInPlayRange() const;
+	/** Set whether or not the cursor should be kept within the playback range during playback in sequencer */
+	void SetKeepCursorInPlayRange(bool bInKeepCursorInPlayRange);
+
 	/** Snaps a time value in seconds to the currently selected interval. */
 	float SnapTimeToInterval(float InTimeValue) const;
 
@@ -237,6 +242,9 @@ protected:
 	UPROPERTY( config )
 	bool bLooping;
 
+	UPROPERTY( config )
+	bool bKeepCursorInPlayRange;
+
 	FOnShowCurveEditorChanged OnShowCurveEditorChanged;
 };
 
diff --git a/Engine/Source/Editor/Sequencer/Private/Tools/EditToolDragOperations.cpp b/Engine/Source/Editor/Sequencer/Private/Tools/EditToolDragOperations.cpp
--- a/Engine/Source/Editor/Sequencer/Private/Tools/EditToolDragOperations.cpp
+++ b/Engine/Source/Editor/Sequencer/Private/Tools/EditToolDragOperations.cpp
@@ -451,7 +451,7 @@ void FMoveKeys::OnBeginDrag(const FPointerEvent& MouseEvent, FVector2D LocalMous
 {
 	check( SelectedKeys.Num() > 0 )
 
-	SSequencerSection::DisableLayoutRegeneration();
+	FSequencerDisplayNode::DisableKeyGoupingRegeneration();
 
 	FDefaultKeySnappingCandidates SnapCandidates(SelectedKeys);
 	SnapField = FSequencerSnapField(Sequencer, SnapCandidates);
@@ -461,8 +461,7 @@ void FMoveKeys::OnBeginDrag(const FPointerEvent& MouseEvent, FVector2D LocalMous
 	BeginTransaction( DummySections, NSLOCTEXT("Sequencer", "MoveKeysTransaction", "Move Keys") );
 
 	const float MouseTime = VirtualTrackArea.PixelToTime(LocalMousePos.X);
-	
-	ModifiedSections.Empty();
+
 	for( FSequencerSelectedKey SelectedKey : SelectedKeys )
 	{
 		UMovieSceneSection* OwningSection = SelectedKey.Section;
@@ -576,5 +575,54 @@ void FMoveKeys::OnEndDrag(const FPointerEvent& MouseEvent, FVector2D LocalMouseP
 {
 	ModifiedSections.Empty();
 	EndTransaction();
-	SSequencerSection::EnableLayoutRegeneration();
+	FSequencerDisplayNode::EnableKeyGoupingRegeneration();
+}
+
+void FDuplicateKeys::OnBeginDrag(const FPointerEvent& MouseEvent, FVector2D LocalMousePos, const FVirtualTrackArea& VirtualTrackArea)
+{
+	// Duplicate and select all the keys
+	TSet<FSequencerSelectedKey> OldSelection = SelectedKeys;
+
+	// Begin an editor transaction and mark the section as transactional so it's state will be saved
+	TArray<FSectionHandle> DummySections;
+	BeginTransaction( DummySections, NSLOCTEXT("Sequencer", "DuplicateKeysTransaction", "Duplicate Keys") );
+
+	// Modify all the sections first
+	for (const FSequencerSelectedKey& SelectedKey : SelectedKeys)
+	{
+		UMovieSceneSection* OwningSection = SelectedKey.Section;
+
+		// Only modify sections once
+		if (!ModifiedSections.Contains(OwningSection ))
+		{
+			OwningSection->SetFlags(RF_Transactional);
+			// Save the current state of the section
+			if (OwningSection->TryModify())
+			{
+				// Section has been modified
+				ModifiedSections.Add( OwningSection );
+			}
+		}
+	}
+
+	// Then duplicate the keys
+
+	// @todo sequencer: selection in transactions
+	Sequencer.GetSelection().EmptySelectedKeys();
+	for (const FSequencerSelectedKey& SelectedKey : OldSelection)
+	{
+		FSequencerSelectedKey NewKey = SelectedKey;
+		NewKey.KeyHandle = SelectedKey.KeyArea->DuplicateKey(SelectedKey.KeyHandle.GetValue());
+		Sequencer.GetSelection().AddToSelection(NewKey);
+	}
+
+	// Now start the move drag
+	FMoveKeys::OnBeginDrag(MouseEvent, LocalMousePos, VirtualTrackArea);
+}
+
+void FDuplicateKeys::OnEndDrag(const FPointerEvent& MouseEvent, FVector2D LocalMousePos, const FVirtualTrackArea& VirtualTrackArea)
+{
+	FMoveKeys::OnEndDrag(MouseEvent, LocalMousePos, VirtualTrackArea);
+
+	EndTransaction();
 }
\ No newline at end of file
diff --git a/Engine/Source/Editor/Sequencer/Private/Tools/EditToolDragOperations.h b/Engine/Source/Editor/Sequencer/Private/Tools/EditToolDragOperations.h
--- a/Engine/Source/Editor/Sequencer/Private/Tools/EditToolDragOperations.h
+++ b/Engine/Source/Editor/Sequencer/Private/Tools/EditToolDragOperations.h
@@ -145,7 +145,7 @@ public:
 	virtual void OnDrag(const FPointerEvent& MouseEvent, FVector2D LocalMousePos, const FVirtualTrackArea& VirtualTrackArea) override;
 	virtual void OnEndDrag(const FPointerEvent& MouseEvent, FVector2D LocalMousePos, const FVirtualTrackArea& VirtualTrackArea) override;
 
-private:	
+protected:
 
 	/** The selected keys being moved. */
 	const TSet<FSequencerSelectedKey>& SelectedKeys;
@@ -159,3 +159,20 @@ private:
 	/** The set of sections being modified */
 	TSet<UMovieSceneSection*> ModifiedSections;
 };
+
+/**
+ * Operation to drag-duplicate the currently selected keys
+ */
+class FDuplicateKeys : public FMoveKeys
+{
+public:
+
+	FDuplicateKeys( FSequencer& InSequencer, const TSet<FSequencerSelectedKey>& InSelectedKeys )
+		: FMoveKeys(InSequencer, InSelectedKeys)
+	{}
+
+public:
+
+	virtual void OnBeginDrag(const FPointerEvent& MouseEvent, FVector2D LocalMousePos, const FVirtualTrackArea& VirtualTrackArea) override;
+	virtual void OnEndDrag(const FPointerEvent& MouseEvent, FVector2D LocalMousePos, const FVirtualTrackArea& VirtualTrackArea) override;
+};
\ No newline at end of file
diff --git a/Engine/Source/Editor/Sequencer/Private/Tools/SequencerEditTool_Movement.cpp b/Engine/Source/Editor/Sequencer/Private/Tools/SequencerEditTool_Movement.cpp
--- a/Engine/Source/Editor/Sequencer/Private/Tools/SequencerEditTool_Movement.cpp
+++ b/Engine/Source/Editor/Sequencer/Private/Tools/SequencerEditTool_Movement.cpp
@@ -66,7 +66,7 @@ FReply FSequencerEditTool_Movement::OnMouseMove(SWidget& OwnerWidget, const FGeo
 		// Otherwise we can attempt a new drag
 		else if (DelayedDrag->AttemptDragStart(MouseEvent))
 		{
-			DragOperation = CreateDrag();
+			DragOperation = CreateDrag(MouseEvent);
 
 			if (DragOperation.IsValid())
 			{
@@ -83,7 +83,7 @@ FReply FSequencerEditTool_Movement::OnMouseMove(SWidget& OwnerWidget, const FGeo
 }
 
 
-TSharedPtr<ISequencerEditToolDragOperation> FSequencerEditTool_Movement::CreateDrag()
+TSharedPtr<ISequencerEditToolDragOperation> FSequencerEditTool_Movement::CreateDrag(const FPointerEvent& MouseEvent)
 {
 	auto PinnedSequencer = Sequencer.Pin();
 	FSequencerSelection& Selection = PinnedSequencer->GetSelection();
@@ -133,7 +133,15 @@ TSharedPtr<ISequencerEditToolDragOperation> FSequencerEditTool_Movement::CreateD
 				Selection.AddToSelection(ThisKey);
 			}
 
-			return MakeShareable( new FMoveKeys( *PinnedSequencer, Selection.GetSelectedKeys() ) );
+			// @todo sequencer: Make this a customizable UI command modifier?
+			if (MouseEvent.IsAltDown())
+			{
+				return MakeShareable( new FDuplicateKeys( *PinnedSequencer, Selection.GetSelectedKeys() ) );
+			}
+			else
+			{
+				return MakeShareable( new FMoveKeys( *PinnedSequencer, Selection.GetSelectedKeys() ) );
+			}
 		}
 	}
 	// If we're not dragging a hotspot, sections take precedence over keys
diff --git a/Engine/Source/Editor/Sequencer/Private/Tools/SequencerEditTool_Movement.h b/Engine/Source/Editor/Sequencer/Private/Tools/SequencerEditTool_Movement.h
--- a/Engine/Source/Editor/Sequencer/Private/Tools/SequencerEditTool_Movement.h
+++ b/Engine/Source/Editor/Sequencer/Private/Tools/SequencerEditTool_Movement.h
@@ -31,7 +31,7 @@ public:
 
 private:
 
-	TSharedPtr<ISequencerEditToolDragOperation> CreateDrag();
+	TSharedPtr<ISequencerEditToolDragOperation> CreateDrag(const FPointerEvent& MouseEvent);
 
 	/** The sequencer itself */
 	TWeakPtr<FSequencer> Sequencer;
diff --git a/Engine/Source/Editor/Sequencer/Public/IKeyArea.h b/Engine/Source/Editor/Sequencer/Public/IKeyArea.h
--- a/Engine/Source/Editor/Sequencer/Public/IKeyArea.h
+++ b/Engine/Source/Editor/Sequencer/Public/IKeyArea.h
@@ -133,6 +133,11 @@ public:
 	 */
 	virtual TArray<FKeyHandle> AddKeyUnique(float Time, EMovieSceneKeyInterpolation InKeyInterpolation, float TimeToCopyFrom = FLT_MAX) = 0;
 
+	/**
+	 * Duplicate the specified key
+	 */
+	virtual TOptional<FKeyHandle> DuplicateKey(FKeyHandle KeyToDuplicate) = 0;
+
 	/** Gets the rich curve associated with this key area.  This can be null but it must be valid in order to be
 	  * edited by the curve editor. */
 	virtual FRichCurve* GetRichCurve() = 0;
diff --git a/Engine/Source/Editor/Sequencer/Public/ISequencerSection.h b/Engine/Source/Editor/Sequencer/Public/ISequencerSection.h
--- a/Engine/Source/Editor/Sequencer/Public/ISequencerSection.h
+++ b/Engine/Source/Editor/Sequencer/Public/ISequencerSection.h
@@ -64,6 +64,11 @@ public:
 	 */
 	virtual const FSlateBrush* GetKeyBrush(FKeyHandle KeyHandle) const { return nullptr; }
 
+	/*
+	 * @return Whether to draw a background for the key area
+	 */
+	virtual bool ShouldDrawKeyAreaBackground() const { return true; }
+
 	/**
 	 * Called when the section is double clicked
 	 *
diff --git a/Engine/Source/Editor/UnrealEd/Private/Commandlets/ContentCommandlets.cpp b/Engine/Source/Editor/UnrealEd/Private/Commandlets/ContentCommandlets.cpp
--- a/Engine/Source/Editor/UnrealEd/Private/Commandlets/ContentCommandlets.cpp
+++ b/Engine/Source/Editor/UnrealEd/Private/Commandlets/ContentCommandlets.cpp
@@ -663,6 +663,8 @@ int32 UResavePackagesCommandlet::Main( const FString& Params )
 	}
 
 	UE_LOG(LogContentCommandlet, Display, TEXT( "[REPORT] %d/%d packages required resaving" ), PackagesRequiringResave, PackageNames.Num() );
+
+
 	return 0;
 }
 
@@ -705,6 +707,17 @@ bool UResavePackagesCommandlet::PerformPreloadOperations( FLinkerLoad* PackageLi
 						 (MaxResaveUE4Version != IGNORE_PACKAGE_VERSION && UE4PackageVersion <= MaxResaveUE4Version) ||
 						 (MaxResaveLicenseeUE4Version != IGNORE_PACKAGE_VERSION && LicenseeUE4PackageVersion <= MaxResaveLicenseeUE4Version);
 
+	// If the package was saved with a higher engine version do not try to resave it. This also addresses problem with people 
+	// building editor locally and resaving content with a 0 CL version (e.g. BUILD_FROM_CL == 0)
+	if (PackageLinker->Summary.SavedByEngineVersion.GetChangelist() > FEngineVersion::Current().GetChangelist())
+	{
+		UE_LOG(LogContentCommandlet, Warning, TEXT("Skipping resave of %s due to engine version mismatch (Package:%d, Editor:%d "), 
+			*PackageLinker->GetArchiveName(),
+			PackageLinker->Summary.SavedByEngineVersion.GetChangelist(), 
+			FEngineVersion::Current().GetChangelist());
+		bSavePackage = false;
+		bResult = true;	
+	}
 
 	// If not, don't resave it.
 	if ( !bAllowResave )
diff --git a/Engine/Source/Programs/AutomationTool/GUBP/LegacyBranchSetup.cs b/Engine/Source/Programs/AutomationTool/GUBP/LegacyBranchSetup.cs
--- a/Engine/Source/Programs/AutomationTool/GUBP/LegacyBranchSetup.cs
+++ b/Engine/Source/Programs/AutomationTool/GUBP/LegacyBranchSetup.cs
@@ -179,7 +179,7 @@ partial class GUBP
 			/// <summary>
 			/// Allows a branch to override the "Localise" commandlet settings.
 			/// </summary>
-			public bool bShouldBuildLocalization = false;
+			public List<string> LocalizationNodesToBuild = new List<string>();
 			public string LocalizationBranchSuffix = "";
         }
         public virtual void ModifyOptions(GUBP bp, ref BranchOptions Options, string Branch)
@@ -1122,9 +1122,16 @@ partial class GUBP
 			}
 		}
 
-		if (HostPlatforms.Contains(UnrealTargetPlatform.Win64) && BranchConfig.HasNode(RootEditorNode.StaticGetFullName(UnrealTargetPlatform.Win64)) && BranchOptions.bShouldBuildLocalization)
+		if (HostPlatforms.Contains(UnrealTargetPlatform.Win64) && BranchConfig.HasNode(RootEditorNode.StaticGetFullName(UnrealTargetPlatform.Win64)))
 		{
-			BranchConfig.AddNode(new BuildEngineLocalizationNode(BranchOptions.LocalizationBranchSuffix));
+			foreach (var LocalizationNodeToBuild in BranchOptions.LocalizationNodesToBuild)
+			{
+				var LocalizationNode = BuildLocalizationNode.GetLocalizationNode(LocalizationNodeToBuild, BranchOptions.LocalizationBranchSuffix, BranchConfig);
+				if (LocalizationNode != null)
+				{
+					BranchConfig.AddNode(LocalizationNode);
+				}
+			}
 		}
 
         BranchConfig.AddNode(new WaitForTestShared(this));
diff --git a/Engine/Source/Programs/AutomationTool/GUBP/LegacyNodes.cs b/Engine/Source/Programs/AutomationTool/GUBP/LegacyNodes.cs
--- a/Engine/Source/Programs/AutomationTool/GUBP/LegacyNodes.cs
+++ b/Engine/Source/Programs/AutomationTool/GUBP/LegacyNodes.cs
@@ -2646,22 +2646,58 @@ partial class GUBP
         }
     }
 
-	public class BuildEngineLocalizationNode : HostPlatformNode
+	public abstract class BuildLocalizationNode : HostPlatformNode
 	{
-		public BuildEngineLocalizationNode(string InLocalizationBranchSuffix)
+		public BuildLocalizationNode(string InLocalizationBranchSuffix, GUBPBranchConfig InBranchConfig)
 			: base(UnrealTargetPlatform.Win64)
 		{
-			//LocalizationBranchSuffix = InLocalizationBranchSuffix;
-
-			AddDependency(RootEditorNode.StaticGetFullName(HostPlatform));
+			LocalizationBranchSuffix = InLocalizationBranchSuffix;
             AddDependency(ToolsNode.StaticGetFullName(HostPlatform));
 		}
 
 		public override void DoBuild(GUBP bp)
 		{
+			var UEProjectDirectory = GetUEProjectDirectory();
+			var UEProjectName = GetUEProjectName();
+			var OneSkyConfigName = GetOneSkyConfigName();
+			var OneSkyProjectGroupName = GetOneSkyProjectGroupName();
+			var OneSkyProjectNames = GetOneSkyProjectNames();
+
+			// Build the correct command line arguments.
+			var CommandLineArguments = "";
+
+			if (!String.IsNullOrEmpty(UEProjectDirectory))
+			{
+				CommandLineArguments += " -UEProjectDirectory=\"" + UEProjectDirectory + "\"";
+			}
+
+			if (!String.IsNullOrEmpty(UEProjectName))
+			{
+				CommandLineArguments += " -UEProjectName=\"" + UEProjectName + "\"";
+			}
+
+			if (!String.IsNullOrEmpty(OneSkyConfigName))
+			{
+				CommandLineArguments += " -OneSkyConfigName=\"" + OneSkyConfigName + "\"";
+			}
+
+			if (!String.IsNullOrEmpty(OneSkyProjectGroupName))
+			{
+				CommandLineArguments += " -OneSkyProjectGroupName=\"" + OneSkyProjectGroupName + "\"";
+			}
+
+			if (!String.IsNullOrEmpty(OneSkyProjectNames))
+			{
+				CommandLineArguments += " -OneSkyProjectNames=\"" + OneSkyProjectNames + "\"";
+			}
+
+			if (!String.IsNullOrEmpty(LocalizationBranchSuffix))
+			{
+				CommandLineArguments += " -OneSkyBranchSuffix=\"" + LocalizationBranchSuffix + "\"";
+			}
+
 			// Run the localise script.
-			// todo: Will need to pass along LocalizationBranchSuffix to the commandlet (once this is set up to work for branches other than Main)
-			CommandUtils.RunUAT(CommandUtils.CmdEnv, "Localise");
+			CommandUtils.RunUAT(CommandUtils.CmdEnv, "Localise" + CommandLineArguments);
 
 			// Don't pass on any build products to other build nodes at the moment.
 			BuildProducts = new List<string>();
@@ -2673,6 +2709,113 @@ partial class GUBP
 			return base.CISFrequencyQuantumShift(BranchConfig) + 6;
 		}
 
+		public virtual string GetLocalizationId()
+		{
+			throw new AutomationException("Unimplemented GetLocalizationId.");
+		}
+
+		protected virtual string GetUEProjectDirectory()
+		{
+			throw new AutomationException("Unimplemented GetUEProjectDirectory.");
+		}
+
+		protected virtual string GetUEProjectName()
+		{
+			throw new AutomationException("Unimplemented GetUEProjectName.");
+		}
+
+		protected virtual string GetOneSkyConfigName()
+		{
+			throw new AutomationException("Unimplemented GetOneSkyConfigName.");
+		}
+
+		protected virtual string GetOneSkyProjectGroupName()
+		{
+			throw new AutomationException("Unimplemented GetOneSkyProjectGroupName.");
+		}
+
+		protected virtual string GetOneSkyProjectNames()
+		{
+			throw new AutomationException("Unimplemented GetOneSkyProjectNames.");
+		}
+
+		public static BuildLocalizationNode GetLocalizationNode(string InLocalizationId, string InLocalizationBranchSuffix, GUBPBranchConfig InBranchConfig)
+		{
+			if (CachedLocalizationNodeTypes == null)
+			{
+				// Find all types that derive from BuildLocalizationNode in any of our DLLs
+				CachedLocalizationNodeTypes = new Dictionary<string, Type>();
+				var LoadedAssemblies = AppDomain.CurrentDomain.GetAssemblies();
+				foreach (var Dll in LoadedAssemblies)
+				{
+					var AllTypes = Dll.GetTypes();
+					foreach (var PotentialLocalizationNodeType in AllTypes)
+					{
+						if (PotentialLocalizationNodeType != typeof(BuildLocalizationNode) && typeof(BuildLocalizationNode).IsAssignableFrom(PotentialLocalizationNodeType))
+						{
+							// Types should implement a static StaticGetLocalizationId method
+							var Method = PotentialLocalizationNodeType.GetMethod("StaticGetLocalizationId");
+							if (Method != null)
+							{
+								try
+								{
+									var NodeLocalizationId = Method.Invoke(null, null) as string;
+									CachedLocalizationNodeTypes.Add(NodeLocalizationId, PotentialLocalizationNodeType);
+								}
+								catch
+								{
+									BuildCommand.LogWarning("Type '{0}' threw when calling its StaticGetLocalizationId method.", PotentialLocalizationNodeType.FullName);
+								}
+							}
+							else
+							{
+								BuildCommand.LogWarning("Type '{0}' derives from BuildLocalizationNode but is missing its StaticGetLocalizationId method.", PotentialLocalizationNodeType.FullName);
+							}
+						}
+					}
+				}
+			}
+
+			Type LocalizationNodeType;
+			CachedLocalizationNodeTypes.TryGetValue(InLocalizationId, out LocalizationNodeType);
+			if (LocalizationNodeType != null)
+			{
+				try
+				{
+					return Activator.CreateInstance(LocalizationNodeType, new object[] { InLocalizationBranchSuffix, InBranchConfig }) as BuildLocalizationNode;
+				}
+				catch
+				{
+					BuildCommand.LogWarning("Unable to create an instance of the type '{0}'", LocalizationNodeType.FullName);
+				}
+			}
+
+			return null;
+		}
+
+		private static Dictionary<string, Type> CachedLocalizationNodeTypes;
+
+		protected string LocalizationBranchSuffix;
+	}
+
+	public class BuildEngineLocalizationNode : BuildLocalizationNode
+	{
+		public BuildEngineLocalizationNode(string InLocalizationBranchSuffix, GUBPBranchConfig InBranchConfig)
+			: base(InLocalizationBranchSuffix, InBranchConfig)
+		{
+			AddDependency(RootEditorNode.StaticGetFullName(HostPlatform));
+		}
+
+		public static string StaticGetLocalizationId()
+		{
+			return "Engine";
+		}
+
+		public override string GetLocalizationId()
+		{
+			return StaticGetLocalizationId();
+		}
+
 		public static string StaticGetFullName()
 		{
 			return "BuildEngineLocalization";
@@ -2683,7 +2826,30 @@ partial class GUBP
 			return StaticGetFullName();
 		}
 
-		//private string LocalizationBranchSuffix;
+		protected override string GetUEProjectDirectory()
+		{
+			return "Engine";
+		}
+
+		protected override string GetUEProjectName()
+		{
+			return "";
+		}
+
+		protected override string GetOneSkyConfigName()
+		{
+			return "OneSkyConfig_EpicGames";
+		}
+
+		protected override string GetOneSkyProjectGroupName()
+		{
+			return "Unreal Engine";
+		}
+
+		protected override string GetOneSkyProjectNames()
+		{
+			return "Engine,Editor,EditorTutorials,PropertyNames,ToolTips,Category,Keywords";
+		}
 	}
 
     public class GameAggregateNode : HostPlatformAggregateNode
diff --git a/Engine/Source/Programs/AutomationTool/Scripts/CookCommand.Automation.cs b/Engine/Source/Programs/AutomationTool/Scripts/CookCommand.Automation.cs
--- a/Engine/Source/Programs/AutomationTool/Scripts/CookCommand.Automation.cs
+++ b/Engine/Source/Programs/AutomationTool/Scripts/CookCommand.Automation.cs
@@ -324,6 +324,8 @@ public partial class Project : CommandUtils
 
                 List<string> PakFiles = new List<string>();
 
+                string CookPlatformString = CurrentPlatform.GetCookPlatform(false, Params.HasDedicatedServerAndClient, Params.CookFlavor);
+
                 if (Path.HasExtension(SourceCookedContentPath) && (!SourceCookedContentPath.EndsWith(".pak")))
                 {
                     // must be a per platform pkg file try this
@@ -333,20 +335,25 @@ public partial class Project : CommandUtils
                     PakFiles.AddRange( Directory.EnumerateFiles(TemporaryPakPath, Params.ShortProjectName+"*.pak", SearchOption.AllDirectories));
                     PakFiles.AddRange( Directory.EnumerateFiles(TemporaryPakPath, "pakchunk*.pak", SearchOption.AllDirectories));
                 }
-
-                string CookPlatformString = CurrentPlatform.GetCookPlatform(false, Params.HasDedicatedServerAndClient, Params.CookFlavor);
-
-                if (!Path.HasExtension(SourceCookedContentPath))
+                else if (!Path.HasExtension(SourceCookedContentPath))
                 {
                     // try find the pak or pkg file
                     string SourceCookedContentPlatformPath = CombinePaths(SourceCookedContentPath, CookPlatformString);
 
-                    foreach (var PakName in Directory.EnumerateFiles(SourceCookedContentPlatformPath, "*.pak", SearchOption.AllDirectories))
+                    foreach (var PakName in Directory.EnumerateFiles(SourceCookedContentPlatformPath, Params.ShortProjectName + "*.pak", SearchOption.AllDirectories))
+                    {
+                        string TemporaryPakFilename = CombinePaths(TemporaryPakPath, Path.GetFileName(PakName));
+                        File.Copy(PakName, TemporaryPakFilename);
+                        PakFiles.Add(TemporaryPakFilename);
+                    }
+
+                    foreach (var PakName in Directory.EnumerateFiles(SourceCookedContentPlatformPath, "pakchunk*.pak", SearchOption.AllDirectories))
                     {
                         string TemporaryPakFilename = CombinePaths(TemporaryPakPath, Path.GetFileName(PakName));
                         File.Copy(PakName, TemporaryPakFilename);
                         PakFiles.Add(TemporaryPakFilename);
                     }
+
                     if ( PakFiles.Count <= 0 )
                     {
                         Log("No Pak files found in " + SourceCookedContentPlatformPath +" :(");
diff --git a/Engine/Source/Programs/AutomationTool/Scripts/Localisation.Automation.cs b/Engine/Source/Programs/AutomationTool/Scripts/Localisation.Automation.cs
--- a/Engine/Source/Programs/AutomationTool/Scripts/Localisation.Automation.cs
+++ b/Engine/Source/Programs/AutomationTool/Scripts/Localisation.Automation.cs
@@ -10,16 +10,17 @@ using UnrealBuildTool;
 using OneSky;
 using EpicGames.OneSkyLocalization.Config;
 
+[Help("Updates the OneSky localization data using the arguments provided.")]
+[Help("UEProjectDirectory", "Sub-path to the project we're gathering for (relative to CmdEnv.LocalRoot).")]
+[Help("UEProjectName", "Optional name of the project we're gathering for (should match its .uproject file, eg QAGame).")]
+[Help("OneSkyConfigName", "Name of the config data to use (see OneSkyConfigHelper).")]
+[Help("OneSkyProjectGroupName", "Name of the project group in OneSky.")]
+[Help("OneSkyProjectNames", "Comma separated list of the projects in OneSky.")]
+[Help("OneSkyBranchSuffix", "Optional suffix to use when uploading the new data to OneSky.")]
 class Localise : BuildCommand
 {
-	struct ProjectInfo
+	struct ProjectImportExportInfo
 	{
-		/** The name of this project */
-		public string ProjectName;
-
-		/** Path to this projects localization config file (relative to the root working directory for the commandlet) */
-		public string LocalizationConfigFile;
-
 		/** The destination path for the files containing the gathered data for this project (extracted from its config file - relative to the root working directory for the commandlet) */
 		public string DestinationPath;
 
@@ -37,6 +38,24 @@ class Localise : BuildCommand
 
 		/** The cultures to generate for this project (extracted from its config file) */
 		public List<string> CulturesToGenerate;
+
+		/** True if we should use a per-culture directly when importing/exporting */
+		public bool bUseCultureDirectory;
+	};
+
+	struct ProjectInfo
+	{
+		/** The name of this project */
+		public string ProjectName;
+
+		/** Path to this projects localization config files (relative to the root working directory for the commandlet) - ordered so that iterating them runs in the correct order */
+		public List<string> LocalizationConfigFiles;
+
+		/** Config data used by the PO file import process */
+		public ProjectImportExportInfo ImportInfo;
+
+		/** Config data used by the PO file export process */
+		public ProjectImportExportInfo ExportInfo;
 	};
 
     private int OneSkyDownloadedPOChangeList = 0;
@@ -45,19 +64,45 @@ class Localise : BuildCommand
 	{
 		var EditorExe = CombinePaths(CmdEnv.LocalRoot, @"Engine/Binaries/Win64/UE4Editor-Cmd.exe");
 
-		// todo: All of this should be passed in via the command line to allow this script to be generic
-		var UEProjectDirectory = "Engine";
-		var OneSkyConfigName = "OneSkyConfig_EpicGames";
-		var OneSkyProjectGroupName = "Unreal Engine";
-		var OneSkyProjectNames = new string[] {
-			"Engine",
-			"Editor",
-			"EditorTutorials",
-			"PropertyNames",
-			"ToolTips",
-			"Category",
-			"Keywords",
-		};
+		// Parse out the required command line arguments
+		var UEProjectDirectory = ParseParamValue("UEProjectDirectory");
+		if (UEProjectDirectory == null)
+		{
+			throw new AutomationException("Missing required command line argument: 'UEProjectDirectory'");
+		}
+
+		var UEProjectName = ParseParamValue("UEProjectName");
+		if (UEProjectName == null)
+		{
+			UEProjectName = "";
+		}
+
+		var OneSkyConfigName = ParseParamValue("OneSkyConfigName");
+		if (OneSkyConfigName == null)
+		{
+			throw new AutomationException("Missing required command line argument: 'OneSkyConfigName'");
+		}
+
+		var OneSkyProjectGroupName = ParseParamValue("OneSkyProjectGroupName");
+		if (OneSkyProjectGroupName == null)
+		{
+			throw new AutomationException("Missing required command line argument: 'OneSkyProjectGroupName'");
+		}
+
+		var OneSkyProjectNames = new List<string>();
+		{
+			var OneSkyProjectNamesStr = ParseParamValue("OneSkyProjectNames");
+			if (OneSkyProjectNamesStr == null)
+			{
+				throw new AutomationException("Missing required command line argument: 'OneSkyProjectNames'");
+			}
+			foreach (var ProjectName in OneSkyProjectNamesStr.Split(','))
+			{
+				OneSkyProjectNames.Add(ProjectName.Trim());
+			}
+		}
+
+		var OneSkyBranchSuffix = ParseParamValue("OneSkyBranchSuffix");
 
 		var RootWorkingDirectory = CombinePaths(CmdEnv.LocalRoot, UEProjectDirectory);
 
@@ -69,27 +114,27 @@ class Localise : BuildCommand
 		}
 
 		OneSkyConfigData OneSkyConfig = OneSkyConfigHelper.Find(OneSkyConfigName);
-        var OneSkyService = new OneSkyService(OneSkyConfig.ApiKey, OneSkyConfig.ApiSecret);
+		var OneSkyService = new OneSkyService(OneSkyConfig.ApiKey, OneSkyConfig.ApiSecret);
 		var OneSkyProjectGroup = GetOneSkyProjectGroup(OneSkyService, OneSkyProjectGroupName);
 
-        // Create changelist for backed up POs from OneSky.
-        if (P4Enabled)
-        {
-            OneSkyDownloadedPOChangeList = P4.CreateChange(P4Env.Client, "OneSky downloaded PO backup.");
-        }
+		// Create changelist for backed up POs from OneSky.
+		if (P4Enabled)
+		{
+			OneSkyDownloadedPOChangeList = P4.CreateChange(P4Env.Client, "OneSky downloaded PO backup.");
+		}
 
-        // Export all text from OneSky
+		// Export all text from OneSky
 		foreach (var ProjectInfo in ProjectInfos)
 		{
 			ExportOneSkyProjectToDirectory(RootWorkingDirectory, OneSkyService, OneSkyProjectGroup, ProjectInfo);
 		}
 
-        // Submit changelist for backed up POs from OneSky.
-        if (P4Enabled)
-        {
-            int SubmittedChangeList;
-            P4.Submit(OneSkyDownloadedPOChangeList, out SubmittedChangeList);
-        }
+		// Submit changelist for backed up POs from OneSky.
+		if (P4Enabled)
+		{
+			int SubmittedChangeList;
+			P4.Submit(OneSkyDownloadedPOChangeList, out SubmittedChangeList);
+		}
 
 		// Setup editor arguments for SCC.
 		string EditorArguments = String.Empty;
@@ -107,28 +152,31 @@ class Localise : BuildCommand
 		if (P4Enabled) { CommandletSCCArguments += (String.IsNullOrEmpty(CommandletSCCArguments) ? "" : " ") + "-EnableSCC"; }
 		if (!AllowSubmit) { CommandletSCCArguments += (String.IsNullOrEmpty(CommandletSCCArguments) ? "" : " ") + "-DisableSCCSubmit"; }
 
-		// Execute commandlet for each project.
+		// Execute commandlet for each config in each project.
 		foreach (var ProjectInfo in ProjectInfos)
 		{
-			var CommandletArguments = String.Format("-config={0}", ProjectInfo.LocalizationConfigFile) + (String.IsNullOrEmpty(CommandletSCCArguments) ? "" : " " + CommandletSCCArguments);
+			foreach (var LocalizationConfigFile in ProjectInfo.LocalizationConfigFiles)
+			{
+				var CommandletArguments = String.Format("-config={0}", LocalizationConfigFile) + (String.IsNullOrEmpty(CommandletSCCArguments) ? "" : " " + CommandletSCCArguments);
 
-			Log("Localization for {0} {1}", EditorArguments, CommandletArguments);
+				Log("Localization for {0} {1}", EditorArguments, CommandletArguments);
 
-			Log("Running UE4Editor to generate localization data");
+				Log("Running UE4Editor to generate localization data");
 
-			string Arguments = String.Format("-run=GatherText {0} {1}", EditorArguments, CommandletArguments);
-			var RunResult = Run(EditorExe, Arguments);
+				string Arguments = String.Format("{0} -run=GatherText {1} {2}", UEProjectName, EditorArguments, CommandletArguments);
+				var RunResult = Run(EditorExe, Arguments);
 
-			if (RunResult.ExitCode != 0)
-			{
-				throw new AutomationException("Error while executing localization commandlet '{0}'", Arguments);
+				if (RunResult.ExitCode != 0)
+				{
+					Console.WriteLine("[ERROR] Error while executing localization commandlet '{0}'", Arguments);
+				}
 			}
 		}
 
-        // Upload all text to OneSky
+		// Upload all text to OneSky
 		foreach (var ProjectInfo in ProjectInfos)
 		{
-			UploadProjectToOneSky(RootWorkingDirectory, OneSkyService, OneSkyProjectGroup, ProjectInfo);
+			UploadProjectToOneSky(RootWorkingDirectory, OneSkyService, OneSkyProjectGroup, OneSkyBranchSuffix, ProjectInfo);
 		}
 	}
 
@@ -137,41 +185,100 @@ class Localise : BuildCommand
 		var ProjectInfo = new ProjectInfo();
 
 		ProjectInfo.ProjectName = ProjectName;
-		ProjectInfo.LocalizationConfigFile = String.Format(@"Config/Localization/{0}.ini", ProjectName);
+		ProjectInfo.LocalizationConfigFiles = new List<string>();
+
+		// Projects generated by the localization dashboard will use multiple config files that must be run in a specific order
+		// Older projects (such as the Engine) would use a single config file containing all the steps
+		// Work out which kind of project we're dealing with...
+		var MonolithicConfigFile = String.Format(@"Config/Localization/{0}.ini", ProjectName);
+		var IsMonolithicConfig = File.Exists(CombinePaths(RootWorkingDirectory, MonolithicConfigFile));
+		if (IsMonolithicConfig)
+		{
+			ProjectInfo.LocalizationConfigFiles.Add(MonolithicConfigFile);
+
+			ProjectInfo.ImportInfo = GenerateProjectImportExportInfo(RootWorkingDirectory, MonolithicConfigFile);
+			ProjectInfo.ExportInfo = ProjectInfo.ImportInfo;
+		}
+		else
+		{
+			var FileSuffixes = new[] { 
+				new { Suffix = "Gather", Required = true }, 
+				new { Suffix = "Import", Required = true }, 
+				new { Suffix = "Export", Required = true }, 
+				new { Suffix = "Compile", Required = true }, 
+				new { Suffix = "GenerateReports", Required = false } 
+			};
+
+			foreach (var FileSuffix in FileSuffixes)
+			{
+				var ModularConfigFile = String.Format(@"Config/Localization/{0}_{1}.ini", ProjectName, FileSuffix.Suffix);
+
+				if (File.Exists(CombinePaths(RootWorkingDirectory, ModularConfigFile)))
+				{
+					ProjectInfo.LocalizationConfigFiles.Add(ModularConfigFile);
+
+					if (FileSuffix.Suffix == "Import")
+					{
+						ProjectInfo.ImportInfo = GenerateProjectImportExportInfo(RootWorkingDirectory, ModularConfigFile);
+					}
+					else if (FileSuffix.Suffix == "Export")
+					{
+						ProjectInfo.ExportInfo = GenerateProjectImportExportInfo(RootWorkingDirectory, ModularConfigFile);
+					}
+				}
+				else if (FileSuffix.Required)
+				{
+					throw new AutomationException("Failed to find a required config file! '{0}'", ModularConfigFile);
+				}
+			}
+		}
+
+		return ProjectInfo;
+	}
 
-		var LocalizationConfig = new ConfigCacheIni(new FileReference(CombinePaths(RootWorkingDirectory, ProjectInfo.LocalizationConfigFile)));
+	private ProjectImportExportInfo GenerateProjectImportExportInfo(string RootWorkingDirectory, string LocalizationConfigFile)
+	{
+		var ProjectImportExportInfo = new ProjectImportExportInfo();
 
-		if (!LocalizationConfig.GetString("CommonSettings", "DestinationPath", out ProjectInfo.DestinationPath))
+		var LocalizationConfig = new ConfigCacheIni(new FileReference(CombinePaths(RootWorkingDirectory, LocalizationConfigFile)));
+
+		if (!LocalizationConfig.GetString("CommonSettings", "DestinationPath", out ProjectImportExportInfo.DestinationPath))
 		{
-			throw new AutomationException("Failed to find a required config key! Section: 'CommonSettings', Key: 'DestinationPath', File: '{0}'", ProjectInfo.LocalizationConfigFile);
+			throw new AutomationException("Failed to find a required config key! Section: 'CommonSettings', Key: 'DestinationPath', File: '{0}'", LocalizationConfigFile);
 		}
 
-		if (!LocalizationConfig.GetString("CommonSettings", "ManifestName", out ProjectInfo.ManifestName))
+		if (!LocalizationConfig.GetString("CommonSettings", "ManifestName", out ProjectImportExportInfo.ManifestName))
 		{
-			throw new AutomationException("Failed to find a required config key! Section: 'CommonSettings', Key: 'ManifestName', File: '{0}'", ProjectInfo.LocalizationConfigFile);
+			throw new AutomationException("Failed to find a required config key! Section: 'CommonSettings', Key: 'ManifestName', File: '{0}'", LocalizationConfigFile);
 		}
 
-		if (!LocalizationConfig.GetString("CommonSettings", "ArchiveName", out ProjectInfo.ArchiveName))
+		if (!LocalizationConfig.GetString("CommonSettings", "ArchiveName", out ProjectImportExportInfo.ArchiveName))
 		{
-			throw new AutomationException("Failed to find a required config key! Section: 'CommonSettings', Key: 'ArchiveName', File: '{0}'", ProjectInfo.LocalizationConfigFile);
+			throw new AutomationException("Failed to find a required config key! Section: 'CommonSettings', Key: 'ArchiveName', File: '{0}'", LocalizationConfigFile);
 		}
 
-		if (!LocalizationConfig.GetString("CommonSettings", "PortableObjectName", out ProjectInfo.PortableObjectName))
+		if (!LocalizationConfig.GetString("CommonSettings", "PortableObjectName", out ProjectImportExportInfo.PortableObjectName))
 		{
-			throw new AutomationException("Failed to find a required config key! Section: 'CommonSettings', Key: 'PortableObjectName', File: '{0}'", ProjectInfo.LocalizationConfigFile);
+			throw new AutomationException("Failed to find a required config key! Section: 'CommonSettings', Key: 'PortableObjectName', File: '{0}'", LocalizationConfigFile);
 		}
 
-		if (!LocalizationConfig.GetString("CommonSettings", "NativeCulture", out ProjectInfo.NativeCulture))
+		if (!LocalizationConfig.GetString("CommonSettings", "NativeCulture", out ProjectImportExportInfo.NativeCulture))
 		{
-			throw new AutomationException("Failed to find a required config key! Section: 'CommonSettings', Key: 'NativeCulture', File: '{0}'", ProjectInfo.LocalizationConfigFile);
+			throw new AutomationException("Failed to find a required config key! Section: 'CommonSettings', Key: 'NativeCulture', File: '{0}'", LocalizationConfigFile);
 		}
 
-		if (!LocalizationConfig.GetArray("CommonSettings", "CulturesToGenerate", out ProjectInfo.CulturesToGenerate))
+		if (!LocalizationConfig.GetArray("CommonSettings", "CulturesToGenerate", out ProjectImportExportInfo.CulturesToGenerate))
 		{
-			throw new AutomationException("Failed to find a required config key! Section: 'CommonSettings', Key: 'CulturesToGenerate', File: '{0}'", ProjectInfo.LocalizationConfigFile);
+			throw new AutomationException("Failed to find a required config key! Section: 'CommonSettings', Key: 'CulturesToGenerate', File: '{0}'", LocalizationConfigFile);
 		}
 
-		return ProjectInfo;
+		if (!LocalizationConfig.GetBool("CommonSettings", "bUseCultureDirectory", out ProjectImportExportInfo.bUseCultureDirectory))
+		{
+			// bUseCultureDirectory is optional, default is true
+			ProjectImportExportInfo.bUseCultureDirectory = true;
+		}
+
+		return ProjectImportExportInfo;
 	}
 
 	private static ProjectGroup GetOneSkyProjectGroup(OneSkyService OneSkyService, string ProjectGroupName)
@@ -207,7 +314,7 @@ class Localise : BuildCommand
 	private void ExportOneSkyProjectToDirectory(string RootWorkingDirectory, OneSkyService OneSkyService, ProjectGroup OneSkyProjectGroup, ProjectInfo ProjectInfo)
 	{
 		var OneSkyProject = GetOneSkyProject(OneSkyService, OneSkyProjectGroup.Name, ProjectInfo.ProjectName);
-		var OneSkyFile = OneSkyProject.UploadedFiles.FirstOrDefault(f => f.Filename == ProjectInfo.PortableObjectName);
+		var OneSkyFile = OneSkyProject.UploadedFiles.FirstOrDefault(f => f.Filename == ProjectInfo.ImportInfo.PortableObjectName);
 
 		//Export
 		if (OneSkyFile != null)
@@ -216,21 +323,21 @@ class Localise : BuildCommand
 			foreach (var OneSkyCulture in OneSkyProject.EnabledCultures)
 			{
 				// Only export the OneSky cultures that we care about for this project
-				if (ProjectInfo.CulturesToGenerate.Contains(OneSkyCulture))
+				if (ProjectInfo.ImportInfo.CulturesToGenerate.Contains(OneSkyCulture))
 				{
 					CulturesToExport.Add(OneSkyCulture);
 				}
 			}
 
-			ExportOneSkyFileToDirectory(OneSkyFile, new DirectoryInfo(CombinePaths(RootWorkingDirectory, ProjectInfo.DestinationPath)), CulturesToExport);
+			ExportOneSkyFileToDirectory(OneSkyFile, new DirectoryInfo(CombinePaths(RootWorkingDirectory, ProjectInfo.ImportInfo.DestinationPath)), CulturesToExport, ProjectInfo.ImportInfo.bUseCultureDirectory);
 		}
 	}
 
-	private void ExportOneSkyFileToDirectory(UploadedFile OneSkyFile, DirectoryInfo DestinationDirectory, IEnumerable<string> Cultures)
+	private void ExportOneSkyFileToDirectory(UploadedFile OneSkyFile, DirectoryInfo DestinationDirectory, IEnumerable<string> Cultures, bool bUseCultureDirectory)
     {
 		foreach (var Culture in Cultures)
         {
-			var CultureDirectory = new DirectoryInfo(Path.Combine(DestinationDirectory.FullName, Culture));
+			var CultureDirectory = (bUseCultureDirectory) ? new DirectoryInfo(Path.Combine(DestinationDirectory.FullName, Culture)) : DestinationDirectory;
 			if (!CultureDirectory.Exists)
             {
 				CultureDirectory.Create();
@@ -242,23 +349,56 @@ class Localise : BuildCommand
 
 				var ExportTranslationState = OneSkyFile.ExportTranslation(Culture, MemoryStream).Result;
 				if (ExportTranslationState == UploadedFile.ExportTranslationState.Success)
-                {
-					MemoryStream.Position = 0;
-					using (Stream FileStream = File.OpenWrite(ExportFile.FullName))
-                    {
-						MemoryStream.CopyTo(FileStream);
-						Console.WriteLine("[SUCCESS] Exporting: '{0}' ({1})", ExportFile.FullName, Culture);
-                    }
-
-					FileInfo ExportFileCopy = new FileInfo(Path.Combine(ExportFile.DirectoryName, String.Format("{0}_FromOneSky{1}", Path.GetFileNameWithoutExtension(ExportFile.Name), ExportFile.Extension)));
-					File.Copy(ExportFile.FullName, ExportFileCopy.FullName, true);
-
-					// Add/check out backed up POs from OneSky.
-					if (P4Enabled)
+				{
+					// Write out the updated PO file so that the gather commandlet will import the new data from it
 					{
-						UE4Build.AddBuildProductsToChangelist(OneSkyDownloadedPOChangeList, new List<string>() { ExportFileCopy.FullName });
+						var ExportFileWasReadOnly = false;
+						if (ExportFile.Exists)
+						{
+							// We're going to clobber the existing PO file, so make sure it's writable (it may be read-only if in Perforce)
+							ExportFileWasReadOnly = ExportFile.IsReadOnly;
+							ExportFile.IsReadOnly = false;
+						}
+
+						MemoryStream.Position = 0;
+						using (Stream FileStream = ExportFile.OpenWrite())
+						{
+							MemoryStream.CopyTo(FileStream);
+							Console.WriteLine("[SUCCESS] Exporting: '{0}' ({1})", ExportFile.FullName, Culture);
+						}
+
+						if (ExportFileWasReadOnly)
+						{
+							ExportFile.IsReadOnly = true;
+						}
 					}
-                }
+
+					// Also update the back-up copy so we can diff against what we got from OneSky, and what the gather commandlet produced
+					{
+						var ExportFileCopy = new FileInfo(Path.Combine(ExportFile.DirectoryName, String.Format("{0}_FromOneSky{1}", Path.GetFileNameWithoutExtension(ExportFile.Name), ExportFile.Extension)));
+
+						var ExportFileCopyWasReadOnly = false;
+						if (ExportFileCopy.Exists)
+						{
+							// We're going to clobber the existing PO file, so make sure it's writable (it may be read-only if in Perforce)
+							ExportFileCopyWasReadOnly = ExportFileCopy.IsReadOnly;
+							ExportFileCopy.IsReadOnly = false;
+						}
+
+						ExportFile.CopyTo(ExportFileCopy.FullName, true);
+
+						if (ExportFileCopyWasReadOnly)
+						{
+							ExportFileCopy.IsReadOnly = true;
+						}
+
+						// Add/check out backed up POs from OneSky.
+						if (P4Enabled)
+						{
+							UE4Build.AddBuildProductsToChangelist(OneSkyDownloadedPOChangeList, new List<string>() { ExportFileCopy.FullName });
+						}
+					}
+				}
 				else if (ExportTranslationState == UploadedFile.ExportTranslationState.NoContent)
                 {
 					Console.WriteLine("[WARNING] Exporting: '{0}' ({1}) has no translations!", ExportFile.FullName, Culture);
@@ -271,28 +411,39 @@ class Localise : BuildCommand
         }
     }
 
-	private void UploadProjectToOneSky(string RootWorkingDirectory, OneSkyService OneSkyService, ProjectGroup OneSkyProjectGroup, ProjectInfo ProjectInfo)
+	private void UploadProjectToOneSky(string RootWorkingDirectory, OneSkyService OneSkyService, ProjectGroup OneSkyProjectGroup, string OneSkyBranchSuffix, ProjectInfo ProjectInfo)
 	{
 		var OneSkyProject = GetOneSkyProject(OneSkyService, OneSkyProjectGroup.Name, ProjectInfo.ProjectName);
 
+		Func<string, FileInfo> GetPathForCulture = (string Culture) =>
+		{
+			if (ProjectInfo.ExportInfo.bUseCultureDirectory)
+			{
+				return new FileInfo(Path.Combine(RootWorkingDirectory, ProjectInfo.ExportInfo.DestinationPath, Culture, ProjectInfo.ExportInfo.PortableObjectName));
+			}
+			else
+			{
+				return new FileInfo(Path.Combine(RootWorkingDirectory, ProjectInfo.ExportInfo.DestinationPath, ProjectInfo.ExportInfo.PortableObjectName));
+			}
+		};
+
 		// Upload the .po file for the native culture first
-		UploadFileToOneSky(OneSkyProject, new FileInfo(Path.Combine(RootWorkingDirectory, ProjectInfo.DestinationPath, ProjectInfo.NativeCulture, ProjectInfo.PortableObjectName)), ProjectInfo.NativeCulture);
+		UploadFileToOneSky(OneSkyProject, OneSkyBranchSuffix, GetPathForCulture(ProjectInfo.ExportInfo.NativeCulture), ProjectInfo.ExportInfo.NativeCulture);
 
 		// Upload the remaining .po files for the other cultures
-		foreach (var Culture in ProjectInfo.CulturesToGenerate)
+		foreach (var Culture in ProjectInfo.ExportInfo.CulturesToGenerate)
 		{
 			// Skip native culture as we uploaded it above
-			if (Culture != ProjectInfo.NativeCulture)
+			if (Culture != ProjectInfo.ExportInfo.NativeCulture)
 			{
-				UploadFileToOneSky(OneSkyProject, new FileInfo(Path.Combine(RootWorkingDirectory, ProjectInfo.DestinationPath, Culture, ProjectInfo.PortableObjectName)), Culture);	
+				UploadFileToOneSky(OneSkyProject, OneSkyBranchSuffix, GetPathForCulture(Culture), Culture);	
 			}
 		}
 	}
 
-	private void UploadFileToOneSky(OneSky.Project OneSkyProject, FileInfo FileToUpload, string Culture)
+	private void UploadFileToOneSky(OneSky.Project OneSkyProject, string OneSkyBranchSuffix, FileInfo FileToUpload, string Culture)
 	{
-		var FileNameToUpload = FileToUpload.FullName;
-		using (var FileStream = File.OpenRead(FileNameToUpload))
+		using (var FileStream = FileToUpload.OpenRead())
 		{
 			// Read the BOM
 			var UTF8BOM = new byte[3];
@@ -304,18 +455,25 @@ class Localise : BuildCommand
 				FileStream.Position = 0;
 			}
 
-			Console.WriteLine("Uploading: '{0}' ({1})", FileNameToUpload, Culture);
+			Console.WriteLine("Uploading: '{0}' ({1})", FileToUpload.FullName, Culture);
+
+			var OneSkyFileName = Path.GetFileName(FileToUpload.FullName);
+			if (!String.IsNullOrEmpty(OneSkyBranchSuffix))
+			{
+				// Apply the branch suffix. OneSky will take care of merging the files from different branches together.
+				var OneSkyFileNameWithSuffix = Path.GetFileNameWithoutExtension(OneSkyFileName) + "_" + OneSkyBranchSuffix + Path.GetExtension(OneSkyFileName);
+				OneSkyFileName = OneSkyFileNameWithSuffix;
+			}
 
-			// todo: need to use the branch suffix here so that OneSky will merge them together
-			var UploadedFile = OneSkyProject.Upload(Path.GetFileName(FileNameToUpload), FileStream, Culture).Result;
+			var UploadedFile = OneSkyProject.Upload(OneSkyFileName, FileStream, Culture).Result;
 
 			if (UploadedFile == null)
 			{
-				Console.WriteLine("[FAILED] Uploading: '{0}' ({1})", FileNameToUpload, Culture);
+				Console.WriteLine("[FAILED] Uploading: '{0}' ({1})", FileToUpload.FullName, Culture);
 			}
 			else
 			{
-				Console.WriteLine("[SUCCESS] Uploading: '{0}' ({1})", FileNameToUpload, Culture);
+				Console.WriteLine("[SUCCESS] Uploading: '{0}' ({1})", FileToUpload.FullName, Culture);
 			}
 		}
 	}
diff --git a/Engine/Source/Runtime/Engine/Classes/Components/StaticMeshComponent.h b/Engine/Source/Runtime/Engine/Classes/Components/StaticMeshComponent.h
--- a/Engine/Source/Runtime/Engine/Classes/Components/StaticMeshComponent.h
+++ b/Engine/Source/Runtime/Engine/Classes/Components/StaticMeshComponent.h
@@ -184,6 +184,10 @@ class ENGINE_API UStaticMeshComponent : public UMeshComponent
 	UPROPERTY(transient)
 	uint32 bForceNavigationObstacle : 1;
 
+	/** If true, mesh painting is disallowed on this instance. Set if vertex colors are overridden in a construction script. */
+	UPROPERTY(EditAnywhere, AdvancedDisplay, BlueprintReadWrite, Category=Rendering)
+	uint32 bDisallowMeshPaintPerInstance : 1;
+
 	/**
 	 *	Ignore this instance of this static mesh when calculating streaming information.
 	 *	This can be useful when doing things like applying character textures to static geometry,
diff --git a/Engine/Source/Runtime/Engine/Classes/MeshVertexPainter/MeshVertexPainter.h b/Engine/Source/Runtime/Engine/Classes/MeshVertexPainter/MeshVertexPainter.h
new file mode 100644
--- /dev/null
+++ b/Engine/Source/Runtime/Engine/Classes/MeshVertexPainter/MeshVertexPainter.h
@@ -0,0 +1,43 @@
+// Copyright 1998-2015 Epic Games, Inc. All Rights Reserved.
+
+#pragma once
+
+UENUM()
+enum class EVertexPaintAxis : uint8
+{
+	X,
+	Y,
+	Z
+};
+
+
+class ENGINE_API FMeshVertexPainter
+{
+public:
+	/**
+	 * Paints a single vertex color to the given mesh component.
+	 *
+	 * @param	StaticMeshComponent		The static mesh component to be painted
+	 * @param	FillColor				The color to be applied.
+	 * @param	bConvertToSRGB			Whether the vertex color should be converted to sRGB.
+	 */
+	static void PaintVerticesSingleColor(UStaticMeshComponent* StaticMeshComponent, const FLinearColor& FillColor, bool bConvertToSRGB = true);
+
+	/**
+	 * Paints vertex colors on a mesh component lerping from the start to the end color along the specified axis.
+	 *
+	 * @param	StaticMeshComponent		The static mesh component to be painted
+	 * @param	StartColor				The color to paint at the minimum bound
+	 * @param	EndColor				The color to paint at the maximum bound
+	 * @param	Axis					The mesh local space axis along which to lerp
+	 * @param	bConvertToSRGB			Whether the vertex color should be converted to sRGB.
+	 */
+	static void PaintVerticesLerpAlongAxis(UStaticMeshComponent* StaticMeshComponent, const FLinearColor& StartColor, const FLinearColor& EndColor, EVertexPaintAxis Axis, bool bConvertToSRGB = true);
+
+	/**
+	 * Removes painted vertex colors from the given mesh component.
+	 *
+	 * @param	StaticMeshComponent		The static mesh component to have its vertex colors removed
+	 */
+	static void RemovePaintedVertices(UStaticMeshComponent* StaticMeshComponent);
+};
diff --git a/Engine/Source/Runtime/Engine/Classes/MeshVertexPainter/MeshVertexPainterKismetLibrary.h b/Engine/Source/Runtime/Engine/Classes/MeshVertexPainter/MeshVertexPainterKismetLibrary.h
new file mode 100644
--- /dev/null
+++ b/Engine/Source/Runtime/Engine/Classes/MeshVertexPainter/MeshVertexPainterKismetLibrary.h
@@ -0,0 +1,25 @@
+// Copyright 1998-2015 Epic Games, Inc. All Rights Reserved.
+
+#pragma once
+
+#include "MeshVertexPainter.h"
+#include "MeshVertexPainterKismetLibrary.generated.h"
+
+UCLASS(MinimalAPI)
+class UMeshVertexPainterKismetLibrary : public UBlueprintFunctionLibrary
+{
+	GENERATED_BODY()
+
+public:
+	/** Paints vertex colors on a mesh component in a specified color. */
+	UFUNCTION(BlueprintCallable, Category = "VertexPaint")
+	static void PaintVerticesSingleColor(UStaticMeshComponent* StaticMeshComponent, const FLinearColor& FillColor, bool bConvertToSRGB = true);
+
+	/** Paints vertex colors on a mesh component lerping from the start to the end color along the specified axis. */
+	UFUNCTION(BlueprintCallable, Category = "VertexPaint")
+	static void PaintVerticesLerpAlongAxis(UStaticMeshComponent* StaticMeshComponent, const FLinearColor& StartColor, const FLinearColor& EndColor, EVertexPaintAxis Axis, bool bConvertToSRGB = true);
+
+	/** Removes vertex colors on a mesh component */
+	UFUNCTION(BlueprintCallable, Category = "VertexPaint")
+	static void RemovePaintedVertices(UStaticMeshComponent* StaticMeshComponent);
+};
diff --git a/Engine/Source/Runtime/Engine/Private/Animation/AnimInstance.cpp b/Engine/Source/Runtime/Engine/Private/Animation/AnimInstance.cpp
--- a/Engine/Source/Runtime/Engine/Private/Animation/AnimInstance.cpp
+++ b/Engine/Source/Runtime/Engine/Private/Animation/AnimInstance.cpp
@@ -2227,7 +2227,7 @@ void UAnimInstance::StopAllMontagesByGroupName(FName InGroupName, const FAlphaBl
 	for (int32 InstanceIndex = MontageInstances.Num() - 1; InstanceIndex >= 0; InstanceIndex--)
 	{
 		FAnimMontageInstance* MontageInstance = MontageInstances[InstanceIndex];
-		if (MontageInstance && MontageInstance->IsActive() && MontageInstance->Montage && (MontageInstance->Montage->GetGroupName() == InGroupName))
+		if (MontageInstance && MontageInstance->Montage && (MontageInstance->Montage->GetGroupName() == InGroupName))
 		{
 			MontageInstances[InstanceIndex]->Stop(BlendOut, true);
 		}
diff --git a/Engine/Source/Runtime/Engine/Private/Animation/SkeletalControl.cpp b/Engine/Source/Runtime/Engine/Private/Animation/SkeletalControl.cpp
--- a/Engine/Source/Runtime/Engine/Private/Animation/SkeletalControl.cpp
+++ b/Engine/Source/Runtime/Engine/Private/Animation/SkeletalControl.cpp
@@ -30,7 +30,7 @@ bool FBoneReference::Initialize(const FBoneContainer& RequiredBones)
 	// If bone name is not found, look into the master skeleton to see if it's found there.
 	// SkeletalMeshes can exclude bones from the master skeleton, and that's OK.
 	// If it's not found in the master skeleton, the bone does not exist at all! so we should report it as a warning.
-	if (BoneIndex == INDEX_NONE)
+	if (BoneIndex == INDEX_NONE && BoneName != NAME_None)
 	{
 		if (USkeleton* SkeletonAsset = RequiredBones.GetSkeletonAsset())
 		{
diff --git a/Engine/Source/Runtime/Engine/Private/Components/SkeletalMeshComponent.cpp b/Engine/Source/Runtime/Engine/Private/Components/SkeletalMeshComponent.cpp
--- a/Engine/Source/Runtime/Engine/Private/Components/SkeletalMeshComponent.cpp
+++ b/Engine/Source/Runtime/Engine/Private/Components/SkeletalMeshComponent.cpp
@@ -1062,7 +1062,6 @@ void USkeletalMeshComponent::RefreshBoneTransforms(FActorComponentTickFunction*
 	}
 
 	AActor* Owner = GetOwner();
-	UE_LOG(LogAnimation, Verbose, TEXT("RefreshBoneTransforms(%s)"), *GetNameSafe(Owner));
 
 	AnimEvaluationContext.SkeletalMesh = SkeletalMesh;
 	AnimEvaluationContext.AnimInstance = AnimScriptInstance;
diff --git a/Engine/Source/Runtime/Engine/Private/Components/StaticMeshComponent.cpp b/Engine/Source/Runtime/Engine/Private/Components/StaticMeshComponent.cpp
--- a/Engine/Source/Runtime/Engine/Private/Components/StaticMeshComponent.cpp
+++ b/Engine/Source/Runtime/Engine/Private/Components/StaticMeshComponent.cpp
@@ -184,6 +184,7 @@ UStaticMeshComponent::UStaticMeshComponent(const FObjectInitializer& ObjectIniti
 	bHasCustomNavigableGeometry = EHasCustomNavigableGeometry::Yes;
 	bOverrideNavigationExport = false;
 	bForceNavigationObstacle = true;
+	bDisallowMeshPaintPerInstance = false;
 
 	GetBodyInstance()->bAutoWeld = true;	//static mesh by default has auto welding
 
@@ -1701,6 +1702,7 @@ void UStaticMeshComponent::ApplyComponentInstanceData(FStaticMeshComponentInstan
 		}
 	}
 
+	if (!bDisallowMeshPaintPerInstance)
 	{
 		FComponentReregisterContext ReregisterStaticMesh(this);
 		StaticMeshInstanceData->ApplyVertexColorData(this);
diff --git a/Engine/Source/Runtime/Engine/Private/EdGraph/EdGraph.cpp b/Engine/Source/Runtime/Engine/Private/EdGraph/EdGraph.cpp
--- a/Engine/Source/Runtime/Engine/Private/EdGraph/EdGraph.cpp
+++ b/Engine/Source/Runtime/Engine/Private/EdGraph/EdGraph.cpp
@@ -132,7 +132,7 @@ void UEdGraph::PostLoad()
 		if (Nodes[i] == nullptr)
 		{
 			Nodes.RemoveAt(i);
-			UE_LOG(LogBlueprint, Warning, TEXT("Found NULL Node in EdGraph Nodes array. A node type may have been deleted without creating an ActiveClassRedictor to K2Node_DeadClass."));
+			UE_LOG(LogBlueprint, Display, TEXT("Found NULL Node in EdGraph Nodes array. A node type may have been deleted without creating an ActiveClassRedictor to K2Node_DeadClass."));
 		}
 	}
 }
diff --git a/Engine/Source/Runtime/Engine/Private/MeshVertexPainter/MeshVertexPainter.cpp b/Engine/Source/Runtime/Engine/Private/MeshVertexPainter/MeshVertexPainter.cpp
new file mode 100644
--- /dev/null
+++ b/Engine/Source/Runtime/Engine/Private/MeshVertexPainter/MeshVertexPainter.cpp
@@ -0,0 +1,110 @@
+// Copyright 1998-2015 Epic Games, Inc. All Rights Reserved.
+
+#include "EnginePrivate.h"
+#include "Components/StaticMeshComponent.h"
+#include "StaticMeshResources.h"
+#include "MeshVertexPainter/MeshVertexPainter.h"
+
+
+void FMeshVertexPainter::PaintVerticesSingleColor(UStaticMeshComponent* StaticMeshComponent, const FLinearColor& FillColor, bool bConvertToSRGB)
+{
+	if (!StaticMeshComponent || !StaticMeshComponent->StaticMesh)
+	{
+		return;
+	}
+
+	const int32 NumMeshLODs = StaticMeshComponent->StaticMesh->GetNumLODs();
+	StaticMeshComponent->SetLODDataCount(NumMeshLODs, NumMeshLODs);
+
+	const FColor Color = FillColor.ToFColor(bConvertToSRGB);
+
+	uint32 LODIndex = 0;
+	for (FStaticMeshComponentLODInfo& LODInfo : StaticMeshComponent->LODData)
+	{
+		StaticMeshComponent->RemoveInstanceVertexColorsFromLOD(LODIndex);
+		check(LODInfo.OverrideVertexColors == nullptr);
+
+		const FStaticMeshLODResources& LODModel = StaticMeshComponent->StaticMesh->RenderData->LODResources[LODIndex];
+		const FPositionVertexBuffer& PositionVertexBuffer = LODModel.PositionVertexBuffer;
+		const uint32 NumVertices = PositionVertexBuffer.GetNumVertices();
+
+		LODInfo.OverrideVertexColors = new FColorVertexBuffer;
+		LODInfo.OverrideVertexColors->InitFromSingleColor(Color, NumVertices);
+
+		BeginInitResource(LODInfo.OverrideVertexColors);
+
+		LODIndex++;
+	}
+
+	StaticMeshComponent->CachePaintedDataIfNecessary();
+	StaticMeshComponent->MarkRenderStateDirty();
+	StaticMeshComponent->bDisallowMeshPaintPerInstance = true;
+}
+
+void FMeshVertexPainter::PaintVerticesLerpAlongAxis(UStaticMeshComponent* StaticMeshComponent, const FLinearColor& StartColor, const FLinearColor& EndColor, EVertexPaintAxis Axis, bool bConvertToSRGB)
+{
+	if (!StaticMeshComponent || !StaticMeshComponent->StaticMesh)
+	{
+		return;
+	}
+
+	const FBoxSphereBounds Bounds = StaticMeshComponent->StaticMesh->GetBounds();
+	const FBox Box = Bounds.GetBox();
+	static_assert(static_cast<int32>(EVertexPaintAxis::X) == 0, "EVertexPaintAxis not correctly defined");
+	const float AxisMin = Box.Min.Component(static_cast<int32>(Axis));
+	const float AxisMax = Box.Max.Component(static_cast<int32>(Axis));
+
+	const int32 NumMeshLODs = StaticMeshComponent->StaticMesh->GetNumLODs();
+	StaticMeshComponent->SetLODDataCount(NumMeshLODs, NumMeshLODs);
+
+	uint32 LODIndex = 0;
+	for (FStaticMeshComponentLODInfo& LODInfo : StaticMeshComponent->LODData)
+	{
+		StaticMeshComponent->RemoveInstanceVertexColorsFromLOD(LODIndex);
+		check(LODInfo.OverrideVertexColors == nullptr);
+
+		const FStaticMeshLODResources& LODModel = StaticMeshComponent->StaticMesh->RenderData->LODResources[LODIndex];
+		const FPositionVertexBuffer& PositionVertexBuffer = LODModel.PositionVertexBuffer;
+		const uint32 NumVertices = PositionVertexBuffer.GetNumVertices();
+
+		TArray<FColor> VertexColors;
+		VertexColors.AddZeroed(NumVertices);
+
+		for (uint32 VertexIndex = 0; VertexIndex < NumVertices; ++VertexIndex)
+		{
+			const FVector& VertexPosition = PositionVertexBuffer.VertexPosition(VertexIndex);
+			const FLinearColor Color = FMath::Lerp(StartColor, EndColor, (VertexPosition.Component(static_cast<int32>(Axis)) - AxisMin) / (AxisMax - AxisMin));
+			VertexColors[VertexIndex] = Color.ToFColor(bConvertToSRGB);
+		}
+
+		LODInfo.OverrideVertexColors = new FColorVertexBuffer;
+		LODInfo.OverrideVertexColors->InitFromColorArray(VertexColors);
+
+		BeginInitResource(LODInfo.OverrideVertexColors);
+
+		LODIndex++;
+	}
+
+	StaticMeshComponent->CachePaintedDataIfNecessary();
+	StaticMeshComponent->MarkRenderStateDirty();
+	StaticMeshComponent->bDisallowMeshPaintPerInstance = true;
+}
+
+void FMeshVertexPainter::RemovePaintedVertices(UStaticMeshComponent* StaticMeshComponent)
+{
+	if (!StaticMeshComponent || !StaticMeshComponent->StaticMesh)
+	{
+		return;
+	}
+
+	uint32 LODIndex = 0;
+	for (FStaticMeshComponentLODInfo& LODInfo : StaticMeshComponent->LODData)
+	{
+		StaticMeshComponent->RemoveInstanceVertexColorsFromLOD(LODIndex);
+		check(LODInfo.OverrideVertexColors == nullptr);
+		LODIndex++;
+	}
+
+	StaticMeshComponent->MarkRenderStateDirty();
+	StaticMeshComponent->bDisallowMeshPaintPerInstance = false;
+}
diff --git a/Engine/Source/Runtime/Engine/Private/MeshVertexPainter/MeshVertexPainterKismetLibrary.cpp b/Engine/Source/Runtime/Engine/Private/MeshVertexPainter/MeshVertexPainterKismetLibrary.cpp
new file mode 100644
--- /dev/null
+++ b/Engine/Source/Runtime/Engine/Private/MeshVertexPainter/MeshVertexPainterKismetLibrary.cpp
@@ -0,0 +1,20 @@
+// Copyright 1998-2015 Epic Games, Inc. All Rights Reserved.
+
+#include "EnginePrivate.h"
+#include "MeshVertexPainter/MeshVertexPainterKismetLibrary.h"
+
+
+void UMeshVertexPainterKismetLibrary::PaintVerticesSingleColor(UStaticMeshComponent* StaticMeshComponent, const FLinearColor& FillColor, bool bConvertToSRGB)
+{
+	FMeshVertexPainter::PaintVerticesSingleColor(StaticMeshComponent, FillColor, bConvertToSRGB);
+}
+
+void UMeshVertexPainterKismetLibrary::PaintVerticesLerpAlongAxis(UStaticMeshComponent* StaticMeshComponent, const FLinearColor& StartColor, const FLinearColor& EndColor, EVertexPaintAxis Axis, bool bConvertToSRGB)
+{
+	FMeshVertexPainter::PaintVerticesLerpAlongAxis(StaticMeshComponent, StartColor, EndColor, Axis, bConvertToSRGB);
+}
+
+void UMeshVertexPainterKismetLibrary::RemovePaintedVertices(UStaticMeshComponent* StaticMeshComponent)
+{
+	FMeshVertexPainter::RemovePaintedVertices(StaticMeshComponent);
+}
diff --git a/Engine/Source/Runtime/Engine/Private/Particles/ParticleSystemRender.cpp b/Engine/Source/Runtime/Engine/Private/Particles/ParticleSystemRender.cpp
--- a/Engine/Source/Runtime/Engine/Private/Particles/ParticleSystemRender.cpp
+++ b/Engine/Source/Runtime/Engine/Private/Particles/ParticleSystemRender.cpp
@@ -5560,6 +5560,8 @@ int32 FDynamicRibbonEmitterData::FillVertexData(struct FAsyncBufferFillData& Dat
 			{
 				WorkingUp = CameraUp;
 			}
+
+			WorkingUp.Normalize();
 		}
 
 		while (TrailPayload)
@@ -5788,6 +5790,7 @@ int32 FDynamicRibbonEmitterData::FillVertexData(struct FAsyncBufferFillData& Dat
 					{
 						WorkingUp = CameraUp;
 					}
+					WorkingUp.Normalize();
 				}
 			}
 		}
diff --git a/Engine/Source/Runtime/Engine/Private/SkeletalMeshComponentPhysics.cpp b/Engine/Source/Runtime/Engine/Private/SkeletalMeshComponentPhysics.cpp
--- a/Engine/Source/Runtime/Engine/Private/SkeletalMeshComponentPhysics.cpp
+++ b/Engine/Source/Runtime/Engine/Private/SkeletalMeshComponentPhysics.cpp
@@ -3591,10 +3591,15 @@ void USkeletalMeshComponent::ParallelUpdateClothState(float DeltaTime, FClothSim
 	for (const FClothingActor& ClothingActor : ClothingActors)
 	{
 		const FClothingAssetData& ClothingAsset = SkeletalMesh->ClothingAssets[ClothingActor.ParentClothingAssetIndex];
+		const NxClothingAsset* ApexClothingAsset = ClothingActor.ParentClothingAsset;
+		NxClothingActor* ApexClothingActor = ClothingActor.ApexClothingActor;
 
-		ApplyWindForCloth(ClothingActor);
+		if (!ApexClothingActor || !ApexClothingAsset)
+		{
+			continue;
+		}
 
-		const NxClothingAsset* ApexClothingAsset = ClothingActor.ParentClothingAsset;
+		ApplyWindForCloth(ClothingActor);
 
 		uint32 NumUsedBones = ApexClothingAsset->getNumUsedBones();
 
@@ -3639,17 +3644,13 @@ void USkeletalMeshComponent::ParallelUpdateClothState(float DeltaTime, FClothSim
 			}
 		}
 
-		if (NxClothingActor* ApexClothingActor = ClothingActor.ApexClothingActor)
-		{
-			// if bUseInternalboneOrder is set, "NumUsedBones" works, otherwise have to use "getNumBones" 
-			ApexClothingActor->updateState(
-				PxGlobalPose,
-				BoneMatrices.GetData(),
-				sizeof(physx::PxMat44),
-				NumUsedBones,
-				CurTeleportMode);
-		}
-
+		// if bUseInternalboneOrder is set, "NumUsedBones" works, otherwise have to use "getNumBones" 
+		ApexClothingActor->updateState(
+			PxGlobalPose,
+			BoneMatrices.GetData(),
+			sizeof(physx::PxMat44),
+			NumUsedBones,
+			CurTeleportMode);
 
 		BoneMatrices.Reset();
 	}
diff --git a/Engine/Source/Runtime/Engine/Private/UnrealClient.cpp b/Engine/Source/Runtime/Engine/Private/UnrealClient.cpp
--- a/Engine/Source/Runtime/Engine/Private/UnrealClient.cpp
+++ b/Engine/Source/Runtime/Engine/Private/UnrealClient.cpp
@@ -19,7 +19,6 @@
 #include "EngineModule.h"
 #include "ContentStreaming.h"
 #include "SceneUtils.h"
-#include "MovieSceneCaptureModule.h"
 #include "NotificationManager.h"
 #include "Performance/EnginePerformanceTargets.h"
 
@@ -709,12 +708,6 @@ FViewport::FViewport(FViewportClient* InViewportClient):
 
 FViewport::~FViewport()
 {
-#if WITH_EDITOR
-	if (auto* MovieSceneCapture = GetMovieSceneCapture())
-	{
-		MovieSceneCapture->Close();
-	}
-#endif
 }
 
 bool FViewport::TakeHighResScreenShot()
@@ -1566,16 +1559,6 @@ void FViewport::SetInitialSize( FIntPoint InitialSizeXY )
 	}
 }
 
-IMovieSceneCaptureInterface* FViewport::GetMovieSceneCapture() const
-{
-	if (MovieSceneCaptureHandle.IsValid())
-	{
-		return IMovieSceneCaptureModule::Get().RetrieveMovieSceneInterface(MovieSceneCaptureHandle);
-	}
-
-	return nullptr;
-}
-
 
 ENGINE_API bool GetViewportScreenShot(FViewport* Viewport, TArray<FColor>& Bitmap, const FIntRect& ViewRect /*= FIntRect()*/)
 {
diff --git a/Engine/Source/Runtime/Engine/Public/AnimationRuntime.h b/Engine/Source/Runtime/Engine/Public/AnimationRuntime.h
--- a/Engine/Source/Runtime/Engine/Public/AnimationRuntime.h
+++ b/Engine/Source/Runtime/Engine/Public/AnimationRuntime.h
@@ -263,6 +263,10 @@ public:
 	static void FillUpSpaceBasesRetargetBasePose(const USkeleton* Skeleton, TArray<FTransform> &SpaceBaseRefPose);
 #endif
 
+	/* Weight utility functions */
+	static bool IsFullWeight(float Weight) { return Weight > 1.f - ZERO_ANIMWEIGHT_THRESH; }
+	static bool HasWeight(float Weight) { return Weight > ZERO_ANIMWEIGHT_THRESH; }
+	
 	/**
 	* Combine CurveKeys (that reference morph targets by name) and ActiveAnims (that reference vertex anims by reference) into the ActiveVertexAnims array.
 	*/
diff --git a/Engine/Source/Runtime/Engine/Public/UnrealClient.h b/Engine/Source/Runtime/Engine/Public/UnrealClient.h
--- a/Engine/Source/Runtime/Engine/Public/UnrealClient.h
+++ b/Engine/Source/Runtime/Engine/Public/UnrealClient.h
@@ -10,12 +10,10 @@
 #include "HitProxies.h"
 #include "InputCoreTypes.h"
 #include "Engine/EngineBaseTypes.h"
-#include "Runtime/MovieSceneCapture/Public/MovieSceneCaptureHandle.h"
 #include "PopupMethodReply.h"
 	
 class FCanvas;
 class FViewportClient;
-class IMovieSceneCaptureInterface;
 
 class SWidget;
 class FCursorReply;
@@ -623,23 +621,6 @@ protected:
 	ENGINE_API virtual void ReleaseDynamicRHI() override;
 	ENGINE_API virtual void InitRHI() override;
 	ENGINE_API virtual void ReleaseRHI() override;
-
-public:
-
-	/** Check if this viewport has a movie scene capture implementation or not */
-	bool HasMovieSceneCapture() const { return MovieSceneCaptureHandle.IsValid(); }
-
-	/** Access the current movie scene capture interface, if set. */
-	ENGINE_API IMovieSceneCaptureInterface* GetMovieSceneCapture() const;
-
-	/** Set this viewport's capture implementation */
-	void SetMovieSceneCapture(FMovieSceneCaptureHandle Handle) { MovieSceneCaptureHandle = Handle; }
-
-protected:
-
-	/** Movie scene capture implementation handle. */
-	FMovieSceneCaptureHandle MovieSceneCaptureHandle;
-
 };
 
 // Shortcuts for checking the state of both left&right variations of control keys.
diff --git a/Engine/Source/Runtime/GameplayAbilities/Private/AbilitySystemBlueprintLibrary.cpp b/Engine/Source/Runtime/GameplayAbilities/Private/AbilitySystemBlueprintLibrary.cpp
--- a/Engine/Source/Runtime/GameplayAbilities/Private/AbilitySystemBlueprintLibrary.cpp
+++ b/Engine/Source/Runtime/GameplayAbilities/Private/AbilitySystemBlueprintLibrary.cpp
@@ -702,6 +702,20 @@ int32 UAbilitySystemBlueprintLibrary::GetActiveGameplayEffectStackCount(FActiveG
 	return 0;
 }
 
+int32 UAbilitySystemBlueprintLibrary::GetActiveGameplayEffectStackLimitCount(FActiveGameplayEffectHandle ActiveHandle)
+{
+	UAbilitySystemComponent* ASC = ActiveHandle.GetOwningAbilitySystemComponent();
+	if (ASC)
+	{
+		const UGameplayEffect* ActiveGE = ASC->GetGameplayEffectDefForHandle(ActiveHandle);
+		if (ActiveGE)
+		{
+			return ActiveGE->StackLimitCount;
+		}
+	}
+	return 0;
+}
+
 float UAbilitySystemBlueprintLibrary::GetModifiedAttributeMagnitude(FGameplayEffectSpecHandle SpecHandle, FGameplayAttribute Attribute)
 {
 	FGameplayEffectSpec* Spec = SpecHandle.Data.Get();
diff --git a/Engine/Source/Runtime/GameplayAbilities/Private/AbilitySystemGlobals.cpp b/Engine/Source/Runtime/GameplayAbilities/Private/AbilitySystemGlobals.cpp
--- a/Engine/Source/Runtime/GameplayAbilities/Private/AbilitySystemGlobals.cpp
+++ b/Engine/Source/Runtime/GameplayAbilities/Private/AbilitySystemGlobals.cpp
@@ -273,6 +273,8 @@ UGameplayCueManager* UAbilitySystemGlobals::GetGameplayCueManager()
 			// Fallback to CDO
 			GlobalGameplayCueManager = UGameplayCueManager::StaticClass()->GetDefaultObject<UGameplayCueManager>();
 		}
+
+		GlobalGameplayCueManager->OnCreated();
 		
 		if (GlobalGameplayCueManager->ShouldAsyncLoadObjectLibrariesAtStart() && GameplayCueNotifyPaths.Num() > 0)
 		{
diff --git a/Engine/Source/Runtime/GameplayAbilities/Private/AttributeSet.cpp b/Engine/Source/Runtime/GameplayAbilities/Private/AttributeSet.cpp
--- a/Engine/Source/Runtime/GameplayAbilities/Private/AttributeSet.cpp
+++ b/Engine/Source/Runtime/GameplayAbilities/Private/AttributeSet.cpp
@@ -522,6 +522,12 @@ static bool CheckForBadScalableFloats_Prop_r(void* Data, UProperty* Prop, UClass
 					ThisScalableFloat->Curve.CurveTable = nullptr;
 					GCurrentBadScalableFloat.Asset->MarkPackageDirty();
 				}
+				else if (ThisScalableFloat->Curve.CurveTable == nullptr)
+				{
+					// Just fix this case up here
+					ThisScalableFloat->Curve.RowName = NAME_None;
+					GCurrentBadScalableFloat.Asset->MarkPackageDirty();
+				}
 				else
 				{
 					GCurrentBadScalableFloat.Property = Prop;
diff --git a/Engine/Source/Runtime/GameplayAbilities/Private/GameplayCueManager.cpp b/Engine/Source/Runtime/GameplayAbilities/Private/GameplayCueManager.cpp
--- a/Engine/Source/Runtime/GameplayAbilities/Private/GameplayCueManager.cpp
+++ b/Engine/Source/Runtime/GameplayAbilities/Private/GameplayCueManager.cpp
@@ -16,6 +16,9 @@
 #define LOCTEXT_NAMESPACE "GameplayCueManager"
 #endif
 
+int32 LogGameplayCueActorSpawning = 0;
+static FAutoConsoleVariableRef CVarLogGameplayCueActorSpawning(TEXT("AbilitySystem.LogGameplayCueActorSpawning"),	LogGameplayCueActorSpawning, TEXT("Log when we create GameplayCueNotify_Actors"), ECVF_Default	);
+
 int32 DisplayGameplayCues = 0;
 static FAutoConsoleVariableRef CVarDisplayGameplayCues(TEXT("AbilitySystem.DisplayGameplayCues"),	DisplayGameplayCues, TEXT("Display GameplayCue events in world as text."), ECVF_Default	);
 
@@ -45,6 +48,12 @@ UGameplayCueManager::UGameplayCueManager(const FObjectInitializer& PCIP)
 	CurrentWorld = nullptr;
 }
 
+void UGameplayCueManager::OnCreated()
+{
+	FWorldDelegates::OnPostWorldCreation.AddUObject(this, &UGameplayCueManager::OnWorldCreated);
+	FWorldDelegates::OnWorldCleanup.AddUObject(this, &UGameplayCueManager::OnWorldCleanup);
+}
+
 bool IsDedicatedServerForGameplayCue()
 {
 #if WITH_EDITOR
@@ -150,10 +159,15 @@ void UGameplayCueManager::EndGameplayCuesFor(AActor* TargetActor)
 	}
 }
 
+int32 GameplayCueActorRecycle = 1;
+static FAutoConsoleVariableRef CVarGameplayCueActorRecycle(TEXT("AbilitySystem.GameplayCueActorRecycle"), GameplayCueActorRecycle, TEXT("Allow recycling of GameplayCue Actors"), ECVF_Default );
+
 AGameplayCueNotify_Actor* UGameplayCueManager::GetInstancedCueActor(AActor* TargetActor, UClass* CueClass, const FGameplayCueParameters& Parameters)
 {
 	QUICK_SCOPE_CYCLE_COUNTER(STAT_GameplayCueManager_GetInstancedCueActor);
 
+
+	// First, see if this actor already have a GameplayCueNotifyActor already going for this CueClass
 	AGameplayCueNotify_Actor* CDO = Cast<AGameplayCueNotify_Actor>(CueClass->ClassDefaultObject);
 	FGCNotifyActorKey	NotifyKey(TargetActor, CueClass, 
 							CDO->bUniqueInstancePerInstigator ? Parameters.GetInstigator() : nullptr, 
@@ -164,7 +178,7 @@ AGameplayCueNotify_Actor* UGameplayCueManager::GetInstancedCueActor(AActor* Targ
 	{		
 		SpawnedCue = WeakPtrPtr->Get();
 		// If the cue is scheduled to be destroyed, don't reuse it, create a new one instead
-		if (SpawnedCue && SpawnedCue->GetLifeSpan() <= 0.0f)
+		if (SpawnedCue && SpawnedCue->GameplayCuePendingRemove() == false)
 		{
 			return SpawnedCue;
 		}
@@ -173,17 +187,48 @@ AGameplayCueNotify_Actor* UGameplayCueManager::GetInstancedCueActor(AActor* Targ
 	// We don't have an instance for this, and we need one, so make one
 	if (ensure(TargetActor) && ensure(CueClass))
 	{
-		FActorSpawnParameters SpawnParams;
+		AActor* NewOwnerActor = TargetActor;
 #if WITH_EDITOR
 		// Don't set owner if we are using fake CDO actor to do anim previewing
-		SpawnParams.Owner = (TargetActor && TargetActor->HasAnyFlags(RF_ClassDefaultObject) == false ? TargetActor : nullptr);
-#else
-		SpawnParams.Owner = TargetActor;
+		NewOwnerActor= (TargetActor && TargetActor->HasAnyFlags(RF_ClassDefaultObject) == false ? TargetActor : nullptr);
 #endif
-		
-		SpawnedCue = GetWorld()->SpawnActor<AGameplayCueNotify_Actor>(CueClass, TargetActor->GetActorLocation(), TargetActor->GetActorRotation(), SpawnParams);
+
+		// Look to reuse an existing one that is stored on the CDO:
+		if (GameplayCueActorRecycle > 0)
+		{
+			FPreallocationInfo& Info = GetPreallocationInfo(GetWorld());
+			TArray<AGameplayCueNotify_Actor*>* PreallocatedList = Info.PreallocatedInstances.Find(CueClass);
+			if (PreallocatedList && PreallocatedList->Num() > 0)
+			{
+				SpawnedCue = PreallocatedList->Pop(false);
+				checkf(SpawnedCue && SpawnedCue->IsPendingKill() == false, TEXT("Spawned Cue is pending kill or null: %s"), *GetNameSafe(SpawnedCue));
+
+				SpawnedCue->SetActorHiddenInGame(false);
+				SpawnedCue->SetOwner(NewOwnerActor);
+				SpawnedCue->SetActorLocationAndRotation(TargetActor->GetActorLocation(), TargetActor->GetActorRotation());
+			}
+		}
+
+		// If we can't reuse, then spawn a new one
+		if (SpawnedCue == nullptr)
+		{
+			FActorSpawnParameters SpawnParams;
+			SpawnParams.Owner = NewOwnerActor;
+			if (SpawnedCue == nullptr)
+			{
+				if (LogGameplayCueActorSpawning)
+				{
+					ABILITY_LOG(Warning, TEXT("Spawning GameplaycueActor: %s"), *CueClass->GetName());
+				}
+
+				SpawnedCue = GetWorld()->SpawnActor<AGameplayCueNotify_Actor>(CueClass, TargetActor->GetActorLocation(), TargetActor->GetActorRotation(), SpawnParams);
+			}
+		}
+
+		// Associate this GameplayCueNotifyActor with this target actor/key
 		if (ensure(SpawnedCue))
 		{
+			SpawnedCue->NotifyKey = NotifyKey;
 			NotifyMapActor.Add(NotifyKey, SpawnedCue);
 		}
 	}
@@ -191,6 +236,36 @@ AGameplayCueNotify_Actor* UGameplayCueManager::GetInstancedCueActor(AActor* Targ
 	return SpawnedCue;
 }
 
+void UGameplayCueManager::NotifyGameplayCueActorFinished(AGameplayCueNotify_Actor* Actor)
+{
+	if (GameplayCueActorRecycle)
+	{
+		AGameplayCueNotify_Actor* CDO = Actor->GetClass()->GetDefaultObject<AGameplayCueNotify_Actor>();
+		if (CDO && Actor->Recycle())
+		{
+			ensure(Actor->IsPendingKill() == false);
+
+			// Remove this now from our internal map so that it doesn't get reused like a currently active cue would
+			if (TWeakObjectPtr<AGameplayCueNotify_Actor>* WeakPtrPtr = NotifyMapActor.Find(Actor->NotifyKey))
+			{
+				WeakPtrPtr->Reset();
+			}
+
+			Actor->SetActorHiddenInGame(true);
+			Actor->DetachRootComponentFromParent();
+
+			FPreallocationInfo& Info = GetPreallocationInfo(Actor->GetWorld());
+			TArray<AGameplayCueNotify_Actor*>& PreAllocatedList = Info.PreallocatedInstances.FindOrAdd(Actor->GetClass());
+			PreAllocatedList.Push(Actor);
+
+			return;
+		}
+	}	
+
+	// We didn't recycle, so just destroy
+	Actor->Destroy();
+}
+
 // ------------------------------------------------------------------------
 
 void UGameplayCueManager::LoadObjectLibraryFromPaths(const TArray<FString>& InPaths)
@@ -244,7 +319,7 @@ void UGameplayCueManager::LoadObjectLibrary_Internal()
 
 	FScopeCycleCounterUObject PreloadScopeActor(GameplayCueNotifyActorObjectLibrary);
 	GameplayCueNotifyActorObjectLibrary->LoadBlueprintAssetDataFromPaths(LoadedPaths);
-	GameplayCueNotifyStaticObjectLibrary->LoadBlueprintAssetDataFromPaths(LoadedPaths);		//No separate cycle counter for this.
+	GameplayCueNotifyStaticObjectLibrary->LoadBlueprintAssetDataFromPaths(LoadedPaths);
 
 	// ---------------------------------------------------------
 	// Determine loading scheme.
@@ -314,17 +389,27 @@ void UGameplayCueManager::BuildCuesToAddToGlobalSet(const TArray<FAssetData>& As
 
 				if (bAsyncLoadAfterAdd)
 				{
-					StreamableManager.SimpleAsyncLoad(StringRef, FStreamableManager::DefaultAsyncLoadPriority - 1);
+					StreamableManager.RequestAsyncLoad(StringRef, FStreamableDelegate::CreateUObject(this, &UGameplayCueManager::OnGameplayCueNotifyAsyncLoadComplete, StringRef));
 				}
 			}
 			else
 			{
-				ABILITY_LOG(Warning, TEXT("Found GameplayCue tag %s in asset %s but there is no corresponding tag in the GameplayTagMAnager."), **FoundGameplayTag, *Data.PackageName.ToString());
+				ABILITY_LOG(Warning, TEXT("Found GameplayCue tag %s in asset %s but there is no corresponding tag in the GameplayTagManager."), **FoundGameplayTag, *Data.PackageName.ToString());
 			}
 		}
 	}
 }
 
+void UGameplayCueManager::OnGameplayCueNotifyAsyncLoadComplete(FStringAssetReference StringRef)
+{
+	UClass* GCClass = FindObject<UClass>(nullptr, *StringRef.ToString());
+	if (ensure(GCClass))
+	{
+		LoadedGameplayCueNotifyClasses.Add(GCClass);
+		CheckForPreallocation(GCClass);
+	}
+}
+
 int32 UGameplayCueManager::FinishLoadingGameplayCueNotifies()
 {
 	int32 NumLoadeded = 0;
@@ -698,6 +783,146 @@ bool UGameplayCueManager::DoesPendingCueExecuteMatch(FGameplayCuePendingExecute&
 	return true;
 }
 
+void UGameplayCueManager::CheckForPreallocation(UClass* GCClass)
+{
+	if (AGameplayCueNotify_Actor* InstancedCue = Cast<AGameplayCueNotify_Actor>(GCClass->ClassDefaultObject))
+	{
+		if (InstancedCue->NumPreallocatedInstances > 0)
+		{
+			// Add this to the global list
+			GameplayCueClassesForPreallocation.Add(InstancedCue);
+
+			// Add it to any world specific lists
+#if WITH_EDITOR
+			for (FPreallocationInfo& Info : PreallocationInfoList_Internal)
+			{
+				Info.ClassesNeedingPreallocation.Push(InstancedCue);
+			}
+#else
+			PreallocationInfo_Internal.ClassesNeedingPreallocation.Push(InstancedCue);
+#endif
+		}
+	}
+}
+
+// -------------------------------------------------------------
+
+void UGameplayCueManager::ResetPreallocation(UWorld* World)
+{
+	FPreallocationInfo& Info = GetPreallocationInfo(World);
+
+	Info.PreallocatedInstances.Reset();
+	Info.ClassesNeedingPreallocation = GameplayCueClassesForPreallocation;
+}
+
+void UGameplayCueManager::UpdatePreallocation(UWorld* World)
+{
+	FPreallocationInfo& Info = GetPreallocationInfo(World);
+
+	if (Info.ClassesNeedingPreallocation.Num() > 0)
+	{
+		AGameplayCueNotify_Actor* CDO = Info.ClassesNeedingPreallocation.Last();
+		TArray<AGameplayCueNotify_Actor*>& PreallocatedList = Info.PreallocatedInstances.FindOrAdd(CDO->GetClass());
+
+		AGameplayCueNotify_Actor* PrespawnedInstance = Cast<AGameplayCueNotify_Actor>(World->SpawnActor(CDO->GetClass()));
+		ensureMsgf(PrespawnedInstance, TEXT("Failed to prespawn GC notify for: %s"), *GetNameSafe(CDO));
+		if (PrespawnedInstance)
+		{
+			if (LogGameplayCueActorSpawning)
+			{
+				ABILITY_LOG(Warning, TEXT("Prespawning GC %s"), *GetNameSafe(CDO));
+			}
+
+			PreallocatedList.Push(PrespawnedInstance);
+			PrespawnedInstance->SetActorHiddenInGame(true);
+
+			if (PreallocatedList.Num() >= CDO->NumPreallocatedInstances)
+			{
+				Info.ClassesNeedingPreallocation.Pop(false);
+			}
+		}
+	}
+}
+
+FPreallocationInfo& UGameplayCueManager::GetPreallocationInfo(UWorld* World)
+{
+#if WITH_EDITOR
+	for (FPreallocationInfo& Info : PreallocationInfoList_Internal)
+	{
+		if (World == Info.OwningWorld)
+		{
+			return Info;
+		}
+	}
+
+	FPreallocationInfo NewInfo;
+	NewInfo.OwningWorld = World;
+
+	PreallocationInfoList_Internal.Add(NewInfo);
+	return PreallocationInfoList_Internal.Last();
+
+#else
+	return PreallocationInfo_Internal;
+#endif
+
+}
+
+void UGameplayCueManager::OnWorldCreated(UWorld* NewWorld)
+{
+	PreallocationInfo_Internal.PreallocatedInstances.Reset();
+	PreallocationInfo_Internal.OwningWorld = NewWorld;
+}
+
+void UGameplayCueManager::OnWorldCleanup(UWorld* World, bool bSessionEnded, bool bCleanupResources)
+{
+
+#if !(UE_BUILD_SHIPPING || UE_BUILD_TEST)
+	DumpPreallocationStats(World);
+#endif
+
+	if (PreallocationInfo_Internal.OwningWorld == World)
+	{
+		// Reset PreallocationInfo_Internal
+		OnWorldCreated(nullptr);
+	}
+
+#if WITH_EDITOR
+	for (int32 idx=0; idx < PreallocationInfoList_Internal.Num(); ++idx)
+	{
+		if (PreallocationInfoList_Internal[idx].OwningWorld == World)
+		{
+			PreallocationInfoList_Internal.RemoveAtSwap(idx, 1, false);
+			break;
+		}
+	}
+#endif	
+	
+}
+
+void UGameplayCueManager::DumpPreallocationStats(UWorld* World)
+{
+	if (World == nullptr)
+	{
+		return;
+	}
+
+	FPreallocationInfo& Info = GetPreallocationInfo(World);
+	for (auto &It : Info.PreallocatedInstances)
+	{
+		if (UClass* ThisClass = It.Key)
+		{
+			if (AGameplayCueNotify_Actor* CDO = ThisClass->GetDefaultObject<AGameplayCueNotify_Actor>())
+			{
+				TArray<AGameplayCueNotify_Actor*>& List = It.Value;
+				if (List.Num() > CDO->NumPreallocatedInstances)
+				{
+					ABILITY_LOG(Warning, TEXT("Notify class: %s was used simultaneously %d times. The CDO default is %d preallocated instanced."), *ThisClass->GetName(), List.Num(),  CDO->NumPreallocatedInstances); 
+				}
+			}
+		}
+	}
+}
+
 #if WITH_EDITOR
 #undef LOCTEXT_NAMESPACE
 #endif
diff --git a/Engine/Source/Runtime/GameplayAbilities/Private/GameplayCueNotify_Actor.cpp b/Engine/Source/Runtime/GameplayAbilities/Private/GameplayCueNotify_Actor.cpp
--- a/Engine/Source/Runtime/GameplayAbilities/Private/GameplayCueNotify_Actor.cpp
+++ b/Engine/Source/Runtime/GameplayAbilities/Private/GameplayCueNotify_Actor.cpp
@@ -15,6 +15,8 @@ AGameplayCueNotify_Actor::AGameplayCueNotify_Actor(const FObjectInitializer& Obj
 	AutoDestroyDelay = 0.f;
 	bUniqueInstancePerSourceObject = false;
 	bUniqueInstancePerInstigator = false;
+
+	NumPreallocatedInstances = 0;
 }
 
 #if WITH_EDITOR
@@ -106,11 +108,11 @@ void AGameplayCueNotify_Actor::HandleGameplayCue(AActor* MyTarget, EGameplayCueE
 			{
 				if (AutoDestroyDelay > 0.f)
 				{
-					SetLifeSpan(AutoDestroyDelay);
+					GetWorld()->GetTimerManager().SetTimer(FinishTimerHandle, this, &AGameplayCueNotify_Actor::GameplayCueFinishedCallback, AutoDestroyDelay);
 				}
 				else
 				{
-					Destroy();
+					GameplayCueFinishedCallback();
 				}
 			}
 			break;
@@ -125,7 +127,7 @@ void AGameplayCueNotify_Actor::HandleGameplayCue(AActor* MyTarget, EGameplayCueE
 void AGameplayCueNotify_Actor::OnOwnerDestroyed()
 {
 	// May need to do extra cleanup in child classes
-	Destroy();
+	GameplayCueFinishedCallback();
 }
 
 bool AGameplayCueNotify_Actor::OnExecute_Implementation(AActor* MyTarget, FGameplayCueParameters Parameters)
@@ -147,3 +149,23 @@ bool AGameplayCueNotify_Actor::OnRemove_Implementation(AActor* MyTarget, FGamepl
 {
 	return false;
 }
+
+void AGameplayCueNotify_Actor::GameplayCueFinishedCallback()
+{
+	if (FinishTimerHandle.IsValid())
+	{
+		GetWorld()->GetTimerManager().ClearTimer(FinishTimerHandle);
+	}
+	
+	UAbilitySystemGlobals::Get().GetGameplayCueManager()->NotifyGameplayCueActorFinished(this);
+}
+
+bool AGameplayCueNotify_Actor::GameplayCuePendingRemove()
+{
+	return GetLifeSpan() > 0.f || FinishTimerHandle.IsValid() || IsPendingKill();
+}
+
+bool AGameplayCueNotify_Actor::Recycle()
+{
+	return false;
+}
\ No newline at end of file
diff --git a/Engine/Source/Runtime/GameplayAbilities/Private/GameplayCueSet.cpp b/Engine/Source/Runtime/GameplayAbilities/Private/GameplayCueSet.cpp
--- a/Engine/Source/Runtime/GameplayAbilities/Private/GameplayCueSet.cpp
+++ b/Engine/Source/Runtime/GameplayAbilities/Private/GameplayCueSet.cpp
@@ -44,7 +44,7 @@ void UGameplayCueSet::AddCues(const TArray<FGameplayCueReferencePair>& CuesToAdd
 
 			// Check for duplicates: we may want to remove this eventually.
 			bool bDupe = false;
-			for (FGameplayCueNotifyData Data : GameplayCueData)
+			for (FGameplayCueNotifyData& Data : GameplayCueData)
 			{
 				if (Data.GameplayCueTag == GameplayCueTag)
 				{
diff --git a/Engine/Source/Runtime/GameplayAbilities/Public/AbilitySystemBlueprintLibrary.h b/Engine/Source/Runtime/GameplayAbilities/Public/AbilitySystemBlueprintLibrary.h
--- a/Engine/Source/Runtime/GameplayAbilities/Public/AbilitySystemBlueprintLibrary.h
+++ b/Engine/Source/Runtime/GameplayAbilities/Public/AbilitySystemBlueprintLibrary.h
@@ -245,6 +245,10 @@ class GAMEPLAYABILITIES_API UAbilitySystemBlueprintLibrary : public UBlueprintFu
 	UFUNCTION(BlueprintCallable, Category = "Ability|GameplayEffect")
 	static int32 GetActiveGameplayEffectStackCount(FActiveGameplayEffectHandle ActiveHandle);
 
+	/** Returns stack limit count of an active Gameplay Effect. Will return 0 if the GameplayEffect is no longer valid. */
+	UFUNCTION(BlueprintCallable, Category = "Ability|GameplayEffect")
+	static int32 GetActiveGameplayEffectStackLimitCount(FActiveGameplayEffectHandle ActiveHandle);
+
 	UFUNCTION(BlueprintPure, Category = "Ability|GameplayEffect", Meta = (DisplayName = "Get Active GameplayEffect Debug String "))
 	static FString GetActiveGameplayEffectDebugString(FActiveGameplayEffectHandle ActiveHandle);
 };
diff --git a/Engine/Source/Runtime/GameplayAbilities/Public/AttributeSet.h b/Engine/Source/Runtime/GameplayAbilities/Public/AttributeSet.h
--- a/Engine/Source/Runtime/GameplayAbilities/Public/AttributeSet.h
+++ b/Engine/Source/Runtime/GameplayAbilities/Public/AttributeSet.h
@@ -218,7 +218,7 @@ public:
 	{
 		// Error checking: checks if we have a curve table specified but no valid curve entry
 		GetValueAtLevel(1.f);
-		bool bInvalid = (Curve.CurveTable != nullptr && FinalCurve == nullptr);
+		bool bInvalid = (Curve.CurveTable != nullptr || Curve.RowName != NAME_None ) && (FinalCurve == nullptr);
 		return !bInvalid;
 	}
 
diff --git a/Engine/Source/Runtime/GameplayAbilities/Public/GameplayCueManager.h b/Engine/Source/Runtime/GameplayAbilities/Public/GameplayCueManager.h
--- a/Engine/Source/Runtime/GameplayAbilities/Public/GameplayCueManager.h
+++ b/Engine/Source/Runtime/GameplayAbilities/Public/GameplayCueManager.h
@@ -7,6 +7,7 @@ class UGameplayCueSet;
 #include "GameplayTags.h"
 #include "GameplayEffect.h"
 #include "GameplayCueNotify_Actor.h"
+#include "GameplayCue_Types.h"
 #include "GameplayCueManager.generated.h"
 
 /**
@@ -46,64 +47,6 @@ class UGameplayCueSet;
  *	
  */
 
-/** Type of payload to pass along with this cue. */
-UENUM()
-enum class EGameplayCuePayloadType : uint8
-{
-	EffectContext,
-	CueParameters,
-	FromSpec,
-};
-
-/** Structure to keep track of pending gameplay cues that haven't been applied yet. */
-USTRUCT()
-struct FGameplayCuePendingExecute
-{
-	GENERATED_USTRUCT_BODY()
-
-	FGameplayCuePendingExecute()
-	: PayloadType(EGameplayCuePayloadType::EffectContext)
-	, OwningComponent(NULL)
-	{
-	}
-
-	UPROPERTY()
-	FGameplayTag GameplayCueTag;
-	
-	/** Prediction key that spawned this cue */
-	UPROPERTY()
-	FPredictionKey PredictionKey;
-
-	/** What type of payload is attached to this cue */
-	UPROPERTY()
-	EGameplayCuePayloadType PayloadType;
-
-	/** What component to send the cue on */
-	UPROPERTY()
-	UAbilitySystemComponent* OwningComponent;
-
-	/** If this cue is from a spec, here's the copy of that spec */
-	UPROPERTY()
-	FGameplayEffectSpecForRPC FromSpec;
-
-	/** Store the full cue parameters or just the effect context depending on type */
-	UPROPERTY()
-	FGameplayCueParameters CueParameters;
-};
-
-/**
- *	FScopedGameplayCueSendContext
- *	Add this around code that sends multiple gameplay cues to allow grouping them into a smalkler number of cues for more efficient networking
- */
-struct GAMEPLAYABILITIES_API FScopedGameplayCueSendContext
-{
-	FScopedGameplayCueSendContext();
-	~FScopedGameplayCueSendContext();
-};
-
-/** Delegate for when GC notifies are added or removed from manager */
-DECLARE_MULTICAST_DELEGATE(FOnGameplayCueNotifyChange);
-
 /**
  *	A self contained handler of a GameplayCue. These are similar to AnimNotifies in implementation.
  */
@@ -128,6 +71,8 @@ class GAMEPLAYABILITIES_API UGameplayCueManager : public UDataAsset
 	/** Send out any pending cues */
 	virtual void FlushPendingCues();
 
+	virtual void OnCreated();
+
 	/** Process a pending cue, return false if the cue should be rejected. */
 	virtual bool ProcessPendingCueExecute(FGameplayCuePendingExecute& PendingCue);
 
@@ -148,6 +93,22 @@ class GAMEPLAYABILITIES_API UGameplayCueManager : public UDataAsset
 	/** Returns the cached instance cue. Creates it if it doesn't exist */
 	virtual AGameplayCueNotify_Actor* GetInstancedCueActor(AActor* TargetActor, UClass* CueClass, const FGameplayCueParameters& Parameters);
 
+	/** Notify that this actor is finished and should be destroyed or recycled */
+	virtual void NotifyGameplayCueActorFinished(AGameplayCueNotify_Actor* Actor);
+
+	/** Resets preallocation for a given world */
+	void ResetPreallocation(UWorld* World);
+
+	/** Prespawns a single actor for gameplaycue notify actor classes that need prespawning (should be called by outside gamecode, such as gamestate) */
+	void UpdatePreallocation(UWorld* World);
+
+	void OnWorldCreated(UWorld* NewWorld);
+
+	void OnWorldCleanup(UWorld* World, bool bSessionEnded, bool bCleanupResources);
+
+	/** Prints what classess exceeded their preallocation sizes during runtime */
+	void DumpPreallocationStats(UWorld* World);
+
 	// -------------------------------------------------------------
 	//  Loading GameplayCueNotifies from ObjectLibraries
 	// -------------------------------------------------------------
@@ -174,38 +135,8 @@ class GAMEPLAYABILITIES_API UGameplayCueManager : public UDataAsset
 
 	UPROPERTY(transient)
 	FStreamableManager	StreamableManager;
-
-	struct FGCNotifyActorKey
-	{
-		FGCNotifyActorKey(AActor* InTargetActor, UClass* InCueClass, AActor* InInstigatorActor=nullptr, const UObject* InSourceObj=nullptr)
-		{
-			TargetActor = InTargetActor;
-			OptionalInstigatorActor = InInstigatorActor;
-			OptionalSourceObject = InSourceObj;
-			CueClass = InCueClass;
-		}
-
-		TWeakObjectPtr<AActor>	TargetActor;
-		TWeakObjectPtr<AActor>	OptionalInstigatorActor;
-		TWeakObjectPtr<UObject>	OptionalSourceObject;
-		TWeakObjectPtr<UClass>	CueClass;
-
-		FORCEINLINE bool operator==(const FGCNotifyActorKey& Other) const
-		{
-			return TargetActor == Other.TargetActor && CueClass == Other.CueClass &&
-					OptionalInstigatorActor == Other.OptionalInstigatorActor && OptionalSourceObject == Other.OptionalSourceObject;
-		}
-	};
-
-	FORCEINLINE friend uint32 GetTypeHash(const FGCNotifyActorKey& Key)
-	{
-		return GetTypeHash(Key.TargetActor)	^
-				GetTypeHash(Key.OptionalInstigatorActor) ^
-				GetTypeHash(Key.OptionalSourceObject) ^
-				GetTypeHash(Key.CueClass);
-	}
 	
-	TMap<FGCNotifyActorKey, TWeakObjectPtr<AGameplayCueNotify_Actor> >		NotifyMapActor;
+	TMap<FGCNotifyActorKey, TWeakObjectPtr<AGameplayCueNotify_Actor> > NotifyMapActor;
 
 	void PrintGameplayCueNotifyMap();
 
@@ -251,6 +182,18 @@ protected:
 
 	void BuildCuesToAddToGlobalSet(const TArray<FAssetData>& AssetDataList, FName TagPropertyName, bool bAsyncLoadAfterAdd, TArray<struct FGameplayCueReferencePair>& OutCuesToAdd);
 
+	void OnGameplayCueNotifyAsyncLoadComplete(FStringAssetReference StringRef);
+
+	void CheckForPreallocation(UClass* GCClass);
+
+	/** Hardref to the gameplaycue notify classes we have async loaded*/
+	UPROPERTY()
+	TArray<UClass*> LoadedGameplayCueNotifyClasses;
+
+	/** Classes that we need to preallocate instances for */
+	UPROPERTY()
+	TArray<AGameplayCueNotify_Actor*> GameplayCueClassesForPreallocation;
+
 	TArray<FString>	LoadedPaths;
 
 	/** List of gameplay cue executes that haven't been processed yet */
@@ -263,4 +206,12 @@ protected:
 
 	/** Cached world we are currently handling cues for. Used for non instanced GC Notifies that need world. */
 	UWorld* CurrentWorld;
+
+	FPreallocationInfo& GetPreallocationInfo(UWorld* World);
+
+	UPROPERTY(transient)
+	TArray<FPreallocationInfo>	PreallocationInfoList_Internal;
+
+	UPROPERTY(transient)
+	FPreallocationInfo	PreallocationInfo_Internal;
 };
\ No newline at end of file
diff --git a/Engine/Source/Runtime/GameplayAbilities/Public/GameplayCueNotify_Actor.h b/Engine/Source/Runtime/GameplayAbilities/Public/GameplayCueNotify_Actor.h
--- a/Engine/Source/Runtime/GameplayAbilities/Public/GameplayCueNotify_Actor.h
+++ b/Engine/Source/Runtime/GameplayAbilities/Public/GameplayCueNotify_Actor.h
@@ -4,6 +4,7 @@
 
 #include "GameplayTags.h"
 #include "GameplayEffect.h"
+#include "GameplayCue_Types.h"
 #include "GameplayCueNotify_Actor.generated.h"
 
 
@@ -38,6 +39,13 @@ class GAMEPLAYABILITIES_API AGameplayCueNotify_Actor : public AActor
 
 	virtual void HandleGameplayCue(AActor* MyTarget, EGameplayCueEvent::Type EventType, FGameplayCueParameters Parameters);
 
+	virtual void GameplayCueFinishedCallback();
+
+	virtual bool GameplayCuePendingRemove();
+
+	/** Reset all state so that it can be reused. Return false if this class connot be recycled */
+	virtual bool Recycle();
+
 #if WITH_EDITOR
 	virtual void PostEditChangeProperty(FPropertyChangedEvent& PropertyChangedEvent) override;
 #endif // WITH_EDITOR
@@ -91,6 +99,14 @@ class GAMEPLAYABILITIES_API AGameplayCueNotify_Actor : public AActor
 	UPROPERTY(EditDefaultsOnly, Category = GameplayCue)
 	bool bUniqueInstancePerSourceObject;
 
+	/** How many instances of the gameplay cue to preallocate */
+	UPROPERTY(EditDefaultsOnly, Category = GameplayCue)
+	int32 NumPreallocatedInstances;
+
+	FGCNotifyActorKey NotifyKey;
+
 private:
 	virtual void DeriveGameplayCueTagFromAssetName();
+
+	FTimerHandle FinishTimerHandle;
 };
diff --git a/Engine/Source/Runtime/GameplayAbilities/Public/GameplayCue_Types.h b/Engine/Source/Runtime/GameplayAbilities/Public/GameplayCue_Types.h
new file mode 100644
--- /dev/null
+++ b/Engine/Source/Runtime/GameplayAbilities/Public/GameplayCue_Types.h
@@ -0,0 +1,118 @@
+// Copyright 1998-2015 Epic Games, Inc. All Rights Reserved.
+
+#pragma once
+
+class UGameplayCueSet;
+class AGameplayCueNotify_Actor;
+
+#include "GameplayTags.h"
+#include "GameplayEffect.h"
+#include "GameplayCue_Types.generated.h"
+
+UENUM()
+enum class EGameplayCuePayloadType : uint8
+{
+	EffectContext,
+	CueParameters,
+	FromSpec,
+};
+
+/** Structure to keep track of pending gameplay cues that haven't been applied yet. */
+USTRUCT()
+struct FGameplayCuePendingExecute
+{
+	GENERATED_USTRUCT_BODY()
+
+	FGameplayCuePendingExecute()
+	: PayloadType(EGameplayCuePayloadType::EffectContext)
+	, OwningComponent(NULL)
+	{
+	}
+
+	UPROPERTY()
+	FGameplayTag GameplayCueTag;
+	
+	/** Prediction key that spawned this cue */
+	UPROPERTY()
+	FPredictionKey PredictionKey;
+
+	/** What type of payload is attached to this cue */
+	UPROPERTY()
+	EGameplayCuePayloadType PayloadType;
+
+	/** What component to send the cue on */
+	UPROPERTY()
+	UAbilitySystemComponent* OwningComponent;
+
+	/** If this cue is from a spec, here's the copy of that spec */
+	UPROPERTY()
+	FGameplayEffectSpecForRPC FromSpec;
+
+	/** Store the full cue parameters or just the effect context depending on type */
+	UPROPERTY()
+	FGameplayCueParameters CueParameters;
+};
+
+/** Struct for pooling and preallocating gameplaycuenotify_actor classes. This data is per world and used to track what actors are available to recycle and which classes need to preallocate instances of those actors */
+USTRUCT()
+struct FPreallocationInfo
+{
+	GENERATED_USTRUCT_BODY()
+
+	TMap<UClass*, TArray<AGameplayCueNotify_Actor*> >	PreallocatedInstances;
+
+	UPROPERTY(transient)
+	TArray<AGameplayCueNotify_Actor*>	ClassesNeedingPreallocation;
+
+	UPROPERTY(transient)
+	UWorld* OwningWorld;
+};
+
+/** Struct that is used by the gameplaycue manager to tie an instanced gameplaycue to the calling gamecode. Usually this is just the target actor, but can also be unique per instigator/sourceobject */
+struct FGCNotifyActorKey
+{
+	FGCNotifyActorKey()
+	{
+
+	}
+
+	FGCNotifyActorKey(AActor* InTargetActor, UClass* InCueClass, AActor* InInstigatorActor=nullptr, const UObject* InSourceObj=nullptr)
+	{
+		TargetActor = InTargetActor;
+		OptionalInstigatorActor = InInstigatorActor;
+		OptionalSourceObject = InSourceObj;
+		CueClass = InCueClass;
+	}
+
+	TWeakObjectPtr<AActor>	TargetActor;
+	TWeakObjectPtr<AActor>	OptionalInstigatorActor;
+	TWeakObjectPtr<UObject>	OptionalSourceObject;
+	TWeakObjectPtr<UClass>	CueClass;
+
+	FORCEINLINE bool operator==(const FGCNotifyActorKey& Other) const
+	{
+		return TargetActor == Other.TargetActor && CueClass == Other.CueClass &&
+				OptionalInstigatorActor == Other.OptionalInstigatorActor && OptionalSourceObject == Other.OptionalSourceObject;
+	}
+};
+
+FORCEINLINE uint32 GetTypeHash(const FGCNotifyActorKey& Key)
+{
+	return GetTypeHash(Key.TargetActor)	^
+			GetTypeHash(Key.OptionalInstigatorActor) ^
+			GetTypeHash(Key.OptionalSourceObject) ^
+			GetTypeHash(Key.CueClass);
+}
+
+/**
+ *	FScopedGameplayCueSendContext
+ *	Add this around code that sends multiple gameplay cues to allow grouping them into a smalkler number of cues for more efficient networking
+ */
+struct GAMEPLAYABILITIES_API FScopedGameplayCueSendContext
+{
+	FScopedGameplayCueSendContext();
+	~FScopedGameplayCueSendContext();
+};
+
+/** Delegate for when GC notifies are added or removed from manager */
+DECLARE_MULTICAST_DELEGATE(FOnGameplayCueNotifyChange);
\ No newline at end of file
diff --git a/Engine/Source/Runtime/MovieScene/Private/MovieSceneClipboard.cpp b/Engine/Source/Runtime/MovieScene/Private/MovieSceneClipboard.cpp
--- a/Engine/Source/Runtime/MovieScene/Private/MovieSceneClipboard.cpp
+++ b/Engine/Source/Runtime/MovieScene/Private/MovieSceneClipboard.cpp
@@ -191,21 +191,25 @@ FMovieSceneClipboard FMovieSceneClipboardBuilder::Commit(TOptional<float> CopyRe
 		}
 	}
 
-	for (auto& Pair : TrackIndex)
+	if (CopyRelativeTo.IsSet())
 	{
-		for (FMovieSceneClipboardKeyTrack& Track : Pair.Value)
+		for (auto& Pair : TrackIndex)
 		{
-			Track.IterateKeys([&](FMovieSceneClipboardKey& Key){
-				Key.SetTime(Key.GetTime() - CopyRelativeTo.GetValue());
-				return true;
-			});
+			for (FMovieSceneClipboardKeyTrack& Track : Pair.Value)
+			{
+				Track.IterateKeys([&](FMovieSceneClipboardKey& Key){
+					Key.SetTime(Key.GetTime() - CopyRelativeTo.GetValue());
+					return true;
+				});
+			}
+
+			Clipboard.KeyTrackGroups.Add(MoveTemp(Pair.Value));
 		}
 
-		Clipboard.KeyTrackGroups.Add(MoveTemp(Pair.Value));
+		Clipboard.GetEnvironment().CardinalTime = CopyRelativeTo.GetValue();
 	}
+	
 	TrackIndex.Reset();
-
-	Clipboard.GetEnvironment().CardinalTime = CopyRelativeTo.GetValue();
 	return Clipboard;
 }
 
diff --git a/Engine/Source/Runtime/MovieScene/Public/MovieSceneClipboard.h b/Engine/Source/Runtime/MovieScene/Public/MovieSceneClipboard.h
--- a/Engine/Source/Runtime/MovieScene/Public/MovieSceneClipboard.h
+++ b/Engine/Source/Runtime/MovieScene/Public/MovieSceneClipboard.h
@@ -269,7 +269,7 @@ struct FMovieSceneClipboardEnvironment
 {
 	FMovieSceneClipboardEnvironment()
 		: CardinalTime(0.f)
-		, DateTime(FDateTime::Now())
+		, DateTime(FDateTime::UtcNow())
 	{}
 
 	/** The cardinal time for a copy-paste operation. Keys are copied with absolute time values */
diff --git a/Engine/Source/Runtime/MovieSceneCapture/Private/MovieSceneCapture.cpp b/Engine/Source/Runtime/MovieSceneCapture/Private/MovieSceneCapture.cpp
--- a/Engine/Source/Runtime/MovieSceneCapture/Private/MovieSceneCapture.cpp
+++ b/Engine/Source/Runtime/MovieSceneCapture/Private/MovieSceneCapture.cpp
@@ -58,7 +58,7 @@ UMovieSceneCapture::UMovieSceneCapture(const FObjectInitializer& Initializer)
 		InheritedCommandLineArguments.AppendChar(' ');
 	}
 
-	AdditionalCommandLineArguments += TEXT("-NoLoadingScreen");
+	AdditionalCommandLineArguments += TEXT("-NoLoadingScreen -NOSCREENMESSAGES -ForceRes");
 
 	Handle = FUniqueMovieSceneCaptureHandle();
 
diff --git a/Engine/Source/Runtime/MovieSceneTracks/Private/Sections/MovieSceneMaterialParameterSection.cpp b/Engine/Source/Runtime/MovieSceneTracks/Private/Sections/MovieSceneMaterialParameterSection.cpp
deleted file mode 100644
diff --git a/Engine/Source/Runtime/MovieSceneTracks/Private/Sections/MovieSceneParameterSection.cpp b/Engine/Source/Runtime/MovieSceneTracks/Private/Sections/MovieSceneParameterSection.cpp
new file mode 100644
--- /dev/null
+++ b/Engine/Source/Runtime/MovieSceneTracks/Private/Sections/MovieSceneParameterSection.cpp
@@ -0,0 +1,213 @@
+// Copyright 1998-2015 Epic Games, Inc. All Rights Reserved.
+
+#include "MovieSceneTracksPrivatePCH.h"
+#include "MovieSceneParameterSection.h"
+
+FScalarParameterNameAndCurve::FScalarParameterNameAndCurve( FName InParameterName )
+{
+	ParameterName = InParameterName;
+}
+
+FVectorParameterNameAndCurves::FVectorParameterNameAndCurves( FName InParameterName )
+{
+	ParameterName = InParameterName;
+}
+
+FColorParameterNameAndCurves::FColorParameterNameAndCurves( FName InParameterName )
+{
+	ParameterName = InParameterName;
+}
+
+UMovieSceneParameterSection::UMovieSceneParameterSection( const FObjectInitializer& ObjectInitializer )
+	: Super( ObjectInitializer )
+{
+}
+
+void UMovieSceneParameterSection::AddScalarParameterKey( FName InParameterName, float InTime, float InValue )
+{
+	FRichCurve* ExistingCurve = nullptr;
+	for ( FScalarParameterNameAndCurve& ScalarParameterNameAndCurve : ScalarParameterNamesAndCurves )
+	{
+		if ( ScalarParameterNameAndCurve.ParameterName == InParameterName )
+		{
+			ExistingCurve = &ScalarParameterNameAndCurve.ParameterCurve;
+			break;
+		}
+	}
+	if ( ExistingCurve == nullptr )
+	{
+		int32 NewIndex = ScalarParameterNamesAndCurves.Add( FScalarParameterNameAndCurve( InParameterName ) );
+		ScalarParameterNamesAndCurves[NewIndex].Index = ScalarParameterNamesAndCurves.Num() + VectorParameterNamesAndCurves.Num() - 1;
+		ExistingCurve = &ScalarParameterNamesAndCurves[NewIndex].ParameterCurve;
+	}
+	ExistingCurve->AddKey(InTime, InValue);
+}
+
+void UMovieSceneParameterSection::AddVectorParameterKey( FName InParameterName, float InTime, FVector InValue )
+{
+	FVectorParameterNameAndCurves* ExistingCurves = nullptr;
+	for ( FVectorParameterNameAndCurves& VectorParameterNameAndCurve : VectorParameterNamesAndCurves )
+	{
+		if ( VectorParameterNameAndCurve.ParameterName == InParameterName )
+		{
+			ExistingCurves = &VectorParameterNameAndCurve;
+			break;
+		}
+	}
+	if ( ExistingCurves == nullptr )
+	{
+		int32 NewIndex = VectorParameterNamesAndCurves.Add( FVectorParameterNameAndCurves( InParameterName ) );
+		VectorParameterNamesAndCurves[NewIndex].Index = VectorParameterNamesAndCurves.Num() + VectorParameterNamesAndCurves.Num() - 1;
+		ExistingCurves = &VectorParameterNamesAndCurves[NewIndex];
+	}
+	ExistingCurves->XCurve.AddKey( InTime, InValue.X );
+	ExistingCurves->YCurve.AddKey( InTime, InValue.Y );
+	ExistingCurves->ZCurve.AddKey( InTime, InValue.Z );
+}
+
+void UMovieSceneParameterSection::AddColorParameterKey( FName InParameterName, float InTime, FLinearColor InValue )
+{
+	FColorParameterNameAndCurves* ExistingCurves = nullptr;
+	for ( FColorParameterNameAndCurves& ColorParameterNameAndCurve : ColorParameterNamesAndCurves )
+	{
+		if ( ColorParameterNameAndCurve.ParameterName == InParameterName )
+		{
+			ExistingCurves = &ColorParameterNameAndCurve;
+			break;
+		}
+	}
+	if ( ExistingCurves == nullptr )
+	{
+		int32 NewIndex = ColorParameterNamesAndCurves.Add( FColorParameterNameAndCurves( InParameterName ) );
+		ColorParameterNamesAndCurves[NewIndex].Index = ColorParameterNamesAndCurves.Num() + ColorParameterNamesAndCurves.Num() - 1;
+		ExistingCurves = &ColorParameterNamesAndCurves[NewIndex];
+	}
+	ExistingCurves->RedCurve.AddKey( InTime, InValue.R );
+	ExistingCurves->GreenCurve.AddKey( InTime, InValue.G );
+	ExistingCurves->BlueCurve.AddKey( InTime, InValue.B );
+	ExistingCurves->AlphaCurve.AddKey( InTime, InValue.A );
+}
+
+bool UMovieSceneParameterSection::RemoveScalarParameter( FName InParameterName )
+{
+	for ( int32 i = 0; i < ScalarParameterNamesAndCurves.Num(); i++ )
+	{
+		if ( ScalarParameterNamesAndCurves[i].ParameterName == InParameterName )
+		{
+			ScalarParameterNamesAndCurves.RemoveAt(i);
+			UpdateParameterIndicesFromRemoval(i);
+			return true;
+		}
+	}
+	return false;
+}
+
+bool UMovieSceneParameterSection::RemoveVectorParameter( FName InParameterName )
+{
+	for ( int32 i = 0; i < VectorParameterNamesAndCurves.Num(); i++ )
+	{
+		if ( VectorParameterNamesAndCurves[i].ParameterName == InParameterName )
+		{
+			VectorParameterNamesAndCurves.RemoveAt( i );
+			UpdateParameterIndicesFromRemoval( i );
+			return true;
+		}
+	}
+	return false;
+}
+
+bool UMovieSceneParameterSection::RemoveColorParameter( FName InParameterName )
+{
+	for ( int32 i = 0; i < ColorParameterNamesAndCurves.Num(); i++ )
+	{
+		if ( ColorParameterNamesAndCurves[i].ParameterName == InParameterName )
+		{
+			ColorParameterNamesAndCurves.RemoveAt( i );
+			UpdateParameterIndicesFromRemoval( i );
+			return true;
+		}
+	}
+	return false;
+}
+
+TArray<FScalarParameterNameAndCurve>* UMovieSceneParameterSection::GetScalarParameterNamesAndCurves()
+{
+	return &ScalarParameterNamesAndCurves;
+}
+
+TArray<FVectorParameterNameAndCurves>* UMovieSceneParameterSection::GetVectorParameterNamesAndCurves()
+{
+	return &VectorParameterNamesAndCurves;
+}
+
+TArray<FColorParameterNameAndCurves>* UMovieSceneParameterSection::GetColorParameterNamesAndCurves()
+{
+	return &ColorParameterNamesAndCurves;
+}
+
+void UMovieSceneParameterSection::GetParameterNames( TSet<FName>& ParameterNames ) const
+{
+	for ( const FScalarParameterNameAndCurve& ScalarParameterNameAndCurve : ScalarParameterNamesAndCurves )
+	{
+		ParameterNames.Add( ScalarParameterNameAndCurve.ParameterName );
+	}
+	for ( const FVectorParameterNameAndCurves& VectorParameterNameAndCurves : VectorParameterNamesAndCurves )
+	{
+		ParameterNames.Add( VectorParameterNameAndCurves.ParameterName );
+	}
+	for ( const FColorParameterNameAndCurves& ColorParameterNameAndCurves : ColorParameterNamesAndCurves )
+	{
+		ParameterNames.Add( ColorParameterNameAndCurves.ParameterName );
+	}
+}
+
+void UMovieSceneParameterSection::Eval( float Position,
+	TArray<FScalarParameterNameAndValue>& OutScalarValues,
+	TArray<FVectorParameterNameAndValue>& OutVectorValues,
+	TArray<FColorParameterNameAndValue>& OutColorValues ) const
+{
+	for ( const FScalarParameterNameAndCurve& ScalarParameterNameAndCurve : ScalarParameterNamesAndCurves )
+	{
+		OutScalarValues.Add( FScalarParameterNameAndValue( ScalarParameterNameAndCurve.ParameterName, ScalarParameterNameAndCurve.ParameterCurve.Eval( Position ) ) );
+	}
+	for ( const FVectorParameterNameAndCurves& VectorParameterNameAndCurves : VectorParameterNamesAndCurves )
+	{
+		OutVectorValues.Add( FVectorParameterNameAndValue( VectorParameterNameAndCurves.ParameterName, FVector(
+			VectorParameterNameAndCurves.XCurve.Eval( Position ),
+			VectorParameterNameAndCurves.YCurve.Eval( Position ),
+			VectorParameterNameAndCurves.ZCurve.Eval( Position ) ) ) );
+	}
+	for ( const FColorParameterNameAndCurves& ColorParameterNameAndCurves : ColorParameterNamesAndCurves )
+	{
+		OutColorValues.Add( FColorParameterNameAndValue( ColorParameterNameAndCurves.ParameterName, FLinearColor(
+			ColorParameterNameAndCurves.RedCurve.Eval( Position ),
+			ColorParameterNameAndCurves.GreenCurve.Eval( Position ),
+			ColorParameterNameAndCurves.BlueCurve.Eval( Position ),
+			ColorParameterNameAndCurves.AlphaCurve.Eval( Position ) ) ) );
+	}
+}
+
+void UMovieSceneParameterSection::UpdateParameterIndicesFromRemoval( int32 RemovedIndex )
+{
+	for ( FScalarParameterNameAndCurve& ScalarParameterAndCurve : ScalarParameterNamesAndCurves )
+	{
+		if ( ScalarParameterAndCurve.Index > RemovedIndex )
+		{
+			ScalarParameterAndCurve.Index--;
+		}
+	}
+	for ( FVectorParameterNameAndCurves& VectorParameterAndCurves : VectorParameterNamesAndCurves )
+	{
+		if ( VectorParameterAndCurves.Index > RemovedIndex )
+		{
+			VectorParameterAndCurves.Index--;
+		}
+	}
+	for ( FColorParameterNameAndCurves& ColorParameterAndCurves : ColorParameterNamesAndCurves )
+	{
+		if ( ColorParameterAndCurves.Index > RemovedIndex )
+		{
+			ColorParameterAndCurves.Index--;
+		}
+	}
+}
\ No newline at end of file
diff --git a/Engine/Source/Runtime/MovieSceneTracks/Private/TrackInstances/MovieSceneMaterialTrackInstance.cpp b/Engine/Source/Runtime/MovieSceneTracks/Private/TrackInstances/MovieSceneMaterialTrackInstance.cpp
--- a/Engine/Source/Runtime/MovieSceneTracks/Private/TrackInstances/MovieSceneMaterialTrackInstance.cpp
+++ b/Engine/Source/Runtime/MovieSceneTracks/Private/TrackInstances/MovieSceneMaterialTrackInstance.cpp
@@ -56,7 +56,7 @@ void FMovieSceneMaterialTrackInstance::RestoreState(const TArray<UObject*>& Runt
 void FMovieSceneMaterialTrackInstance::Update( float Position, float LastPosition, const TArray<UObject*>& RuntimeObjects, class IMovieScenePlayer& Player, FMovieSceneSequenceInstance& SequenceInstance, EMovieSceneUpdatePass UpdatePass ) 
 {
 	TArray<FScalarParameterNameAndValue> ScalarValues;
-	TArray<FVectorParameterNameAndValue> VectorValues;
+	TArray<FColorParameterNameAndValue> VectorValues;
 	MaterialTrack->Eval( Position, ScalarValues, VectorValues );
 
 	// Iterate from back to front to allow for fast remove of invalid weak pointers.
@@ -69,7 +69,7 @@ void FMovieSceneMaterialTrackInstance::Update( float Position, float LastPositio
 			{
 				DynamicMaterialInstance->SetScalarParameterValue( ScalarValue.ParameterName, ScalarValue.Value );
 			}
-			for ( const FVectorParameterNameAndValue& VectorValue : VectorValues )
+			for ( const FColorParameterNameAndValue& VectorValue : VectorValues )
 			{
 				DynamicMaterialInstance->SetVectorParameterValue( VectorValue.ParameterName, VectorValue.Value );
 			}
diff --git a/Engine/Source/Runtime/MovieSceneTracks/Private/TrackInstances/MovieSceneParticleParameterTrackInstance.cpp b/Engine/Source/Runtime/MovieSceneTracks/Private/TrackInstances/MovieSceneParticleParameterTrackInstance.cpp
new file mode 100644
--- /dev/null
+++ b/Engine/Source/Runtime/MovieSceneTracks/Private/TrackInstances/MovieSceneParticleParameterTrackInstance.cpp
@@ -0,0 +1,175 @@
+// Copyright 1998-2015 Epic Games, Inc. All Rights Reserved.
+
+#include "MovieSceneTracksPrivatePCH.h"
+#include "MovieSceneParticleParameterTrack.h"
+#include "MovieSceneParticleParameterTrackInstance.h"
+
+FMovieSceneParticleParameterTrackInstance::FMovieSceneParticleParameterTrackInstance( UMovieSceneParticleParameterTrack& InParticleParameterTrack )
+{
+	ParticleParameterTrack = &InParticleParameterTrack;
+}
+
+void FMovieSceneParticleParameterTrackInstance::SaveState(const TArray<UObject*>& RuntimeObjects, IMovieScenePlayer& Player, FMovieSceneSequenceInstance& SequenceInstance)
+{
+	for( UObject* Object : RuntimeObjects )
+	{
+		AEmitter* EmitterActor = Cast<AEmitter>(Object);
+		if ( EmitterActor != nullptr )
+		{
+			UParticleSystemComponent* ParticleSystemComponent = EmitterActor->GetParticleSystemComponent();
+			if ( ParticleSystemComponent != nullptr )
+			{
+				TSharedPtr<FInitialParameterValues> InitialParameterValues;
+				TSharedPtr<FInitialParameterValues>* InitialParameterValuesPtr = ObjectToInitialParameterValuesMap.Find( FObjectKey( Object ) );
+				if ( InitialParameterValuesPtr == nullptr )
+				{
+					InitialParameterValues = MakeShareable( new FInitialParameterValues() );
+					ObjectToInitialParameterValuesMap.Add( FObjectKey( Object ), InitialParameterValues );
+				}
+				else
+				{
+					InitialParameterValues = *InitialParameterValuesPtr;
+				}
+
+				const TArray<FParticleSysParam> InstanceParameters = ParticleSystemComponent->GetAsyncInstanceParameters();
+				for ( const FParticleSysParam& ParticleSystemParameter : InstanceParameters )
+				{
+					switch ( ParticleSystemParameter.ParamType )
+					{
+					case PSPT_Scalar:
+						InitialParameterValues->ScalarNameToValueMap.Add( ParticleSystemParameter.Name, ParticleSystemParameter.Scalar );
+						break;
+					case PSPT_Vector:
+						InitialParameterValues->VectorNameToValueMap.Add( ParticleSystemParameter.Name, ParticleSystemParameter.Vector );
+						break;
+					case PSPT_Color:
+						InitialParameterValues->ColorNameToValueMap.Add( ParticleSystemParameter.Name, ParticleSystemParameter.Color );
+						break;
+					}
+				}
+			}
+		}
+	}
+}
+
+void FMovieSceneParticleParameterTrackInstance::RestoreState(const TArray<UObject*>& RuntimeObjects, IMovieScenePlayer& Player, FMovieSceneSequenceInstance& SequenceInstance)
+{
+	for ( UObject* Object : RuntimeObjects )
+	{
+		TSharedPtr<FInitialParameterValues>* InitialParameterValuesPtr = ObjectToInitialParameterValuesMap.Find(FObjectKey(Object));
+		if ( InitialParameterValuesPtr != nullptr )
+		{
+			TSharedPtr<FInitialParameterValues> InitialParameterValues = *InitialParameterValuesPtr;
+
+			AEmitter* EmitterActor = Cast<AEmitter>( Object );
+			if ( EmitterActor != nullptr )
+			{
+				UParticleSystemComponent* ParticleSystemComponent = EmitterActor->GetParticleSystemComponent();
+				if ( ParticleSystemComponent != nullptr )
+				{
+					const TArray<FParticleSysParam> InstanceParameters = ParticleSystemComponent->GetAsyncInstanceParameters();
+
+					// Set saved parameters back to their initial values, and collect new parameters to be cleared.
+					for ( const FParticleSysParam& ParticleSystemParameter : InstanceParameters )
+					{
+						switch ( ParticleSystemParameter.ParamType )
+						{
+						case PSPT_Scalar:
+						{
+							float* InitialFloatValue = InitialParameterValues->ScalarNameToValueMap.Find(ParticleSystemParameter.Name);
+							if ( InitialFloatValue != nullptr )
+							{
+								ParticleSystemComponent->SetFloatParameter(ParticleSystemParameter.Name, *InitialFloatValue);
+							}
+							break;
+						}
+						case PSPT_Vector:
+						{
+							FVector* InitialVectorValue = InitialParameterValues->VectorNameToValueMap.Find( ParticleSystemParameter.Name );
+							if ( InitialVectorValue != nullptr )
+							{
+								ParticleSystemComponent->SetVectorParameter( ParticleSystemParameter.Name, *InitialVectorValue );
+							}
+							break;
+						}
+						case PSPT_Color:
+						{
+							FLinearColor* InitialColorValue = InitialParameterValues->ColorNameToValueMap.Find( ParticleSystemParameter.Name );
+							if ( InitialColorValue != nullptr )
+							{
+								ParticleSystemComponent->SetColorParameter( ParticleSystemParameter.Name, *InitialColorValue );
+							}
+							break;
+						}
+						}
+					}
+				}
+			}
+		}
+	}
+	ObjectToInitialParameterValuesMap.Empty();
+}
+
+void FMovieSceneParticleParameterTrackInstance::Update( float Position, float LastPosition, const TArray<UObject*>& RuntimeObjects, class IMovieScenePlayer& Player, FMovieSceneSequenceInstance& SequenceInstance, EMovieSceneUpdatePass UpdatePass )
+{
+	TArray<FScalarParameterNameAndValue> ScalarNamesAndValues;
+	TArray<FVectorParameterNameAndValue> VectorNamesAndValues;
+	TArray<FColorParameterNameAndValue> ColorNamesAndValues;
+	ParticleParameterTrack->Eval( Position, ScalarNamesAndValues, VectorNamesAndValues, ColorNamesAndValues);
+	for ( UParticleSystemComponent* ParticleSystemComponent : ParticleSystemComponents )
+	{
+		for ( const FScalarParameterNameAndValue& ScalarNameAndValue : ScalarNamesAndValues )
+		{
+			ParticleSystemComponent->SetFloatParameter( ScalarNameAndValue.ParameterName, ScalarNameAndValue.Value );
+		}
+		for ( const FVectorParameterNameAndValue& VectorNameAndValue : VectorNamesAndValues )
+		{
+			ParticleSystemComponent->SetVectorParameter( VectorNameAndValue.ParameterName, VectorNameAndValue.Value );
+		}
+		for ( const FColorParameterNameAndValue& ColorNameAndValue : ColorNamesAndValues )
+		{
+			ParticleSystemComponent->SetColorParameter( ColorNameAndValue.ParameterName, ColorNameAndValue.Value );
+		}
+	}
+}
+
+void FMovieSceneParticleParameterTrackInstance::RefreshInstance( const TArray<UObject*>& RuntimeObjects, class IMovieScenePlayer& Player, FMovieSceneSequenceInstance& SequenceInstance )
+{
+	ParticleSystemComponents.Empty();
+
+	// Collection the animated parameter names.
+	TSet<FName> UsedParameterNames;
+	for ( UMovieSceneSection* Section : ParticleParameterTrack->GetAllSections() )
+	{
+		UMovieSceneParameterSection* ParameterSection = Cast<UMovieSceneParameterSection>(Section);
+		ParameterSection->GetParameterNames(UsedParameterNames);
+	}
+
+	// Cache the particle components which have matching parameter names.
+	for ( UObject* Object : RuntimeObjects )
+	{
+		AEmitter* EmitterActor = Cast<AEmitter>( Object );
+		if ( EmitterActor != nullptr )
+		{
+			UParticleSystemComponent* ParticleSystemComponent = EmitterActor->GetParticleSystemComponent();
+			if ( ParticleSystemComponent != nullptr )
+			{
+				const TArray<FParticleSysParam> InstanceParameters = ParticleSystemComponent->GetAsyncInstanceParameters();
+				for ( const FParticleSysParam& ParticleSystemParameter : InstanceParameters )
+				{
+					if ( UsedParameterNames.Contains( ParticleSystemParameter.Name ) )
+					{
+						ParticleSystemComponents.Add(ParticleSystemComponent);
+						break;
+					}
+				}
+			}
+		}
+	}
+}
+
+void FMovieSceneParticleParameterTrackInstance::ClearInstance( IMovieScenePlayer& Player, FMovieSceneSequenceInstance& SequenceInstance )
+{
+	ParticleSystemComponents.Empty();
+	ObjectToInitialParameterValuesMap.Empty();
+}
diff --git a/Engine/Source/Runtime/MovieSceneTracks/Private/TrackInstances/MovieSceneShotTrackInstance.cpp b/Engine/Source/Runtime/MovieSceneTracks/Private/TrackInstances/MovieSceneShotTrackInstance.cpp
--- a/Engine/Source/Runtime/MovieSceneTracks/Private/TrackInstances/MovieSceneShotTrackInstance.cpp
+++ b/Engine/Source/Runtime/MovieSceneTracks/Private/TrackInstances/MovieSceneShotTrackInstance.cpp
@@ -39,7 +39,12 @@ void FMovieSceneShotTrackInstance::RefreshInstance(const TArray<UObject*>& Runti
 	}
 }
 
-
+	
+void FMovieSceneShotTrackInstance::RestoreState(const TArray<UObject*>& RuntimeObjects, IMovieScenePlayer& Player, FMovieSceneSequenceInstance& SequenceInstance)
+{
+	Player.UpdateCameraCut(nullptr, nullptr);
+}
+	
 void FMovieSceneShotTrackInstance::Update(float Position, float LastPosition, const TArray<UObject*>& RuntimeObjects, class IMovieScenePlayer& Player, FMovieSceneSequenceInstance& SequenceInstance, EMovieSceneUpdatePass UpdatePass) 
 {
 	UObject* CameraObject = nullptr;
diff --git a/Engine/Source/Runtime/MovieSceneTracks/Private/Tracks/MovieScene3DAttachTrack.cpp b/Engine/Source/Runtime/MovieSceneTracks/Private/Tracks/MovieScene3DAttachTrack.cpp
--- a/Engine/Source/Runtime/MovieSceneTracks/Private/Tracks/MovieScene3DAttachTrack.cpp
+++ b/Engine/Source/Runtime/MovieSceneTracks/Private/Tracks/MovieScene3DAttachTrack.cpp
@@ -25,6 +25,7 @@ void UMovieScene3DAttachTrack::AddConstraint(float KeyTime, float ConstraintEndT
 {
 	// add the section
 	UMovieScene3DAttachSection* NewSection = NewObject<UMovieScene3DAttachSection>(this);
+	NewSection->SetFlags(RF_Transactional);
 	NewSection->AddAttach(KeyTime, ConstraintEndTime, ConstraintId);
 	NewSection->InitialPlacement( ConstraintSections, KeyTime, ConstraintEndTime, SupportsMultipleRows() );
 	NewSection->AttachSocketName = SocketName;
diff --git a/Engine/Source/Runtime/MovieSceneTracks/Private/Tracks/MovieSceneMaterialTrack.cpp b/Engine/Source/Runtime/MovieSceneTracks/Private/Tracks/MovieSceneMaterialTrack.cpp
--- a/Engine/Source/Runtime/MovieSceneTracks/Private/Tracks/MovieSceneMaterialTrack.cpp
+++ b/Engine/Source/Runtime/MovieSceneTracks/Private/Tracks/MovieSceneMaterialTrack.cpp
@@ -1,7 +1,7 @@
 // Copyright 1998-2015 Epic Games, Inc. All Rights Reserved.
 
 #include "MovieSceneTracksPrivatePCH.h"
-#include "MovieSceneMaterialParameterSection.h"
+#include "MovieSceneParameterSection.h"
 #include "MovieSceneMaterialTrack.h"
 #include "IMovieScenePlayer.h"
 #include "MovieSceneMaterialTrackInstance.h"
@@ -13,7 +13,7 @@ UMovieSceneMaterialTrack::UMovieSceneMaterialTrack( const FObjectInitializer& Ob
 
 UMovieSceneSection* UMovieSceneMaterialTrack::CreateNewSection()
 {
-	return NewObject<UMovieSceneMaterialParameterSection>( this, UMovieSceneMaterialParameterSection::StaticClass(), NAME_None, RF_Transactional );
+	return NewObject<UMovieSceneParameterSection>( this, UMovieSceneParameterSection::StaticClass(), NAME_None, RF_Transactional );
 }
 
 void UMovieSceneMaterialTrack::RemoveAllAnimationData()
@@ -60,10 +60,10 @@ const TArray<UMovieSceneSection*>& UMovieSceneMaterialTrack::GetAllSections() co
 
 void UMovieSceneMaterialTrack::AddScalarParameterKey( FName ParameterName, float Time, float Value )
 {
-	UMovieSceneMaterialParameterSection* NearestSection = Cast<UMovieSceneMaterialParameterSection>(MovieSceneHelpers::FindNearestSectionAtTime(Sections, Time));
+	UMovieSceneParameterSection* NearestSection = Cast<UMovieSceneParameterSection>(MovieSceneHelpers::FindNearestSectionAtTime(Sections, Time));
 	if ( NearestSection == nullptr )
 	{
-		NearestSection = Cast<UMovieSceneMaterialParameterSection>(CreateNewSection());
+		NearestSection = Cast<UMovieSceneParameterSection>(CreateNewSection());
 		NearestSection->SetStartTime( Time );
 		NearestSection->SetEndTime( Time );
 		Sections.Add( NearestSection );
@@ -74,42 +74,29 @@ void UMovieSceneMaterialTrack::AddScalarParameterKey( FName ParameterName, float
 	}
 }
 
-void UMovieSceneMaterialTrack::AddVectorParameterKey( FName ParameterName, float Time, FLinearColor Value )
+void UMovieSceneMaterialTrack::AddColorParameterKey( FName ParameterName, float Time, FLinearColor Value )
 {
-	UMovieSceneMaterialParameterSection* NearestSection = Cast<UMovieSceneMaterialParameterSection>( MovieSceneHelpers::FindNearestSectionAtTime( Sections, Time ) );
+	UMovieSceneParameterSection* NearestSection = Cast<UMovieSceneParameterSection>( MovieSceneHelpers::FindNearestSectionAtTime( Sections, Time ) );
 	if ( NearestSection == nullptr )
 	{
-		NearestSection = Cast<UMovieSceneMaterialParameterSection>( CreateNewSection() );
+		NearestSection = Cast<UMovieSceneParameterSection>( CreateNewSection() );
 		NearestSection->SetStartTime( Time );
 		NearestSection->SetEndTime( Time );
 		Sections.Add( NearestSection );
 	}
 	if (NearestSection->TryModify())
 	{
-		NearestSection->AddVectorParameterKey( ParameterName, Time, Value );
+		NearestSection->AddColorParameterKey( ParameterName, Time, Value );
 	}
 }
 
-void UMovieSceneMaterialTrack::Eval( float Position, TArray<FScalarParameterNameAndValue>& OutScalarValues, TArray<FVectorParameterNameAndValue>& OutVectorValues ) const
+void UMovieSceneMaterialTrack::Eval( float Position, TArray<FScalarParameterNameAndValue>& OutScalarValues, TArray<FColorParameterNameAndValue>& OutColorValues ) const
 {
+	TArray<FVectorParameterNameAndValue> OutVectorValues;
 	for ( UMovieSceneSection* Section : Sections )
 	{
-		UMovieSceneMaterialParameterSection* MaterialParameterSection = Cast<UMovieSceneMaterialParameterSection>( Section );
-		if ( MaterialParameterSection != nullptr )
-		{
-			for ( const FScalarParameterNameAndCurve& ScalarParameterNameAndCurve : *MaterialParameterSection->GetScalarParameterNamesAndCurves() )
-			{
-				OutScalarValues.Add( FScalarParameterNameAndValue( ScalarParameterNameAndCurve.ParameterName, ScalarParameterNameAndCurve.ParameterCurve.Eval( Position ) ) );
-			}
-			for ( const FVectorParameterNameAndCurves& VectorParameterNameAndCurves : *MaterialParameterSection->GetVectorParameterNamesAndCurves() )
-			{
-				OutVectorValues.Add( FVectorParameterNameAndValue( VectorParameterNameAndCurves.ParameterName, FLinearColor(
-					VectorParameterNameAndCurves.RedCurve.Eval( Position ),
-					VectorParameterNameAndCurves.GreenCurve.Eval( Position ),
-					VectorParameterNameAndCurves.BlueCurve.Eval( Position ),
-					VectorParameterNameAndCurves.AlphaCurve.Eval( Position ) ) ) );
-			}
-		}
+		UMovieSceneParameterSection* ParameterSection = Cast<UMovieSceneParameterSection>( Section );
+		ParameterSection->Eval(Position, OutScalarValues, OutVectorValues, OutColorValues);
 	}
 }
 
diff --git a/Engine/Source/Runtime/MovieSceneTracks/Private/Tracks/MovieSceneParticleParameterTrack.cpp b/Engine/Source/Runtime/MovieSceneTracks/Private/Tracks/MovieSceneParticleParameterTrack.cpp
new file mode 100644
--- /dev/null
+++ b/Engine/Source/Runtime/MovieSceneTracks/Private/Tracks/MovieSceneParticleParameterTrack.cpp
@@ -0,0 +1,128 @@
+// Copyright 1998-2015 Epic Games, Inc. All Rights Reserved.
+
+#include "MovieSceneTracksPrivatePCH.h"
+#include "MovieSceneParameterSection.h"
+#include "MovieSceneParticleParameterTrack.h"
+#include "IMovieScenePlayer.h"
+#include "MovieSceneParticleParameterTrackInstance.h"
+
+#define LOCTEXT_NAMESPACE "ParticleParameterTrack"
+
+UMovieSceneParticleParameterTrack::UMovieSceneParticleParameterTrack( const FObjectInitializer& ObjectInitializer )
+	: Super( ObjectInitializer )
+{
+}
+
+TSharedPtr<IMovieSceneTrackInstance> UMovieSceneParticleParameterTrack::CreateInstance()
+{
+	return MakeShareable( new FMovieSceneParticleParameterTrackInstance( *this ) );
+}
+
+UMovieSceneSection* UMovieSceneParticleParameterTrack::CreateNewSection()
+{
+	return NewObject<UMovieSceneParameterSection>( this, UMovieSceneParameterSection::StaticClass(), NAME_None, RF_Transactional );
+}
+
+void UMovieSceneParticleParameterTrack::RemoveAllAnimationData()
+{
+	Sections.Empty();
+}
+
+bool UMovieSceneParticleParameterTrack::HasSection(const UMovieSceneSection& Section) const
+{
+	return Sections.Contains(&Section);
+}
+
+void UMovieSceneParticleParameterTrack::AddSection(UMovieSceneSection& Section)
+{
+	Sections.Add(&Section);
+}
+
+void UMovieSceneParticleParameterTrack::RemoveSection(UMovieSceneSection& Section)
+{
+	Sections.Remove(&Section);
+}
+
+bool UMovieSceneParticleParameterTrack::IsEmpty() const
+{
+	return Sections.Num() == 0;
+}
+
+TRange<float> UMovieSceneParticleParameterTrack::GetSectionBoundaries() const
+{
+	TArray< TRange<float> > Bounds;
+
+	for ( int32 SectionIndex = 0; SectionIndex < Sections.Num(); ++SectionIndex )
+	{
+		Bounds.Add( Sections[SectionIndex]->GetRange() );
+	}
+
+	return TRange<float>::Hull( Bounds );
+}
+
+const TArray<UMovieSceneSection*>& UMovieSceneParticleParameterTrack::GetAllSections() const
+{
+	return Sections;
+}
+
+
+#if WITH_EDITORONLY_DATA
+FText UMovieSceneParticleParameterTrack::GetDisplayName() const
+{
+	return LOCTEXT("DisplayName", "Particle Parameter");
+}
+#endif
+
+
+void UMovieSceneParticleParameterTrack::AddScalarParameterKey( FName ParameterName, float Time, float Value )
+{
+	UMovieSceneParameterSection* NearestSection = Cast<UMovieSceneParameterSection>(MovieSceneHelpers::FindNearestSectionAtTime(Sections, Time));
+	if ( NearestSection == nullptr )
+	{
+		NearestSection = Cast<UMovieSceneParameterSection>(CreateNewSection());
+		NearestSection->SetStartTime( Time );
+		NearestSection->SetEndTime( Time );
+		Sections.Add( NearestSection );
+	}
+	NearestSection->AddScalarParameterKey(ParameterName, Time, Value);
+}
+
+void UMovieSceneParticleParameterTrack::AddVectorParameterKey( FName ParameterName, float Time, FVector Value )
+{
+	UMovieSceneParameterSection* NearestSection = Cast<UMovieSceneParameterSection>( MovieSceneHelpers::FindNearestSectionAtTime( Sections, Time ) );
+	if ( NearestSection == nullptr )
+	{
+		NearestSection = Cast<UMovieSceneParameterSection>( CreateNewSection() );
+		NearestSection->SetStartTime( Time );
+		NearestSection->SetEndTime( Time );
+		Sections.Add( NearestSection );
+	}
+	NearestSection->AddVectorParameterKey( ParameterName, Time, Value );
+}
+
+void UMovieSceneParticleParameterTrack::AddColorParameterKey( FName ParameterName, float Time, FLinearColor Value )
+{
+	UMovieSceneParameterSection* NearestSection = Cast<UMovieSceneParameterSection>( MovieSceneHelpers::FindNearestSectionAtTime( Sections, Time ) );
+	if ( NearestSection == nullptr )
+	{
+		NearestSection = Cast<UMovieSceneParameterSection>( CreateNewSection() );
+		NearestSection->SetStartTime( Time );
+		NearestSection->SetEndTime( Time );
+		Sections.Add( NearestSection );
+	}
+	NearestSection->AddColorParameterKey( ParameterName, Time, Value );
+}
+
+void UMovieSceneParticleParameterTrack::Eval( float Position, TArray<FScalarParameterNameAndValue>& OutScalarValues, TArray<FVectorParameterNameAndValue>& OutVectorValues, TArray<FColorParameterNameAndValue>& OutColorValues ) const
+{
+	for ( UMovieSceneSection* Section : Sections )
+	{
+		UMovieSceneParameterSection* ParameterSection = Cast<UMovieSceneParameterSection>( Section );
+		if ( ParameterSection != nullptr )
+		{
+			ParameterSection->Eval( Position, OutScalarValues, OutVectorValues, OutColorValues );
+		}
+	}
+}
+
+#undef LOCTEXT_NAMESPACE
diff --git a/Engine/Source/Runtime/MovieSceneTracks/Public/Sections/MovieSceneMaterialParameterSection.h b/Engine/Source/Runtime/MovieSceneTracks/Public/Sections/MovieSceneMaterialParameterSection.h
deleted file mode 100644
diff --git a/Engine/Source/Runtime/MovieSceneTracks/Public/Sections/MovieSceneParameterSection.h b/Engine/Source/Runtime/MovieSceneTracks/Public/Sections/MovieSceneParameterSection.h
new file mode 100644
--- /dev/null
+++ b/Engine/Source/Runtime/MovieSceneTracks/Public/Sections/MovieSceneParameterSection.h
@@ -0,0 +1,264 @@
+// Copyright 1998-2015 Epic Games, Inc. All Rights Reserved.
+
+#pragma once
+
+#include "MovieSceneSection.h"
+#include "MovieSceneParameterSection.generated.h"
+
+
+/**
+* Structure representing the animated value of a scalar parameter.
+*/
+struct FScalarParameterNameAndValue
+{
+	/** Creates a new FScalarParameterAndValue with a parameter name and a value. */
+	FScalarParameterNameAndValue( FName InParameterName, float InValue )
+	{
+		ParameterName = InParameterName;
+		Value = InValue;
+	}
+
+	/** The name of the scalar parameter. */
+	FName ParameterName;
+
+	/** The animated value of the scalar parameter. */
+	float Value;
+};
+
+
+/**
+* Structure representing the animated value of a vector parameter.
+*/
+struct FVectorParameterNameAndValue
+{
+	/** Creates a new FVectorParameterAndValue with a parameter name and a value. */
+	FVectorParameterNameAndValue( FName InParameterName, FVector InValue )
+	{
+		ParameterName = InParameterName;
+		Value = InValue;
+	}
+
+	/** The name of the vector parameter. */
+	FName ParameterName;
+
+	/** The animated value of the vector parameter. */
+	FVector Value;
+};
+
+
+/**
+* Structure representing the animated value of a color parameter.
+*/
+struct FColorParameterNameAndValue
+{
+	/** Creates a new FColorParameterAndValue with a parameter name and a value. */
+	FColorParameterNameAndValue( FName InParameterName, FLinearColor InValue )
+	{
+		ParameterName = InParameterName;
+		Value = InValue;
+	}
+
+	/** The name of the color parameter. */
+	FName ParameterName;
+
+	/** The animated value of the color parameter. */
+	FLinearColor Value;
+};
+
+
+/**
+ * Structure representing an animated scalar parameter and it's associated animation curve.
+ */
+USTRUCT()
+struct FScalarParameterNameAndCurve
+{
+	GENERATED_USTRUCT_BODY()
+
+	FScalarParameterNameAndCurve()
+	{
+		Index = 0;
+	}
+
+	/** Creates a new FScalarParameterNameAndCurve for a specific scalar parameter. */
+	FScalarParameterNameAndCurve( FName InParameterName );
+
+	/** The name of the scalar parameter which is being animated. */
+	UPROPERTY()
+	FName ParameterName;
+
+	UPROPERTY()
+	int32 Index;
+
+	/** The curve which contains the animation data for the scalar parameter. */
+	UPROPERTY()
+	FRichCurve ParameterCurve;
+};
+
+
+/**
+* Structure representing an animated vector parameter and it's associated animation curve.
+*/
+USTRUCT()
+struct FVectorParameterNameAndCurves
+{
+	GENERATED_USTRUCT_BODY()
+
+	FVectorParameterNameAndCurves() 
+	{
+		Index = 0;
+	}
+
+	/** Creates a new FVectorParameterNameAndCurve for a specific vector parameter. */
+	FVectorParameterNameAndCurves( FName InParameterName );
+
+	/** The name of the vector parameter which is being animated. */
+	UPROPERTY()
+	FName ParameterName;
+
+	UPROPERTY()
+	int32 Index;
+
+	/** The curve which contains the animation data for the x component of the vector parameter. */
+	UPROPERTY()
+	FRichCurve XCurve;
+
+	/** The curve which contains the animation data for the y component of the vector parameter. */
+	UPROPERTY()
+	FRichCurve YCurve;
+
+	/** The curve which contains the animation data for the z component of the vector parameter. */
+	UPROPERTY()
+	FRichCurve ZCurve;
+};
+
+
+/**
+* Structure representing an animated vector parameter and it's associated animation curve.
+*/
+USTRUCT()
+struct FColorParameterNameAndCurves
+{
+	GENERATED_USTRUCT_BODY()
+
+	FColorParameterNameAndCurves()
+	{
+		Index = 0;
+	}
+
+	/** Creates a new FVectorParameterNameAndCurve for a specific color parameter. */
+	FColorParameterNameAndCurves( FName InParameterName );
+
+	/** The name of the vector parameter which is being animated. */
+	UPROPERTY()
+	FName ParameterName;
+
+	UPROPERTY()
+	int32 Index;
+
+	/** The curve which contains the animation data for the red component of the color parameter. */
+	UPROPERTY()
+	FRichCurve RedCurve;
+
+	/** The curve which contains the animation data for the green component of the color parameter. */
+	UPROPERTY()
+	FRichCurve GreenCurve;
+
+	/** The curve which contains the animation data for the blue component of the color parameter. */
+	UPROPERTY()
+	FRichCurve BlueCurve;
+
+	/** The curve which contains the animation data for the alpha component of the color parameter. */
+	UPROPERTY()
+	FRichCurve AlphaCurve;
+};
+
+
+/**
+ * A single movie scene section which can contain data for multiple named parameters.
+ */
+UCLASS( MinimalAPI )
+class UMovieSceneParameterSection : public UMovieSceneSection
+{
+	GENERATED_UCLASS_BODY()
+
+public:
+	/** Adds a a key for a specific scalar parameter at the specified time with the specified value. */
+	void AddScalarParameterKey(FName InParameterName, float InTime, float InValue);
+
+	/** Adds a a key for a specific vector parameter at the specified time with the specified value. */
+	void AddVectorParameterKey( FName InParameterName, float InTime, FVector InValue );
+
+	/** Adds a a key for a specific color parameter at the specified time with the specified value. */
+	void AddColorParameterKey( FName InParameterName, float InTime, FLinearColor InValue );
+
+	/** 
+	 * Removes a scalar parameter from this section. 
+	 * 
+	 * @param InParameterName The name of the scalar parameter to remove.
+	 * @returns True if a parameter with that name was found and removed, otherwise false.
+	 */
+	MOVIESCENETRACKS_API bool RemoveScalarParameter( FName InParameterName );
+
+	/**
+	* Removes a vector parameter from this section.
+	*
+	* @param InParameterName The name of the vector parameter to remove.
+	* @returns True if a parameter with that name was found and removed, otherwise false.
+	*/
+	MOVIESCENETRACKS_API bool RemoveVectorParameter( FName InParameterName );
+
+	/**
+	* Removes a color parameter from this section.
+	*
+	* @param InParameterName The name of the color parameter to remove.
+	* @returns True if a parameter with that name was found and removed, otherwise false.
+	*/
+	MOVIESCENETRACKS_API bool RemoveColorParameter( FName InParameterName );
+
+	/** Gets the animated scalar parameters and their associated curves. */
+	MOVIESCENETRACKS_API TArray<FScalarParameterNameAndCurve>* GetScalarParameterNamesAndCurves();
+
+	/** Gets the animated color parameters and their associated curves. */
+	MOVIESCENETRACKS_API TArray<FVectorParameterNameAndCurves>* GetVectorParameterNamesAndCurves();
+
+	/** Gets the animated color parameters and their associated curves. */
+	MOVIESCENETRACKS_API TArray<FColorParameterNameAndCurves>* GetColorParameterNamesAndCurves();
+
+	/** Gets the set of all parameter names used by this section. */
+	MOVIESCENETRACKS_API void GetParameterNames(TSet<FName>& ParameterNames) const;
+
+	/**
+	* Gets the animated values for this track.
+	* @param Position The playback position to use for evaluation.
+	* @param OutScalarValues An array of FScalarParameterNameAndValue objects representing each animated scalar parameter and it's animated value.
+	* @param OutVectorValues An array of FVectorParameterNameAndValue objects representing each animated vector parameter and it's animated value.
+	* @param OutColorValues An array of FColorParameterNameAndValue objects representing each animated color parameter and it's animated value.
+	*/
+	void Eval( float Position, 
+		TArray<FScalarParameterNameAndValue>& OutScalarValues,
+		TArray<FVectorParameterNameAndValue>& OutVectorValues,
+		TArray<FColorParameterNameAndValue>& OutColorValues ) const;
+
+
+private:
+	void UpdateParameterIndicesFromRemoval(int32 RemovedIndex);
+
+private:
+	/**
+	 * The scalar parameter names and their associated curves.
+	 */
+	UPROPERTY()
+	TArray<FScalarParameterNameAndCurve> ScalarParameterNamesAndCurves;
+
+	/**
+	* The vector parameter names and their associated curves.
+	*/
+	UPROPERTY()
+	TArray<FVectorParameterNameAndCurves> VectorParameterNamesAndCurves;
+
+	/**
+	* The vector parameter names and their associated curves.
+	*/
+	UPROPERTY()
+	TArray<FColorParameterNameAndCurves> ColorParameterNamesAndCurves;
+};
\ No newline at end of file
diff --git a/Engine/Source/Runtime/MovieSceneTracks/Public/TrackInstances/MovieSceneParticleParameterTrackInstance.h b/Engine/Source/Runtime/MovieSceneTracks/Public/TrackInstances/MovieSceneParticleParameterTrackInstance.h
new file mode 100644
--- /dev/null
+++ b/Engine/Source/Runtime/MovieSceneTracks/Public/TrackInstances/MovieSceneParticleParameterTrackInstance.h
@@ -0,0 +1,46 @@
+// Copyright 1998-2015 Epic Games, Inc. All Rights Reserved.
+
+#pragma once
+
+#include "IMovieSceneTrackInstance.h"
+#include "ObjectKey.h"
+
+class FTrackInstancePropertyBindings;
+class UMovieSceneParticleParameterTrack;
+
+/**
+ * A movie scene track instance for material tracks.
+ */
+class FMovieSceneParticleParameterTrackInstance
+	: public IMovieSceneTrackInstance
+{
+private:
+
+	struct FInitialParameterValues
+	{
+		TMap<FName, float> ScalarNameToValueMap;
+		TMap<FName, FVector> VectorNameToValueMap;
+		TMap<FName, FLinearColor> ColorNameToValueMap;
+	};
+
+public:
+	FMovieSceneParticleParameterTrackInstance( UMovieSceneParticleParameterTrack& InParticleParameterTrack );
+
+	/** IMovieSceneTrackInstance interface */
+	virtual void SaveState (const TArray<UObject*>& RuntimeObjects, IMovieScenePlayer& Player, FMovieSceneSequenceInstance& SequenceInstance) override;
+	virtual void RestoreState (const TArray<UObject*>& RuntimeObjects, IMovieScenePlayer& Player, FMovieSceneSequenceInstance& SequenceInstance) override;
+	virtual void Update( float Position, float LastPosition, const TArray<UObject*>& RuntimeObjects, IMovieScenePlayer& Player, FMovieSceneSequenceInstance& SequenceInstance, EMovieSceneUpdatePass UpdatePass ) override;
+	virtual void RefreshInstance( const TArray<UObject*>& RuntimeObjects, IMovieScenePlayer& Player, FMovieSceneSequenceInstance& SequenceInstance ) override;
+	virtual void ClearInstance( IMovieScenePlayer& Player, FMovieSceneSequenceInstance& SequenceInstance ) override;
+
+private:
+
+	/** Track that is being instanced */
+	UMovieSceneParticleParameterTrack* ParticleParameterTrack;
+
+	/** The initial values of instance parameters for the objects controlled by this track. */
+	TMap<FObjectKey, TSharedPtr<FInitialParameterValues>> ObjectToInitialParameterValuesMap;
+
+	/** The ParticleSystemComponents to be animated on update. */
+	TArray<UParticleSystemComponent*> ParticleSystemComponents;
+};
diff --git a/Engine/Source/Runtime/MovieSceneTracks/Public/TrackInstances/MovieSceneShotTrackInstance.h b/Engine/Source/Runtime/MovieSceneTracks/Public/TrackInstances/MovieSceneShotTrackInstance.h
--- a/Engine/Source/Runtime/MovieSceneTracks/Public/TrackInstances/MovieSceneShotTrackInstance.h
+++ b/Engine/Source/Runtime/MovieSceneTracks/Public/TrackInstances/MovieSceneShotTrackInstance.h
@@ -23,7 +23,7 @@ public:
 
 	virtual void ClearInstance(IMovieScenePlayer& Player, FMovieSceneSequenceInstance& SequenceInstance) override;
 	virtual void RefreshInstance(const TArray<UObject*>& RuntimeObjects, IMovieScenePlayer& Player, FMovieSceneSequenceInstance& SequenceInstance) override;
-	virtual void RestoreState(const TArray<UObject*>& RuntimeObjects, IMovieScenePlayer& Player, FMovieSceneSequenceInstance& SequenceInstance) override { }
+	virtual void RestoreState(const TArray<UObject*>& RuntimeObjects, IMovieScenePlayer& Player, FMovieSceneSequenceInstance& SequenceInstance) override;
 	virtual void SaveState(const TArray<UObject*>& RuntimeObjects, IMovieScenePlayer& Player, FMovieSceneSequenceInstance& SequenceInstance) override { }
 	virtual void Update(float Position, float LastPosition, const TArray<UObject*>& RuntimeObjects, IMovieScenePlayer& Player, FMovieSceneSequenceInstance& SequenceInstance, EMovieSceneUpdatePass UpdatePass) override;
 
diff --git a/Engine/Source/Runtime/MovieSceneTracks/Public/Tracks/MovieSceneMaterialTrack.h b/Engine/Source/Runtime/MovieSceneTracks/Public/Tracks/MovieSceneMaterialTrack.h
--- a/Engine/Source/Runtime/MovieSceneTracks/Public/Tracks/MovieSceneMaterialTrack.h
+++ b/Engine/Source/Runtime/MovieSceneTracks/Public/Tracks/MovieSceneMaterialTrack.h
@@ -3,49 +3,10 @@
 #pragma once
 
 #include "MovieSceneTrack.h"
+#include "MovieSceneParameterSection.h"
 #include "MovieSceneMaterialTrack.generated.h"
 
 
-/**
- * Structure representing the animated value of a scalar parameter.
- */
-struct FScalarParameterNameAndValue
-{
-	/** Creates a new FScalarParameterAndValue with a parameter name and a value. */
-	FScalarParameterNameAndValue( FName InParameterName, float InValue )
-	{
-		ParameterName = InParameterName;
-		Value = InValue;
-	}
-
-	/** The name of the scalar parameter. */
-	FName ParameterName;
-
-	/** The animated value of the scalar parameter. */
-	float Value;
-};
-
-
-/**
-* Structure representing the animated value of a vector parameter.
-*/
-struct FVectorParameterNameAndValue
-{
-	/** Creates a new FVectorParameterAndValue with a parameter name and a value. */
-	FVectorParameterNameAndValue( FName InParameterName, FLinearColor InValue )
-	{
-		ParameterName = InParameterName;
-		Value = InValue;
-	}
-
-	/** The name of the vector parameter. */
-	FName ParameterName;
-
-	/** The animated value of the vector parameter. */
-	FLinearColor Value;
-};
-
-
 /**
  * Handles manipulation of material parameters in a movie scene.
  */
@@ -79,12 +40,12 @@ public:
 	void MOVIESCENETRACKS_API AddScalarParameterKey( FName ParameterName, float Position, float Value );
 
 	/**
-	* Adds a Vector parameter key to the track.
+	* Adds a color parameter key to the track.
 	* @param ParameterName The name of the parameter to add a key for.
 	* @param Time The time to add the new key.
 	* @param The value for the new key.
 	*/
-	void MOVIESCENETRACKS_API AddVectorParameterKey( FName ParameterName, float Position, FLinearColor Value );
+	void MOVIESCENETRACKS_API AddColorParameterKey( FName ParameterName, float Position, FLinearColor Value );
 
 	/**
 	 * Gets the animated values for this track.
@@ -92,7 +53,7 @@ public:
 	 * @param OutScalarValues An array of FScalarParameterNameAndValue objects representing each animated scalar parameter and it's animated value.
 	 * @param OutVectorValues An array of FVectorParameterNameAndValue objects representing each animated vector parameter and it's animated value.
 	 */
-	void Eval( float Position, TArray<FScalarParameterNameAndValue>& OutScalarValues, TArray<FVectorParameterNameAndValue>& OutVectorValues ) const;
+	void Eval( float Position, TArray<FScalarParameterNameAndValue>& OutScalarValues, TArray<FColorParameterNameAndValue>& OutVectorValues ) const;
 
 private:
 
diff --git a/Engine/Source/Runtime/MovieSceneTracks/Public/Tracks/MovieSceneParticleParameterTrack.h b/Engine/Source/Runtime/MovieSceneTracks/Public/Tracks/MovieSceneParticleParameterTrack.h
new file mode 100644
--- /dev/null
+++ b/Engine/Source/Runtime/MovieSceneTracks/Public/Tracks/MovieSceneParticleParameterTrack.h
@@ -0,0 +1,78 @@
+// Copyright 1998-2015 Epic Games, Inc. All Rights Reserved.
+
+#pragma once
+
+#include "MovieSceneTrack.h"
+#include "MovieSceneParameterSection.h"
+#include "MovieSceneParticleParameterTrack.generated.h"
+
+
+/**
+ * Handles manipulation of material parameters in a movie scene.
+ */
+UCLASS( MinimalAPI )
+class UMovieSceneParticleParameterTrack
+	: public UMovieSceneTrack
+{
+	GENERATED_UCLASS_BODY()
+
+public:
+
+	// UMovieSceneTrack interface
+
+	virtual TSharedPtr<IMovieSceneTrackInstance> CreateInstance() override;
+	virtual UMovieSceneSection* CreateNewSection() override;
+	virtual void RemoveAllAnimationData() override;
+	virtual bool HasSection(const UMovieSceneSection& Section) const override;
+	virtual void AddSection(UMovieSceneSection& Section) override;
+	virtual void RemoveSection(UMovieSceneSection& Section) override;
+	virtual bool IsEmpty() const override;
+	virtual TRange<float> GetSectionBoundaries() const override;
+	virtual const TArray<UMovieSceneSection*>& GetAllSections() const override;
+
+#if WITH_EDITORONLY_DATA
+	virtual FText GetDisplayName() const override;
+#endif
+
+
+public:
+
+	/**
+	 * Adds a scalar parameter key to the track. 
+	 * @param ParameterName The name of the parameter to add a key for.
+	 * @param Time The time to add the new key.
+	 * @param The value for the new key.
+	 */
+	void MOVIESCENETRACKS_API AddScalarParameterKey( FName ParameterName, float Position, float Value );
+
+	/**
+	* Adds a Vector parameter key to the track.
+	* @param ParameterName The name of the parameter to add a key for.
+	* @param Time The time to add the new key.
+	* @param The value for the new key.
+	*/
+	void MOVIESCENETRACKS_API AddVectorParameterKey( FName ParameterName, float Position, FVector Value );
+
+	/**
+	* Adds a Vector parameter key to the track.
+	* @param ParameterName The name of the parameter to add a key for.
+	* @param Time The time to add the new key.
+	* @param The value for the new key.
+	*/
+	void MOVIESCENETRACKS_API AddColorParameterKey( FName ParameterName, float Position, FLinearColor Value );
+
+	/**
+	 * Gets the animated values for this track.
+	 * @param Position The playback position to use for evaluation.
+	 * @param OutScalarValues An array of FScalarParameterNameAndValue objects representing each animated scalar parameter and it's animated value.
+	 * @param OutVectorValues An array of FVectorParameterNameAndValue objects representing each animated vector parameter and it's animated value.
+	 * @param OutVectorValues An array of FVectorParameterNameAndValue objects representing each animated color parameter and it's animated value.
+	 */
+	void Eval( float Position, TArray<FScalarParameterNameAndValue>& OutScalarValues, TArray<FVectorParameterNameAndValue>& OutVectorValues, TArray<FColorParameterNameAndValue>& OutColorValues ) const;
+
+private:
+
+	/** The sections owned by this track .*/
+	UPROPERTY()
+	TArray<UMovieSceneSection*> Sections;
+};
diff --git a/Engine/Source/Runtime/Online/Android/OnlineSubsystemGooglePlay/Private/OnlineExternalUIInterfaceGooglePlay.h b/Engine/Source/Runtime/Online/Android/OnlineSubsystemGooglePlay/Private/OnlineExternalUIInterfaceGooglePlay.h
--- a/Engine/Source/Runtime/Online/Android/OnlineSubsystemGooglePlay/Private/OnlineExternalUIInterfaceGooglePlay.h
+++ b/Engine/Source/Runtime/Online/Android/OnlineSubsystemGooglePlay/Private/OnlineExternalUIInterfaceGooglePlay.h
@@ -21,7 +21,7 @@ private:
 public:
 	FOnlineExternalUIGooglePlay(FOnlineSubsystemGooglePlay* InSubsystem);
 
-	virtual bool ShowLoginUI(const int ControllerIndex, bool bShowOnlineOnly, const FOnLoginUIClosedDelegate& Delegate = FOnLoginUIClosedDelegate() = FOnLoginUIClosedDelegate()) override;
+	virtual bool ShowLoginUI(const int ControllerIndex, bool bShowOnlineOnly, const FOnLoginUIClosedDelegate& Delegate = FOnLoginUIClosedDelegate()) override;
 	virtual bool ShowFriendsUI(int32 LocalUserNum) override;
 	virtual bool ShowInviteUI(int32 LocalUserNum, FName SessionMame) override;
 	virtual bool ShowAchievementsUI(int32 LocalUserNum) override;
diff --git a/Engine/Source/Runtime/Online/IOS/OnlineSubsystemIOS/Private/OnlineExternalUIInterfaceIOS.h b/Engine/Source/Runtime/Online/IOS/OnlineSubsystemIOS/Private/OnlineExternalUIInterfaceIOS.h
--- a/Engine/Source/Runtime/Online/IOS/OnlineSubsystemIOS/Private/OnlineExternalUIInterfaceIOS.h
+++ b/Engine/Source/Runtime/Online/IOS/OnlineSubsystemIOS/Private/OnlineExternalUIInterfaceIOS.h
@@ -16,7 +16,7 @@ PACKAGE_SCOPE:
 public:
 
 	//~ Begin IOnlineExternalUI Interface
-	virtual bool ShowLoginUI(const int ControllerIndex, bool bShowOnlineOnly, const FOnLoginUIClosedDelegate& Delegate = FOnLoginUIClosedDelegate() = FOnLoginUIClosedDelegate()) override;
+	virtual bool ShowLoginUI(const int ControllerIndex, bool bShowOnlineOnly, const FOnLoginUIClosedDelegate& Delegate = FOnLoginUIClosedDelegate()) override;
 	virtual bool ShowFriendsUI(int32 LocalUserNum) override;
 	virtual bool ShowInviteUI(int32 LocalUserNum, FName SessionMame = GameSessionName) override;
 	virtual bool ShowAchievementsUI(int32 LocalUserNum) override;
diff --git a/Engine/Source/Runtime/Online/OnlineSubsystem/Public/Interfaces/OnlineExternalUIInterface.h b/Engine/Source/Runtime/Online/OnlineSubsystem/Public/Interfaces/OnlineExternalUIInterface.h
--- a/Engine/Source/Runtime/Online/OnlineSubsystem/Public/Interfaces/OnlineExternalUIInterface.h
+++ b/Engine/Source/Runtime/Online/OnlineSubsystem/Public/Interfaces/OnlineExternalUIInterface.h
@@ -47,6 +47,8 @@ struct FShowWebUrlParams
 	int32 SizeY;
 	/** if specified then restricted to only navigate within these domains */
 	TArray<FString> AllowedDomains;
+	/** portion of url for detecting callback.  Eg. "&code=", "redirect=", etc */
+	FString CallbackPath;
 
 	/**
 	 * Constructor
@@ -94,7 +96,7 @@ public:
  	 *
 	 * @return true if it was able to show the UI, false if it failed
 	 */
-	virtual bool ShowLoginUI(const int ControllerIndex, bool bShowOnlineOnly, const FOnLoginUIClosedDelegate& Delegate = FOnLoginUIClosedDelegate() = FOnLoginUIClosedDelegate()) = 0;
+	virtual bool ShowLoginUI(const int ControllerIndex, bool bShowOnlineOnly, const FOnLoginUIClosedDelegate& Delegate = FOnLoginUIClosedDelegate()) = 0;
 
 	/**
 	 * Displays the UI that shows a user's list of friends
diff --git a/Engine/Source/Runtime/Online/OnlineSubsystemUtils/Private/PartyBeaconState.cpp b/Engine/Source/Runtime/Online/OnlineSubsystemUtils/Private/PartyBeaconState.cpp
--- a/Engine/Source/Runtime/Online/OnlineSubsystemUtils/Private/PartyBeaconState.cpp
+++ b/Engine/Source/Runtime/Online/OnlineSubsystemUtils/Private/PartyBeaconState.cpp
@@ -81,7 +81,7 @@ bool UPartyBeaconState::ReconfigureTeamAndPlayerCount(int32 InNumTeams, int32 In
 	bool bSuccess = false;
 
 	//Check total existing reservations against new total maximum
-	if (NumConsumedReservations < InNumReservations)
+	if (NumConsumedReservations <= InNumReservations)
 	{
 		bool bTeamError = false;
 		// Check teams with reservations against new team count
diff --git a/Engine/Source/Runtime/SlateNullRenderer/Private/SlateNullRendererModule.cpp b/Engine/Source/Runtime/SlateNullRenderer/Private/SlateNullRendererModule.cpp
--- a/Engine/Source/Runtime/SlateNullRenderer/Private/SlateNullRendererModule.cpp
+++ b/Engine/Source/Runtime/SlateNullRenderer/Private/SlateNullRendererModule.cpp
@@ -2,6 +2,21 @@
 
 #include "SlateNullRendererPrivatePCH.h"
 
+class FSlateNullShaderResourceManager : public ISlateAtlasProvider, public FSlateShaderResourceManager
+{
+public:
+	// ISlateAtlasProvider interface
+	virtual int32 GetNumAtlasPages() const override { return 0; }
+	virtual FIntPoint GetAtlasPageSize() const override { return FIntPoint(0, 0); }
+	virtual class FSlateShaderResource* GetAtlasPageResource(const int32 InIndex) const override { return nullptr; }
+	virtual bool IsAtlasPageResourceAlphaOnly() const override { return false; }
+
+	// FSlateShaderResourceManager interface
+	virtual FSlateShaderResourceProxy* GetShaderResource( const FSlateBrush& InBrush ) override { return nullptr; }
+	virtual FSlateResourceHandle GetResourceHandle( const FSlateBrush& InBrush ) override { return FSlateResourceHandle(); }
+	virtual ISlateAtlasProvider* GetTextureAtlasProvider() { return this; }
+};
+
 /** A null font texture resource to represent fonts */
 class FSlateFontTextureNull : public FSlateShaderResource
 #if WITH_ENGINE
@@ -109,10 +124,17 @@ private:
 
 			SlateFontServices = MakeShareable(new FSlateFontServices(FontCache, FontCache));
 		}
+
+		if (!ResourceManager.IsValid())
+		{
+			ResourceManager = MakeShareable(new FSlateNullShaderResourceManager);
+			FSlateDataPayload::ResourceManager = ResourceManager.Get();
+		}
 	}
 
 private:
 	TSharedPtr<FSlateFontServices> SlateFontServices;
+	TSharedPtr<FSlateShaderResourceManager> ResourceManager;
 };
 
 
diff --git a/Engine/Source/Runtime/UMG/Private/Slate/SlateDataSheet.cpp b/Engine/Source/Runtime/UMG/Private/Slate/SlateDataSheet.cpp
new file mode 100644
--- /dev/null
+++ b/Engine/Source/Runtime/UMG/Private/Slate/SlateDataSheet.cpp
@@ -0,0 +1,62 @@
+// Copyright 1998-2015 Epic Games, Inc. All Rights Reserved.
+
+#include "UMGPrivatePCH.h"
+#include "SlateDataSheet.h"
+
+
+const FUpdateTextureRegion2D USlateDataSheet::DataSheetUpdateRegion(
+	0,
+	0,
+	0,
+	0,
+	USlateDataSheet::Data_Width,
+	USlateDataSheet::Data_Height);
+
+USlateDataSheet::USlateDataSheet(const FObjectInitializer& ObjectInitializer)
+	:Super(ObjectInitializer)
+{
+}
+
+
+void USlateDataSheet::Init()
+{
+	// Ensure that we have a texture resource
+	static const EPixelFormat DataFormat = EPixelFormat::PF_B8G8R8A8;
+	ensure(GPixelFormats[DataFormat].BlockBytes == Data_PixelSize);
+
+	DataTexture = UTexture2D::CreateTransient(Data_Width, Data_Height, DataFormat);
+#if WITH_EDITORONLY_DATA
+	DataTexture->MipGenSettings = TMGS_NoMipmaps;
+#endif
+	DataTexture->SRGB = false;
+	DataTexture->AddressX = TA_Clamp;
+	DataTexture->AddressY = TA_Clamp;
+	DataTexture->Filter = TF_Nearest;
+	DataTexture->CompressionSettings = TC_EditorIcon;
+	DataTexture->NeverStream = true;
+	DataTexture->UpdateResource();
+
+}
+
+void USlateDataSheet::EnqueueUpdateToGPU()
+{
+	if (ensureMsgf(DataTexture != nullptr, TEXT("Make sure to call Init()")))
+	{
+		const uint32 DataPitch = Data_Width * Data_PixelSize;
+		const uint32 DataSize = DataPitch*Data_Height;
+		uint8* SrcData = (uint8*)FMemory::Malloc(DataSize);
+		FMemory::Memcpy(SrcData, Data, DataSize);
+
+		auto DataCleanup = [](uint8* SrcData, const FUpdateTextureRegion2D* Regions)
+		{
+			FMemory::Free(SrcData);
+		};
+		DataTexture->UpdateTextureRegions(0, 1, &DataSheetUpdateRegion, DataPitch, Data_PixelSize, SrcData, DataCleanup);
+	}
+}
+
+
+UTexture2D* USlateDataSheet::GetTexture() const
+{
+	return DataTexture;
+}
\ No newline at end of file
diff --git a/Engine/Source/Runtime/UMG/Private/Slate/SlateVectorArtInstanceData.cpp b/Engine/Source/Runtime/UMG/Private/Slate/SlateVectorArtInstanceData.cpp
--- a/Engine/Source/Runtime/UMG/Private/Slate/SlateVectorArtInstanceData.cpp
+++ b/Engine/Source/Runtime/UMG/Private/Slate/SlateVectorArtInstanceData.cpp
@@ -3,23 +3,6 @@
 #include "UMGPrivatePCH.h"
 #include "SlateVectorArtInstanceData.h"
 
-void FSlateVectorArtInstanceData::SetPositionFixedPoint16(FVector2D Position)
-{
-	// We are effectively doing fixed point positioning
-	// With 13 bits for the whole part and 3 bits of decimal precision.
-	// xxxxxxxxxxxxx.xxx
-	//     [0..8191].[0..7]
-	{
-		uint32 CurrentDataX = static_cast<uint32>(GetData().X) & 0xFFFF0000;
-		GetData().X = CurrentDataX | FMath::RoundToInt(Position.X*8.0f) << 0;
-	}
-
-	{
-		uint32 CurrentDataY = static_cast<uint32>(GetData().Y) & 0xFFFF0000;
-		GetData().Y = CurrentDataY | FMath::RoundToInt(Position.Y*8.0f) << 0;
-	}
-}
-
 void FSlateVectorArtInstanceData::SetPosition(FVector2D Position)
 {
 	GetData().X = Position.X;
@@ -30,3 +13,8 @@ void FSlateVectorArtInstanceData::SetScale(float Scale)
 {
 	GetData().Z = Scale;
 }
+
+void FSlateVectorArtInstanceData::SetBaseAddress(float Address)
+{
+	GetData().W = Address;
+}
diff --git a/Engine/Source/Runtime/UMG/Public/Slate/SlateDataSheet.h b/Engine/Source/Runtime/UMG/Public/Slate/SlateDataSheet.h
new file mode 100644
--- /dev/null
+++ b/Engine/Source/Runtime/UMG/Public/Slate/SlateDataSheet.h
@@ -0,0 +1,65 @@
+// Copyright 1998-2015 Epic Games, Inc. All Rights Reserved.
+#pragma once
+
+
+#include "SlateDataSheet.generated.h"
+
+class UTexture2D;
+
+/**
+ * A texture used for communicating data to the GPU.
+ * Used in combination with SlateVectorArtData and SlateVectorArtInstanceData to
+ * pass data to UI materials.
+ */
+UCLASS()
+class UMG_API USlateDataSheet : public UObject
+{
+	GENERATED_UCLASS_BODY()
+
+public:
+	void Init();
+
+	void SetUnitFloat(uint32 Address, int32 Value);
+	
+	FORCEINLINE_DEBUGGABLE void SetUnitFloat(uint32 Address, float Value)
+	{
+		SetUINT24(Address, FMath::RoundToInt(Value * (256 * 256 * 256 - 1)));
+	}
+
+	FORCEINLINE_DEBUGGABLE void SetUINT24(uint32 Address, int32 Value)
+	{
+		FColor& Texel = (Data[Address]);
+		Texel.R = (uint8)(Value & 0x000000FF);
+		Texel.G = (uint8)((Value & 0x0000FF00) >> 8);
+		Texel.B = (uint8)((Value & 0x00FF0000) >> 16);
+		Texel.A = (uint8)((Value & 0x00000000));
+	}
+	
+	void SetUINT24(uint32 Address, float Value);
+
+	void EnqueueUpdateToGPU();
+
+	UTexture2D* GetTexture() const;
+
+protected:
+
+	static const uint32 Data_Width = 256;
+	static const uint32 Data_Height = 1;
+	static const uint32 Data_PixelSize = sizeof(FColor);
+	static const FUpdateTextureRegion2D DataSheetUpdateRegion;
+
+	UPROPERTY(Transient)
+	UTexture2D* DataTexture;
+
+	FColor Data[Data_Width * Data_Height];
+
+
+	
+
+	
+
+
+
+};
+
+
diff --git a/Engine/Source/Runtime/UMG/Public/Slate/SlateVectorArtInstanceData.h b/Engine/Source/Runtime/UMG/Public/Slate/SlateVectorArtInstanceData.h
--- a/Engine/Source/Runtime/UMG/Public/Slate/SlateVectorArtInstanceData.h
+++ b/Engine/Source/Runtime/UMG/Public/Slate/SlateVectorArtInstanceData.h
@@ -9,58 +9,11 @@ public:
 	const FVector4& GetData() const { return Data; }
 	FVector4& GetData() { return Data; }
 
-	/**
-	 * A Float4 can pack 12 bytes worth of data (we cannot safely use the exponent bits)
-	 */
-	template<int32 Component, int32 ByteIndex>
-	void PackFloatIntoByte(float InValue)
-	{
-		PackByteIntoByte<Component, ByteIndex>(static_cast<uint8>(FMath::RoundToInt(InValue * 0xFFu)));
-	}
-	
-	template<int32 Component, int32 ByteIndex>
-	void PackByteIntoByte(uint8 InValue)
-	{
-		// Each float has 24 usable bits of mantissa, but we cannot access the bits directly.
-		// We will not respect the IEEE "normalized mantissa" rules, so let 
-		// the compiler/FPU do conversions from byte to float and vice versa for us.
-
-		//Produces a value like 0xFFFF00FF; where the 00s are shifted by `Position` bytes.
-		static const uint32 Mask = ~(0xFF << ByteIndex * 8);
-
-		// Clear out a byte to OR with while maintaining the rest of the data intact.
-		uint32 CurrentData = static_cast<uint32>(Data[Component]) & Mask;
-		// OR in the value
-		Data[Component] = CurrentData | (InValue << ByteIndex*8);
-	}
-
-	static float AdjustScaleInput(float UnadjustedScaleInput)
-	{
-		// We only support scaling UP (initial "vector" art is at 1:1 resolution)
-		// We also don't want to waste any bits on supporting fractional scale values.
-		// So store a value between 0 and 3, with an assumed base scale of 1.0f
-		static const float MAXIMUM_SCALE = 4.0f;
-		return FMath::Clamp(UnadjustedScaleInput, 0.0f, MAXIMUM_SCALE) / MAXIMUM_SCALE;
-	}
-
-	template<int32 Component, int32 ByteIndex>
-	void PackByteIntoByte(float InValue)
-	{
-		checkf(false, TEXT("PackByteIntoByte() being used with a float value. This is almost definitely an error. Use PackFloatIntoByte()."));
-	}
-
-	template<int32 Component, int32 ByteIndex>
-	float UnpackValue() const
-	{
-		// See PackFloatIntoByte
-		static const float Denominator = static_cast<float>(0x1 << (ByteIndex+1)*8);
-		return Data[Component] / Denominator;
-	}
-
 	void SetPositionFixedPoint16(FVector2D Position);
 
 	void SetPosition(FVector2D Position);
 	void SetScale(float Scale);
+	void SetBaseAddress(float Address);
 
 protected:
 	FVector4 Data;
