Prevent UnrealBot from recreating workspaces with each run. It doesn't matter what user view is specified when creating stream workspaces because it'll be overriden by the stream.

diff --git a/Engine/Source/Programs/AutomationTool/AutomationUtils/P4Utils.cs b/Engine/Source/Programs/AutomationTool/AutomationUtils/P4Utils.cs
--- a/Engine/Source/Programs/AutomationTool/AutomationUtils/P4Utils.cs
+++ b/Engine/Source/Programs/AutomationTool/AutomationUtils/P4Utils.cs
@@ -90,7 +90,7 @@ namespace AutomationTool
 				&& LineEnd == Other.LineEnd 
 				&& Options == Other.Options 
 				&& SubmitOptions == Other.SubmitOptions
-				&& Enumerable.SequenceEqual(View, Other.View);
+				&& (!String.IsNullOrEmpty(Stream) || Enumerable.SequenceEqual(View, Other.View));
 		}
 
 		public override string ToString()
