Linux: fix project packaging (UE-22018).

- Minor UAT change to strip quotes from a string (File.Exists() doesn't seem to do that at least on mono).
- Case-sensitive compilation fix.

diff --git a/Engine/Source/Programs/AutomationTool/AutomationUtils/Automation.cs b/Engine/Source/Programs/AutomationTool/AutomationUtils/Automation.cs
--- a/Engine/Source/Programs/AutomationTool/AutomationUtils/Automation.cs
+++ b/Engine/Source/Programs/AutomationTool/AutomationUtils/Automation.cs
@@ -190,7 +190,7 @@ AutomationTool.exe [-verbose] [-compileonly] [-p4] Command0 [-Arg0 -Arg1 -Arg2 
 				{
 					throw new AutomationException("The -ProjectScripts argument may only be specified once");
 				}
-				var ProjectFileName = CurrentParam.Substring(CurrentParam.IndexOf('=') + 1);
+				var ProjectFileName = CurrentParam.Substring(CurrentParam.IndexOf('=') + 1).Replace("\"", "");
 				if(!File.Exists(ProjectFileName))
 				{
 					throw new AutomationException("Project '{0}' does not exist", ProjectFileName);
diff --git a/Engine/Source/Programs/UnrealHeaderTool/Private/UHTMakefile/ClassMetaDataArchiveProxy.cpp b/Engine/Source/Programs/UnrealHeaderTool/Private/UHTMakefile/ClassMetaDataArchiveProxy.cpp
--- a/Engine/Source/Programs/UnrealHeaderTool/Private/UHTMakefile/ClassMetaDataArchiveProxy.cpp
+++ b/Engine/Source/Programs/UnrealHeaderTool/Private/UHTMakefile/ClassMetaDataArchiveProxy.cpp
@@ -1,7 +1,7 @@
 // Copyright 1998-2015 Epic Games, Inc. All Rights Reserved.
 #include "UnrealHeaderTool.h"
 #include "UHTMakefile/UHTMakefile.h"
-#include "UHTMakefile/ClassMetadataArchiveProxy.h"
+#include "UHTMakefile/ClassMetaDataArchiveProxy.h"
 
 FClassMetaDataArchiveProxy::FClassMetaDataArchiveProxy(const FUHTMakefile& UHTMakefile, const FClassMetaData* ClassMetaData)
 {
