[INTEGRATE] Change 2777449 by [REDACTED_PERSON][REDACTED_HANDLE] on 2015/11/23 07:24

       Update threading model in S3 automation scripts to fix slow Amazon uploads.

diff --git a/Engine/Source/Programs/AutomationTool/AutomationUtils/MCPPublic.cs b/Engine/Source/Programs/AutomationTool/AutomationUtils/MCPPublic.cs
--- a/Engine/Source/Programs/AutomationTool/AutomationUtils/MCPPublic.cs
+++ b/Engine/Source/Programs/AutomationTool/AutomationUtils/MCPPublic.cs
@@ -750,8 +750,9 @@ namespace EpicGames.MCP.Automation
 		/// </summary>
 		/// <param name="Container">The name of the folder or container from which contains the file being checked.</param>
 		/// <param name="Identifier">The identifier or filename of the file to check.</param>
+		/// <param name="bQuiet">If set to true, all log output for the operation is supressed.</param>
 		/// <returns>True if the file exists in cloud storage, false otherwise.</returns>
-		abstract public bool FileExists(string Container, string Identifier);
+		abstract public bool FileExists(string Container, string Identifier, bool bQuiet = false);
 
 		/// <summary>
 		/// Retrieves a file from the cloud storage provider
@@ -771,9 +772,10 @@ namespace EpicGames.MCP.Automation
 		/// <param name="ContentType">The MIME type of the file being uploaded.</param>
 		/// <param name="bOverwrite">If true, will overwrite an existing file.  If false, will throw an exception if the file exists.</param>
 		/// <param name="bMakePublic">Specified whether the file should be made public readable.</param>
-		public PostFileResult PostFile(string Container, string Identifier, byte[] Contents, string ContentType = null, bool bOverwrite = true, bool bMakePublic = false)
+		/// <param name="bQuiet">If set to true, all log output for the operation is supressed.</param>
+		public PostFileResult PostFile(string Container, string Identifier, byte[] Contents, string ContentType = null, bool bOverwrite = true, bool bMakePublic = false, bool bQuiet = false)
 		{
-			return PostFileAsync(Container, Identifier, Contents, ContentType, bOverwrite, bMakePublic).Result;
+			return PostFileAsync(Container, Identifier, Contents, ContentType, bOverwrite, bMakePublic, bQuiet).Result;
 		}
 
 		/// <summary>
@@ -785,7 +787,8 @@ namespace EpicGames.MCP.Automation
 		/// <param name="ContentType">The MIME type of the file being uploaded.</param>
 		/// <param name="bOverwrite">If true, will overwrite an existing file.  If false, will throw an exception if the file exists.</param>
 		/// <param name="bMakePublic">Specified whether the file should be made public readable.</param>
-		abstract public Task<PostFileResult> PostFileAsync(string Container, string Identifier, byte[] Contents, string ContentType = null, bool bOverwrite = true, bool bMakePublic = false);
+		/// <param name="bQuiet">If set to true, all log output for the operation is supressed.</param>
+		abstract public Task<PostFileResult> PostFileAsync(string Container, string Identifier, byte[] Contents, string ContentType = null, bool bOverwrite = true, bool bMakePublic = false, bool bQuiet = false);
 
 		/// <summary>
 		/// Posts a file to the cloud storage provider.
@@ -795,10 +798,10 @@ namespace EpicGames.MCP.Automation
 		/// <param name="SourceFilePath">The full path of the file to upload.</param>
 		/// <param name="ContentType">The MIME type of the file being uploaded.</param>
 		/// <param name="bOverwrite">If true, will overwrite an existing file.  If false, will throw an exception if the file exists.</param>
-		/// <param name="bMakePublic">Specified whether the file should be made public readable.</param>
-		public PostFileResult PostFile(string Container, string Identifier, string SourceFilePath, string ContentType = null, bool bOverwrite = true, bool bMakePublic = false)
+		/// <param name="bQuiet">If set to true, all log output for the operation is supressed.</param>
+		public PostFileResult PostFile(string Container, string Identifier, string SourceFilePath, string ContentType = null, bool bOverwrite = true, bool bMakePublic = false, bool bQuiet = false)
 		{
-			return PostFileAsync(Container, Identifier, SourceFilePath, ContentType, bOverwrite, bMakePublic).Result;
+			return PostFileAsync(Container, Identifier, SourceFilePath, ContentType, bOverwrite, bMakePublic, bQuiet).Result;
 		}
 
 		/// <summary>
@@ -810,7 +813,8 @@ namespace EpicGames.MCP.Automation
 		/// <param name="ContentType">The MIME type of the file being uploaded.</param>
 		/// <param name="bOverwrite">If true, will overwrite an existing file.  If false, will throw an exception if the file exists.</param>
 		/// <param name="bMakePublic">Specified whether the file should be made public readable.</param>
-		abstract public Task<PostFileResult> PostFileAsync(string Container, string Identifier, string SourceFilePath, string ContentType = null, bool bOverwrite = true, bool bMakePublic = false);
+		/// <param name="bQuiet">If set to true, all log output for the operation is supressed.</param>
+		abstract public Task<PostFileResult> PostFileAsync(string Container, string Identifier, string SourceFilePath, string ContentType = null, bool bOverwrite = true, bool bMakePublic = false, bool bQuiet = false);
 
 		/// <summary>
 		/// Deletes a file from cloud storage
@@ -833,7 +837,7 @@ namespace EpicGames.MCP.Automation
 		/// <param name="Prefix">A string with which the identifier or filename should start. Typically used to specify a relative directory within the container to list all of its files recursively. Specify null to return all files.</param>
 		/// <param name="Recursive">Indicates whether the list of files returned should traverse subdirectories</param>
 		/// <returns>An array of paths to the files in the specified location and matching the prefix constraint.</returns>
-		abstract public string[] ListFiles(string Container, string Prefix = null, bool bRecursive = true);
+		abstract public string[] ListFiles(string Container, string Prefix = null, bool bRecursive = true, bool bQuiet = false);
 
 		/// <summary>
 		/// Sets one or more items of metadata on an object in cloud storage
