Fix P4Utils not being able to parse changelist forms if there is a 'Type:' field between 'Status:' and 'Description:'. Prevents IT from enabling restricted changelists, because GUBP can't update the CIS counter.

diff --git a/Engine/Source/Programs/AutomationTool/AutomationUtils/P4Utils.cs b/Engine/Source/Programs/AutomationTool/AutomationUtils/P4Utils.cs
--- a/Engine/Source/Programs/AutomationTool/AutomationUtils/P4Utils.cs
+++ b/Engine/Source/Programs/AutomationTool/AutomationUtils/P4Utils.cs
@@ -1670,16 +1670,14 @@ namespace AutomationTool
 
 				string StatusStr = "Status:";
 				int StatusOffset = CmdOutput.LastIndexOf(StatusStr);
-				string DescStr = "Description:";
-				int DescOffset = CmdOutput.LastIndexOf(DescStr);
 
-				if (StatusOffset < 1 || DescOffset < 1 || StatusOffset > DescOffset)
+				if (StatusOffset < 1)
 				{
 					CommandUtils.LogError("Change {0} could not be parsed\n{1}", CL, CmdOutput);
 					return false;
 				}
 
-				string Status = CmdOutput.Substring(StatusOffset + StatusStr.Length, DescOffset - StatusOffset - StatusStr.Length).Trim();
+				string Status = CmdOutput.Substring(StatusOffset + StatusStr.Length).TrimStart().Split('\n')[0].TrimEnd();
 				CommandUtils.Log("Change {0} exists ({1})", CL, Status);
 				Pending = (Status == "pending");
 				return true;
