Updating UAT paths to CrashReporterServer.

diff --git a/Engine/Source/Programs/AutomationTool/GUBP/LegacyNodes.cs b/Engine/Source/Programs/AutomationTool/GUBP/LegacyNodes.cs
--- a/Engine/Source/Programs/AutomationTool/GUBP/LegacyNodes.cs
+++ b/Engine/Source/Programs/AutomationTool/GUBP/LegacyNodes.cs
@@ -770,9 +770,9 @@ partial class GUBP
                 Agenda.DotNetProjects.AddRange(
                     new string[] 
 			    {
-                    CombinePaths(@"Engine\Source\Programs\NotForLicensees\CrashReportServer\CrashReportCommon\CrashReportCommon.csproj"),
-					CombinePaths(@"Engine\Source\Programs\NotForLicensees\CrashReportServer\CrashReportReceiver\CrashReportReceiver.csproj"),
-					CombinePaths(@"Engine\Source\Programs\NotForLicensees\CrashReportServer\CrashReportProcess\CrashReportProcess.csproj"),
+                    CombinePaths(@"Engine\Source\Programs\CrashReporter\CrashReportCommon\CrashReportCommon.csproj"),
+					CombinePaths(@"Engine\Source\Programs\CrashReporter\CrashReportReceiver\CrashReportReceiver.csproj"),
+					CombinePaths(@"Engine\Source\Programs\CrashReporter\CrashReportProcess\CrashReportProcess.csproj"),
                     CombinePaths(@"Engine\Source\Programs\CrashReporter\RegisterPII\RegisterPII.csproj"),
 			    });
                 Agenda.DotNetSolutions.AddRange(
