Fix CIS failure to build InternalTools due to RegisterPII being deleted.

diff --git a/Engine/Source/Programs/AutomationTool/GUBP/LegacyNodes.cs b/Engine/Source/Programs/AutomationTool/GUBP/LegacyNodes.cs
--- a/Engine/Source/Programs/AutomationTool/GUBP/LegacyNodes.cs
+++ b/Engine/Source/Programs/AutomationTool/GUBP/LegacyNodes.cs
@@ -773,7 +773,6 @@ partial class GUBP
                     CombinePaths(@"Engine\Source\Programs\CrashReporter\CrashReportCommon\CrashReportCommon.csproj"),
 					CombinePaths(@"Engine\Source\Programs\CrashReporter\CrashReportReceiver\CrashReportReceiver.csproj"),
 					CombinePaths(@"Engine\Source\Programs\CrashReporter\CrashReportProcess\CrashReportProcess.csproj"),
-                    CombinePaths(@"Engine\Source\Programs\CrashReporter\RegisterPII\RegisterPII.csproj"),
 			    });
                 Agenda.DotNetSolutions.AddRange(
                     new string[] 
