Fix preflights from incorrectly culling nodes in the setup step, then culling them in the individual nodes (causing lots of nodes to be NOPs). -TimeIndex parameter is set by the workflow, but overridden on preflights with -UserTimeIndex. Switch the order to allow -UserTimeIndex to override.

diff --git a/Engine/Source/Programs/AutomationTool/GUBP/GUBP.cs b/Engine/Source/Programs/AutomationTool/GUBP/GUBP.cs
--- a/Engine/Source/Programs/AutomationTool/GUBP/GUBP.cs
+++ b/Engine/Source/Programs/AutomationTool/GUBP/GUBP.cs
@@ -183,10 +183,10 @@ public partial class GUBP : BuildCommand
 		FindCompletionState(Graph.BuildNodes, JobInfo, LocalOnly);
 		ComputeDependentFrequencies(Graph.BuildNodes);
 
-		int TimeIndex = ParseParamInt("TimeIndex", -1);
+		int TimeIndex = ParseParamInt("UserTimeIndex", -1);
 		if (TimeIndex == -1)
 		{
-			TimeIndex = ParseParamInt("UserTimeIndex", -1);
+			TimeIndex = ParseParamInt("TimeIndex", -1);
 		}
 		if (ParseParam("CIS") && ExplicitTriggerName == "" && CommanderSetup) // explicit triggers will already have a time index assigned
 		{
