Use quiet mode for zip command when running it from UAT

diff --git a/Engine/Source/Programs/AutomationTool/AutomationUtils/CommandUtils.cs b/Engine/Source/Programs/AutomationTool/AutomationUtils/CommandUtils.cs
--- a/Engine/Source/Programs/AutomationTool/AutomationUtils/CommandUtils.cs
+++ b/Engine/Source/Programs/AutomationTool/AutomationUtils/CommandUtils.cs
@@ -2097,13 +2097,13 @@ namespace AutomationTool
 					FilesList += " \"" + FilteredFile + "\"";
 					if (FilesList.Length > 32000)
 					{
-						CommandUtils.RunAndLog(CommandUtils.CmdEnv, "zip", "-g " + ZipFileName + FilesList);
+						CommandUtils.RunAndLog(CommandUtils.CmdEnv, "zip", "-g -q " + ZipFileName + FilesList);
 						FilesList = "";
 					}
 				}
 				if (FilesList.Length > 0)
 				{
-					CommandUtils.RunAndLog(CommandUtils.CmdEnv, "zip", "-g " + ZipFileName + FilesList);
+					CommandUtils.RunAndLog(CommandUtils.CmdEnv, "zip", "-g -q " + ZipFileName + FilesList);
 				}
 				CommandUtils.PopDir();
 			}
