fix for DDC stats warning

diff --git a/Engine/Source/Programs/AutomationTool/AutomationUtils/CommandletUtils.cs b/Engine/Source/Programs/AutomationTool/AutomationUtils/CommandletUtils.cs
--- a/Engine/Source/Programs/AutomationTool/AutomationUtils/CommandletUtils.cs
+++ b/Engine/Source/Programs/AutomationTool/AutomationUtils/CommandletUtils.cs
@@ -259,14 +259,6 @@ namespace AutomationTool
 			{
 				CommandUtils.LogWarning("Commandlet {0} failed to copy the local log file from {1} to {2}. The log file will be lost.", Commandlet, LocalLogFile, DestLogFile);
 			}
-			// copy the ddc access log to the destination folder
-			string DestDDCLogFile = LogUtils.GetUniqueLogName(CombinePaths(CmdEnv.LogFolder, Commandlet + "-DDC"));
-			string LocalDDCLogFile = CombinePaths(CmdEnv.EngineSavedFolder, "CookRunDDC.txt");
-			if (!CommandUtils.FileExists(LocalDDCLogFile) || !CommandUtils.CopyFile_NoExceptions(LocalDDCLogFile, DestDDCLogFile))
-			{
-				CommandUtils.LogWarning("Commandlet {0} failed to copy the local log file from {1} to {2}. The log file will be lost.", Commandlet, LocalDDCLogFile, DestDDCLogFile);
-			}
-
             string ProjectStatsDirectory = CombinePaths((ProjectName == null)? CombinePaths(CmdEnv.LocalRoot, "Engine") : Path.GetDirectoryName(ProjectName.FullName), "Saved", "Stats");
             if (Directory.Exists(ProjectStatsDirectory))
             {
@@ -275,7 +267,7 @@ namespace AutomationTool
                 {
                     if (!CommandUtils.CopyFile_NoExceptions(StatsFile, CombinePaths(DestCookerStats, Path.GetFileName(StatsFile))))
                     {
-                        CommandUtils.LogWarning("Commandlet {0} failed to copy the local log file from {1} to {2}. The log file will be lost.", Commandlet, LocalDDCLogFile, DestDDCLogFile);
+						CommandUtils.LogWarning("Commandlet {0} failed to copy the local log file from {1} to {2}. The log file will be lost.", Commandlet, StatsFile, CombinePaths(DestCookerStats, Path.GetFileName(StatsFile)));
                     }
                 }
             }
@@ -286,10 +278,6 @@ namespace AutomationTool
 
 			// Whether it was copied correctly or not, delete the local log as it was only a temporary file. 
 			CommandUtils.DeleteFile_NoExceptions(LocalLogFile);
-			if (CommandUtils.FileExists(LocalDDCLogFile))
-			{
-				CommandUtils.DeleteFile_NoExceptions(LocalDDCLogFile);
-			}
 
 			if (RunResult.ExitCode != 0)
 			{
