Comment out the last (on PC) warning when building DDC, due to stats directory not being present. We're not using DDC stats at the moment.

diff --git a/Engine/Source/Programs/AutomationTool/AutomationUtils/CommandletUtils.cs b/Engine/Source/Programs/AutomationTool/AutomationUtils/CommandletUtils.cs
--- a/Engine/Source/Programs/AutomationTool/AutomationUtils/CommandletUtils.cs
+++ b/Engine/Source/Programs/AutomationTool/AutomationUtils/CommandletUtils.cs
@@ -271,10 +271,10 @@ namespace AutomationTool
                     }
                 }
             }
-            else
-            {
-                CommandUtils.LogWarning("Failed to find directory {0} will not save stats", ProjectStatsDirectory);
-            }
+//			else
+//			{
+//				CommandUtils.LogWarning("Failed to find directory {0} will not save stats", ProjectStatsDirectory);
+//			}
 
 			// Whether it was copied correctly or not, delete the local log as it was only a temporary file. 
 			CommandUtils.DeleteFile_NoExceptions(LocalLogFile);
