Copying //UE4/Dev-Build to //UE4/Main

==========================
MAJOR FEATURES + CHANGES
==========================

Change 2790858 on 2015/12/04 by [REDACTED_PERSON]

	Fix ERRORLEVEL not being correctly returned by build.bat, causing failed builds to be treated as succeeded builds when launched from Visual Studio.

Change 2792683 on 2015/12/07 by [REDACTED_PERSON]

	Removed some usages of RunningRocket() function

	Checks for whether engine tools need rebuilding should be ignored in all installed engine builds.
	New Intermediate folders should be within project directory and engine's should not be removed when engine installed.

Change 2794194 on 2015/12/08 by [REDACTED_PERSON]

	Removed IsRocket function.
	Changed IsEngineInstalled so that -NotInstalledEngine will always disable installed engine behavior.

Change 2801483 on 2015/12/14 by [REDACTED_PERSON]

	Removed uses of UnrealBuildTool.RunningRocket()

	Mac Tool chain test was down to Engine libraries being pre-compiled
	Allow Third Party dependencies in any monolithic target, not just games
	All installed engine builds should use shared build environment
	Removed test no longer needed now that projects from installed builds have game targets
	Include engine source if it exists in any installed build to improve intellisense in VCProjects

Change 2803276 on 2015/12/15 by [REDACTED_PERSON]

	Removed Rocket switch from ProjectParams as it can always check Automation.RunningRocket()
	Removed bGeneratingRocketProjectFiles from ProjectFileGenerator as it can always check UnrealBuildTool.RunningRocket()
	Changed a few usages to IsEngineInstalled where obvious

Change 2814190 on 2016/01/03 by [REDACTED_PERSON]

	Convert ParallelExecutor to managed code, so we can use it trivially from UAT without having to compile with UBT first. Should be moved into UBT at some point.

	Also add a support class for creating managed child processes, which are terminated automatically if the parent process is closed. Currently only implemented on Windows, using P/Invoke. Also allows setting priority level of the process, which could replace the hacky way that we determine the physical processor count in UBT.

Change 2822503 on 2016/01/10 by [REDACTED_PERSON]

	AnalyticsET now uses the Data Router protocol.
	* Automatically detects configs that try to use the old endpoint and routes them to the new endpoint. This code is temp until internal games merge the code and have a chance to update their configs.
	* There is a new attribute on BeginSession called "LegacyURL" that is "true" if the code needs to have its configs updated, as well as a LogAnalytics warning.
	* Updated EngineAnalytics and AutomationAnalytics providers to use the new URL and specify any appropriate new, optional config values.
	* Update CrashReporter provider to use new URL and simplified it's provider configuration (now uses AnalyticsET config struct directly).
	Removed recently added GetUE4TypeOverride and moved it to the one place in Engine that needs it.

diff --git a/Engine/Build/BatchFiles/Build.bat b/Engine/Build/BatchFiles/Build.bat
--- a/Engine/Build/BatchFiles/Build.bat
+++ b/Engine/Build/BatchFiles/Build.bat
@@ -13,7 +13,7 @@ REM %3 is the configuration name
 IF EXIST ..\..\Engine\Binaries\DotNET\UnrealBuildTool.exe (
          ..\..\Engine\Binaries\DotNET\UnrealBuildTool.exe %* -DEPLOY
 		 popd
-		 EXIT /B !RET!
+		 EXIT /B !ERRORLEVEL!
 ) ELSE (
 	ECHO UnrealBuildTool.exe not found in ..\..\Engine\Binaries\DotNET\UnrealBuildTool.exe 
 	popd
diff --git a/Engine/Build/InstalledEngineFilters.ini b/Engine/Build/InstalledEngineFilters.ini
--- a/Engine/Build/InstalledEngineFilters.ini
+++ b/Engine/Build/InstalledEngineFilters.ini
@@ -44,6 +44,7 @@
 /Engine/Source/Editor/...
 /Engine/Source/ThirdParty/.../*.Build.cs
 /Engine/Source/ThirdParty/Licenses/...
+-/Engine/Source/Runtime/SQLiteSupport/...
 
 ; Shaders
 /Engine/Shaders/...
diff --git a/Engine/Source/Programs/AutomationTool/AutomationUtils/AutomationUtils.Automation.csproj b/Engine/Source/Programs/AutomationTool/AutomationUtils/AutomationUtils.Automation.csproj
--- a/Engine/Source/Programs/AutomationTool/AutomationUtils/AutomationUtils.Automation.csproj
+++ b/Engine/Source/Programs/AutomationTool/AutomationUtils/AutomationUtils.Automation.csproj
@@ -56,6 +56,7 @@
     <Compile Include="BuildCommand.cs" />
     <Compile Include="BuildEnvironment.cs" />
     <Compile Include="BuildUtils.cs" />
+    <Compile Include="ManagedProcess.cs" />
     <Compile Include="CommandletUtils.cs" />
     <Compile Include="CommandUtils.cs" />
     <Compile Include="DeploymentContext.cs" />
@@ -72,6 +73,7 @@
     <Compile Include="P4Blame.cs" />
     <Compile Include="P4Environment.cs" />
     <Compile Include="P4Utils.cs" />
+    <Compile Include="ParallelExecutor.cs" />
     <Compile Include="Platform.cs" />
     <Compile Include="ProcessUtils.cs" />
     <Compile Include="ProjectParams.cs" />
diff --git a/Engine/Source/Programs/AutomationTool/AutomationUtils/CommandUtils.cs b/Engine/Source/Programs/AutomationTool/AutomationUtils/CommandUtils.cs
--- a/Engine/Source/Programs/AutomationTool/AutomationUtils/CommandUtils.cs
+++ b/Engine/Source/Programs/AutomationTool/AutomationUtils/CommandUtils.cs
@@ -1947,13 +1947,6 @@ namespace AutomationTool
             return CachedRootBuildStorageDirectory;
         }
 
-		// Included for compatibility during //UE4/Main import
-		[Obsolete]
-		public static string RootSharedTempStorageDirectory()
-		{
-			return RootBuildStorageDirectory();
-		}
-
         public static bool DirectoryExistsAndIsWritable_NoExceptions(string Dir)
         {
             if (!DirectoryExists_NoExceptions(Dir))
diff --git a/Engine/Source/Programs/AutomationTool/AutomationUtils/ParallelExecutor.cs b/Engine/Source/Programs/AutomationTool/AutomationUtils/ParallelExecutor.cs
new file mode 100644
--- /dev/null
+++ b/Engine/Source/Programs/AutomationTool/AutomationUtils/ParallelExecutor.cs
@@ -0,0 +1,405 @@
+﻿using System;
+using System.Collections.Generic;
+using System.Diagnostics;
+using System.Linq;
+using System.Text;
+using System.Threading;
+using System.Threading.Tasks;
+using System.Xml;
+
+namespace AutomationTool
+{
+	static class ParallelExecutor
+	{
+		[DebuggerDisplay("{Caption}")]
+		class BuildAction
+		{
+			public int SortIndex = -1;
+
+			public string Caption;
+			public string GroupPrefix;
+			public string OutputPrefix;
+			public string ToolPath;
+			public string ToolArguments;
+			public string WorkingDirectory;
+			public bool bSkipIfProjectFailed;
+
+			public List<BuildAction> Dependants = new List<BuildAction>();
+			public List<BuildAction> Dependencies = new List<BuildAction>();
+
+			public int TotalDependants;
+			public int MissingDependencyCount;
+
+			public IReadOnlyDictionary<string, string> Environment;
+		}
+
+		class BuildActionExecutor
+		{
+			public BuildAction Action;
+			public List<string> LogLines = new List<string>();
+			public int ExitCode = -1;
+
+			private ManualResetEvent CompletedEvent;
+			private List<BuildActionExecutor> CompletedActions;
+
+			public BuildActionExecutor(BuildAction InAction, ManualResetEvent InCompletedEvent, List<BuildActionExecutor> InCompletedActions)
+			{
+				Action = InAction;
+				CompletedEvent = InCompletedEvent;
+				CompletedActions = InCompletedActions;
+			}
+
+			public void Run()
+			{
+				if(!String.IsNullOrEmpty(Action.OutputPrefix))
+				{
+					LogLines.Add(Action.OutputPrefix);
+				}
+
+				using(ManagedProcess Process = new ManagedProcess(Action.ToolPath, Action.ToolArguments, Action.WorkingDirectory, Action.Environment, null, ManagedProcessPriority.BelowNormal))
+				{
+					LogLines.AddRange(Process.ReadAllLines());
+					ExitCode = Process.ExitCode;
+				}
+
+				lock(CompletedActions)
+				{
+					CompletedActions.Add(this);
+				}
+
+				CompletedEvent.Set();
+			}
+		}
+
+		public static int Execute(string ActionsFileName)
+		{
+			return Execute(ActionsFileName, Environment.ProcessorCount);
+		}
+
+		public static int Execute(string ActionsFileName, int MaxProcesses)
+		{
+			List<BuildAction> Actions = ReadActions(ActionsFileName);
+
+			CommandUtils.Log("Building with {0} {1}...", MaxProcesses, (MaxProcesses == 1)? "process" : "processes");
+
+			// Create the list of things to process
+			List<BuildAction> QueuedActions = new List<BuildAction>();
+			foreach(BuildAction Action in Actions)
+			{
+				if(Action.MissingDependencyCount == 0)
+				{
+					QueuedActions.Add(Action);
+				}
+			}
+
+			// Create a job object for all the child processes
+			int ExitCode = 0;
+			string CurrentPrefix = "";
+			Dictionary<BuildActionExecutor, Thread> ExecutingActions = new Dictionary<BuildActionExecutor,Thread>();
+			List<BuildActionExecutor> CompletedActions = new List<BuildActionExecutor>();
+
+			using(ManualResetEvent CompletedEvent = new ManualResetEvent(false))
+			{
+				while(QueuedActions.Count > 0 || ExecutingActions.Count > 0)
+				{
+					// Sort the actions by the number of things dependent on them
+					QueuedActions.Sort((A, B) => (A.TotalDependants == B.TotalDependants)? (B.SortIndex - A.SortIndex) : (B.TotalDependants - A.TotalDependants));
+
+					// Create threads up to the maximum number of actions
+					while(ExecutingActions.Count < MaxProcesses && QueuedActions.Count > 0)
+					{
+						BuildAction Action = QueuedActions[QueuedActions.Count - 1];
+						QueuedActions.RemoveAt(QueuedActions.Count - 1);
+
+						BuildActionExecutor ExecutingAction = new BuildActionExecutor(Action, CompletedEvent, CompletedActions);
+
+						Thread ExecutingThread = new Thread(() => { ExecutingAction.Run(); });
+						ExecutingThread.Name = String.Format("Build:{0}", Action.Caption);
+						ExecutingThread.Start();
+
+						ExecutingActions.Add(ExecutingAction, ExecutingThread);
+					}
+
+					// Wait for something to finish
+					CompletedEvent.WaitOne();
+
+					// Wait for something to finish and flush it to the log
+					lock(CompletedActions)
+					{
+						foreach(BuildActionExecutor CompletedAction in CompletedActions)
+						{
+							// Join the thread
+							Thread CompletedThread = ExecutingActions[CompletedAction];
+							CompletedThread.Join();
+							ExecutingActions.Remove(CompletedAction);
+
+							// Write it to the log
+							if(CompletedAction.LogLines.Count > 0)
+							{
+								if(CurrentPrefix != CompletedAction.Action.GroupPrefix)
+								{
+									CurrentPrefix = CompletedAction.Action.GroupPrefix;
+									CommandUtils.Log(CurrentPrefix);
+								}
+								foreach(string LogLine in CompletedAction.LogLines)
+								{
+									CommandUtils.Log(LogLine);
+								}
+							}
+
+							// Check the exit code
+							if(CompletedAction.ExitCode != 0)
+							{
+								ExitCode = CompletedAction.ExitCode;
+								if(!CompletedAction.Action.bSkipIfProjectFailed)
+								{
+									foreach(KeyValuePair<BuildActionExecutor, Thread> Pair in ExecutingActions)
+									{
+										Pair.Value.Abort();
+										Pair.Value.Join();
+									}
+									return ExitCode;
+								}
+							}
+
+							// Mark all the dependents as done
+							foreach(BuildAction DependantAction in CompletedAction.Action.Dependants)
+							{
+								if(--DependantAction.MissingDependencyCount == 0)
+								{
+									QueuedActions.Add(DependantAction);
+								}
+							}
+						}
+						CompletedActions.Clear();
+					}
+				}
+			}
+			return ExitCode;
+		}
+
+		private static List<BuildAction> ReadActions(string InputPath)
+		{
+			XmlDocument File = new XmlDocument();
+			File.Load(InputPath);
+
+			XmlNode RootNode = File.FirstChild;
+			if(RootNode == null || RootNode.Name != "BuildSet")
+			{
+				throw new Exception("Incorrect node at root of graph; expected 'BuildSet'");
+			}
+
+			XmlNode EnvironmentsNode = RootNode.SelectSingleNode("Environments");
+			if(EnvironmentsNode == null)
+			{
+				throw new Exception("Missing Environments node under root in XML document");
+			}
+
+			// Get the current environment variables
+			Dictionary<string, string> CurrentEnvironment = new Dictionary<string,string>(StringComparer.InvariantCultureIgnoreCase);
+			foreach(System.Collections.DictionaryEntry Entry in Environment.GetEnvironmentVariables())
+			{
+				CurrentEnvironment[Entry.Key.ToString()] = Entry.Value.ToString();
+			}
+
+			// Read the tool environment
+			Dictionary<string, XmlNode> NameToTool = new Dictionary<string,XmlNode>();
+			Dictionary<string, IReadOnlyDictionary<string, string>> NameToEnvironment = new Dictionary<string,IReadOnlyDictionary<string, string>>();
+			for(XmlNode EnvironmentNode = EnvironmentsNode.FirstChild; EnvironmentNode != null; EnvironmentNode = EnvironmentNode.NextSibling)
+			{
+				if(EnvironmentNode.Name == "Environment")
+				{
+					// Read the tools nodes
+					XmlNode ToolsNode = EnvironmentNode.SelectSingleNode("Tools");
+					if(ToolsNode != null)
+					{
+						for(XmlNode ToolNode = ToolsNode.FirstChild; ToolNode != null; ToolNode = ToolNode.NextSibling)
+						{
+							if(ToolNode.Name == "Tool")
+							{
+								XmlNode Name = ToolNode.Attributes.GetNamedItem("Name");
+								if(Name != null)
+								{
+									NameToTool.Add(Name.Value, ToolNode);
+								}
+							}
+						}
+					}
+
+					// Read the environment variables for this environment. Each environment has its own set of variables 
+					XmlNode VariablesNode = EnvironmentNode.SelectSingleNode("Variables");
+					if(VariablesNode != null)
+					{
+						XmlNode Name = EnvironmentNode.Attributes.GetNamedItem("Name");
+						if(Name != null)
+						{
+							Dictionary<string, string> NamedEnvironment = new Dictionary<string,string>(CurrentEnvironment, StringComparer.InvariantCultureIgnoreCase);
+							for (XmlNode VariableNode = VariablesNode.FirstChild; VariableNode != null; VariableNode = VariableNode.NextSibling)
+							{
+								if (VariableNode.Name == "Variable")
+								{
+									XmlNode VariableName = VariableNode.Attributes.GetNamedItem("Name");
+									if(VariableName != null)
+									{
+										NamedEnvironment[VariableName.Value] = VariableNode.Attributes.GetNamedItem("Value").Value;
+									}
+								}
+							}
+							NameToEnvironment[Name.Value] = NamedEnvironment;
+						}
+					}
+				}
+			}
+
+			// Read all the tasks for each project, and convert them into actions
+			List<BuildAction> Actions = new List<BuildAction>();
+			for(XmlNode ProjectNode = RootNode.FirstChild; ProjectNode != null; ProjectNode = ProjectNode.NextSibling)
+			{
+				if(ProjectNode.Name == "Project")
+				{
+					int SortIndex = 0;
+					Dictionary<string, BuildAction> NameToAction = new Dictionary<string,BuildAction>();
+					for(XmlNode TaskNode = ProjectNode.FirstChild; TaskNode != null; TaskNode = TaskNode.NextSibling)
+					{
+						if(TaskNode.Name == "Task")
+						{
+							XmlNode ToolNode;
+							if(NameToTool.TryGetValue(TaskNode.Attributes.GetNamedItem("Tool").Value, out ToolNode))
+							{
+								BuildAction Action = FindOrAddAction(NameToAction, TaskNode.Attributes.GetNamedItem("Name").Value, Actions);
+								Action.SortIndex = ++SortIndex;
+								TryGetAttribute(TaskNode, "Caption", out Action.Caption);
+								
+								TryGetAttribute(ToolNode, "GroupPrefix", out Action.GroupPrefix);
+								TryGetAttribute(ToolNode, "OutputPrefix", out Action.OutputPrefix);
+								TryGetAttribute(ToolNode, "Path", out Action.ToolPath);
+								TryGetAttribute(ToolNode, "WorkingDir", out Action.WorkingDirectory);
+
+								string EnvironmentName;
+								if(!TryGetAttribute(ProjectNode, "Env", out EnvironmentName))
+								{
+									Action.Environment = CurrentEnvironment;
+								}
+								else if(!NameToEnvironment.TryGetValue(EnvironmentName, out Action.Environment))
+								{
+									throw new Exception(String.Format("Couldn't find environment '{0}'", EnvironmentName));
+								}
+
+								if(TryGetAttribute(ToolNode, "Params", out Action.ToolArguments))
+								{
+									Action.ToolArguments = ExpandEnvironmentVariables(Action.ToolArguments, Action.Environment);
+								}
+
+								string SkipIfProjectFailed;
+								if(TryGetAttribute(TaskNode, "SkipIfProjectFailed", out SkipIfProjectFailed) && SkipIfProjectFailed == "True")
+								{
+									Action.bSkipIfProjectFailed = true;
+								}
+
+								string DependsOnList;
+								if(TryGetAttribute(TaskNode, "DependsOn", out DependsOnList))
+								{
+									foreach (string DependsOn in DependsOnList.Split(';'))
+									{
+										BuildAction DependsOnAction = FindOrAddAction(NameToAction, DependsOn, Actions);
+										if(!Action.Dependencies.Contains(DependsOnAction))
+										{
+											Action.Dependencies.Add(DependsOnAction);
+										}
+										if(!DependsOnAction.Dependants.Contains(Action))
+										{
+											DependsOnAction.Dependants.Add(Action);
+										}
+									}
+								}
+
+								HashSet<BuildAction> VisitedActions = new HashSet<BuildAction>();
+								RecursiveIncDependents(Action, VisitedActions);
+
+								Action.MissingDependencyCount = Action.Dependencies.Count;
+							}
+						}
+					}
+
+					// Make sure there have been no actions added which were referenced but never declared
+					foreach(KeyValuePair<string, BuildAction> Pair in NameToAction)
+					{
+						if(Pair.Value.SortIndex == -1)
+						{
+							throw new Exception(String.Format("Action {0} was referenced but never declared", Pair.Key));
+						}
+					}
+				}
+			}
+			return Actions;
+		}
+
+		static bool TryGetAttribute(XmlNode Node, string Name, out string Value)
+		{
+			XmlNode Attribute = Node.Attributes.GetNamedItem(Name);
+			if(Attribute == null)
+			{
+				Value = null;
+				return false;
+			}
+			else
+			{
+				Value = Attribute.Value;
+				return true;
+			}
+		}
+
+		static string ExpandEnvironmentVariables(string Input, IReadOnlyDictionary<string, string> Environment)
+		{
+			string Output = Input;
+			for(int StartIndex = 0;;)
+			{
+				int Index = Output.IndexOf("$(", StartIndex);
+				if(Index == -1) break;
+
+				int EndIndex = Output.IndexOf(")", Index);
+				if(EndIndex == -1) break;
+
+				string VariableName = Output.Substring(Index + 2, EndIndex - Index - 2);
+
+				string VariableValue;
+				if(Environment.TryGetValue(VariableName, out VariableValue))
+				{
+					Output = Output.Substring(0, Index) + VariableValue + Output.Substring(EndIndex + 1);
+					StartIndex = Index + VariableValue.Length;
+				}
+				else
+				{
+					CommandUtils.LogWarning("Unknown environment variable '{0}' in script", VariableName);
+					StartIndex = EndIndex + 1;
+				}
+			}
+			return Output;
+		}
+
+		private static BuildAction FindOrAddAction(Dictionary<string, BuildAction> NameToAction, string Name, List<BuildAction> Actions)
+		{
+			BuildAction Action;
+			if(!NameToAction.TryGetValue(Name, out Action))
+			{
+				Action = new BuildAction();
+				Actions.Add(Action);
+				NameToAction.Add(Name, Action);
+			}
+			return Action;
+		}
+
+		private static void RecursiveIncDependents(BuildAction Action, HashSet<BuildAction> VisitedActions)
+		{
+			foreach(BuildAction Dependency in Action.Dependencies)
+			{
+				if(!VisitedActions.Contains(Action))
+				{
+					VisitedActions.Add(Action);
+					Dependency.TotalDependants++;
+					RecursiveIncDependents(Dependency, VisitedActions);
+				}
+			}
+		}
+	}
+}
diff --git a/Engine/Source/Programs/AutomationTool/AutomationUtils/ProjectParams.cs b/Engine/Source/Programs/AutomationTool/AutomationUtils/ProjectParams.cs
--- a/Engine/Source/Programs/AutomationTool/AutomationUtils/ProjectParams.cs
+++ b/Engine/Source/Programs/AutomationTool/AutomationUtils/ProjectParams.cs
@@ -270,7 +270,6 @@ namespace AutomationTool
             this.CrashIndex = InParams.CrashIndex;
             this.Port = InParams.Port;
 			this.SkipServer = InParams.SkipServer;
-			this.Rocket = InParams.Rocket;
 			this.Unattended = InParams.Unattended;
             this.ServerDeviceAddress = InParams.ServerDeviceAddress;
             this.DeviceUsername = InParams.DeviceUsername;
@@ -384,7 +383,6 @@ namespace AutomationTool
             bool? RunAutomationTests = null,
             string RunAutomationTest = null,
             int? CrashIndex = null,
-            bool? Rocket = null,
 			bool? SkipCook = null,
 			bool? SkipCookOnTheFly = null,
 			bool? SkipPak = null,
@@ -609,7 +607,6 @@ namespace AutomationTool
             this.RunAutomationTest = ParseParamValueIfNotSpecified(Command, RunAutomationTest, "RunAutomationTest");
             this.RunAutomationTests = this.RunAutomationTest != "" || GetParamValueIfNotSpecified(Command, RunAutomationTests, this.RunAutomationTests, "RunAutomationTests");
             this.SkipServer = GetParamValueIfNotSpecified(Command, SkipServer, this.SkipServer, "skipserver");
-			this.Rocket = GetParamValueIfNotSpecified(Command, Rocket, this.Rocket || Automation.RunningRocket(), "rocket");
 			this.UE4Exe = ParseParamValueIfNotSpecified(Command, UE4Exe, "ue4exe", "UE4Editor-Cmd.exe");
 			this.Unattended = GetParamValueIfNotSpecified(Command, Unattended, this.Unattended, "unattended");
 			this.DeviceUsername = ParseParamValueIfNotSpecified(Command, DeviceUsername, "deviceuser", String.Empty);
@@ -803,13 +800,7 @@ namespace AutomationTool
 		public bool ForeignCode { private set; get; }
 
 		/// <summary>
-		/// Shared: true if we are running rocket
-		/// </summary>
-		[Help("Rocket", "true if we are running rocket")]
-		public bool Rocket { private set; get; }
-
-		/// <summary>
-		/// Shared: true if we are running rocket
+		/// Shared: true if we should build crash reporter
 		/// </summary>
 		[Help("CrashReporter", "true if we should build crash reporter")]
 		public bool CrashReporter { private set; get; }
@@ -1703,7 +1694,7 @@ namespace AutomationTool
 
 			IsProgramTarget = ProjectType == TargetRules.TargetType.Program;
 
-			if (String.IsNullOrEmpty(EditorTarget) && ProjectType != TargetRules.TargetType.Program && CommandUtils.IsNullOrEmpty(EditorTargetsList) && !Rocket)
+			if (String.IsNullOrEmpty(EditorTarget) && ProjectType != TargetRules.TargetType.Program && CommandUtils.IsNullOrEmpty(EditorTargetsList) && !Automation.RunningRocket())
 			{
 				if (Properties.bWasGenerated)
 				{
@@ -1714,7 +1705,7 @@ namespace AutomationTool
 					throw new AutomationException("Editor target not found!");
 				}
 			}
-			if (String.IsNullOrEmpty(GameTarget) && Run && !NoClient && (Cook || CookOnTheFly) && CommandUtils.IsNullOrEmpty(ClientCookedTargetsList) && !Rocket)
+			if (String.IsNullOrEmpty(GameTarget) && Run && !NoClient && (Cook || CookOnTheFly) && CommandUtils.IsNullOrEmpty(ClientCookedTargetsList) && !Automation.RunningRocket())
 			{
 				throw new AutomationException("Game target not found. Game target is required with -cook or -cookonthefly");
 			}
@@ -2281,7 +2272,7 @@ namespace AutomationTool
 				CommandUtils.LogLog("Prereqs={0}", Prereqs);
 				CommandUtils.LogLog("NoBootstrapExe={0}", NoBootstrapExe);
 				CommandUtils.LogLog("RawProjectPath={0}", RawProjectPath);
-				CommandUtils.LogLog("Rocket={0}", Rocket);
+				CommandUtils.LogLog("Rocket={0}", Automation.RunningRocket());
 				CommandUtils.LogLog("Run={0}", Run);
 				CommandUtils.LogLog("ServerConfigsToBuild={0}", string.Join(",", ServerConfigsToBuild));
 				CommandUtils.LogLog("ServerCookedTargets={0}", ServerCookedTargets.ToString());
diff --git a/Engine/Source/Programs/AutomationTool/AutomationUtils/UE4Build.cs b/Engine/Source/Programs/AutomationTool/AutomationUtils/UE4Build.cs
--- a/Engine/Source/Programs/AutomationTool/AutomationUtils/UE4Build.cs
+++ b/Engine/Source/Programs/AutomationTool/AutomationUtils/UE4Build.cs
@@ -21,7 +21,6 @@ namespace AutomationTool
 		private BuildCommand OwnerCommand;
 
 		private String XGEConsoleExePath = "";
-		private String ParallelExecutorPath = "";
 
 		public void PrepareBuildProduct(string File, bool DeleteProduct = true)
 		{
@@ -735,27 +734,8 @@ namespace AutomationTool
 			return XGEConsoleExePath;
 		}
 
-		public string ParallelExecutorExec()
-		{
-			if (string.IsNullOrEmpty(ParallelExecutorPath))
-			{
-				string RootPath = Path.Combine(CmdEnv.LocalRoot, @"Engine\Binaries\Win64");
-				string FullPath = Path.Combine(RootPath, "ParallelExecutor.exe");
-				if (FileExists(FullPath))
-				{
-					ParallelExecutorPath = FullPath;
-				}
-			}
-			return ParallelExecutorPath;
-		}
-
 		public bool ProcessXGEItems(List<XGEItem> Actions, string XGETool, string Args, string TaskFilePath, bool DoRetries, bool SpecialTestFlag, bool ShowProgress)
 		{
-            if (string.IsNullOrEmpty(XGETool))
-			{
-				throw new AutomationException("Unable to find xge tool: " + XGETool);
-			}
-
 			TelemetryStopwatch CombineXGEStopwatch = new TelemetryStopwatch("CombineXGEItemFiles.{0}", Path.GetFileNameWithoutExtension(XGETool));
 
 			XmlDocument XGETaskDocument;	
@@ -825,10 +805,22 @@ namespace AutomationTool
 								LogVerbose("Running {0} *******", XGETool);
 								PushDir(CombinePaths(CmdEnv.LocalRoot, @"\Engine\Source"));
 								int SuccesCode;
-								string LogFile = GetRunAndLogOnlyName(CmdEnv, XGETool);
-								string Output = RunAndLog(XGETool, Args, out SuccesCode, LogFile);
+								string LogFile;
+								bool bOutputContainsProject;
+								if(XGETool == null)
+								{
+									LogFile = GetRunAndLogOnlyName(CmdEnv, "ParallelExecutor");
+									SuccesCode = ParallelExecutor.Execute(TaskFilePath);
+									bOutputContainsProject = false;
+								}
+								else
+								{
+									LogFile = GetRunAndLogOnlyName(CmdEnv, XGETool);
+									string Output = RunAndLog(XGETool, Args, out SuccesCode, LogFile);
+									bOutputContainsProject = Output.Contains("------Project:");
+								}
 								PopDir();
-								if (ConnectionRetries > 0 && (SuccesCode == 4 || SuccesCode == 2) && !Output.Contains("------Project:"))
+								if (ConnectionRetries > 0 && (SuccesCode == 4 || SuccesCode == 2) && !bOutputContainsProject)
 								{
 									LogWarning(String.Format("{0} failure on the local connection timeout", XGETool));
 									if (ConnectionRetries < 2)
@@ -1334,7 +1326,7 @@ namespace AutomationTool
 			}
 
 			// only run ParallelExecutor if not running XGE (and we've requested ParallelExecutor and it exists)
-			bool bCanUseParallelExecutor = !bCanUseXGE && InUseParallelExecutor && !string.IsNullOrEmpty(ParallelExecutorExec());
+			bool bCanUseParallelExecutor = !bCanUseXGE && InUseParallelExecutor;
 			LogLog("************************* UE4Build:");
 			LogLog("************************* ForceMonolithic: {0}", bForceMonolithic);
 			LogLog("************************* ForceNonUnity:{0} ", bForceNonUnity);
@@ -1361,7 +1353,7 @@ namespace AutomationTool
 						// When building a target for Mac or iOS, use UBT's -flushmac option to clean up the remote builder
 						bool bForceFlushMac = DeleteBuildProducts && (Target.Platform == UnrealBuildTool.UnrealTargetPlatform.Mac || Target.Platform == UnrealBuildTool.UnrealTargetPlatform.IOS);
 						LogSetProgress(InShowProgress, "Building header tool...");
-						BuildWithUBT(Target.TargetName, Target.Platform, Target.Config.ToString(), Target.UprojectPath, bForceMonolithic, bForceNonUnity, bForceDebugInfo, bForceFlushMac, true, Target.UBTArgs, bForceUnity);
+						BuildWithUBT(Target.TargetName, Target.Platform, Target.Config.ToString(), Target.UprojectPath, bForceMonolithic, bForceNonUnity, bForceDebugInfo, bForceFlushMac, !CanUseXGE(Target.Platform), Target.UBTArgs, bForceUnity);
 					}
 				}
 
@@ -1377,19 +1369,15 @@ namespace AutomationTool
 				}
 				if (XGEItems.Count > 0)
 				{
-					string XGETool = "";
-					string Args = "";
+					string XGETool = null;
+					string Args = null;
 					if (bCanUseXGE) 
 					{
 						XGETool = XGEConsoleExePath;
 						Args = "\"" + TaskFilePath + "\" /Rebuild /MaxCPUS=200";
 					}
-					else if (bCanUseParallelExecutor)
-					{
-						XGETool = ParallelExecutorExec();
-						Args = "\"" + TaskFilePath + "\"";
-					}
-					if (XGETool.Equals(""))
+
+					if (!bCanUseParallelExecutor && String.IsNullOrEmpty(XGETool))
 					{
 						throw new AutomationException("BUILD FAILED: no tool present to process XGE files");
 					}
diff --git a/Engine/Source/Programs/AutomationTool/GUBP/ElectricCommander.cs b/Engine/Source/Programs/AutomationTool/GUBP/ElectricCommander.cs
--- a/Engine/Source/Programs/AutomationTool/GUBP/ElectricCommander.cs
+++ b/Engine/Source/Programs/AutomationTool/GUBP/ElectricCommander.cs
@@ -213,9 +213,9 @@ namespace AutomationTool
 			}
 		}
 
-        public void SaveStatus(TempStorageNodeInfo TempStorageNodeInfo, bool bSaveSharedTempStorage, string JobStepIDForFailure = null)
+        public void SaveStatus(string SharedStorageDir, string NodeStorageName, bool bSaveSharedTempStorage, string JobStepIDForFailure = null)
 		{
-			string Contents = "Just a status record: " + TempStorageNodeInfo.NodeStorageName;
+			string Contents = "Just a status record: " + NodeStorageName;
 			if (!String.IsNullOrEmpty(JobStepIDForFailure) && CommandUtils.IsBuildMachine)
 			{
 				try
@@ -228,10 +228,10 @@ namespace AutomationTool
 					CommandUtils.LogWarning(LogUtils.FormatException(Ex));
 				}
 			}
-			string RecordOfSuccess = CommandUtils.CombinePaths(CommandUtils.CmdEnv.LocalRoot, "Engine", "Saved", "Logs", TempStorageNodeInfo.NodeStorageName +".log");
+			string RecordOfSuccess = CommandUtils.CombinePaths(CommandUtils.CmdEnv.LocalRoot, "Engine", "Saved", "Logs", NodeStorageName +".log");
 			CommandUtils.CreateDirectory(Path.GetDirectoryName(RecordOfSuccess));
 			CommandUtils.WriteAllText(RecordOfSuccess, Contents);
-            TempStorage.StoreToTempStorage(TempStorageNodeInfo, new List<string> { RecordOfSuccess }, !bSaveSharedTempStorage, CommandUtils.CmdEnv.LocalRoot);		
+            TempStorage.StoreToTempStorage(NodeStorageName, new List<string> { RecordOfSuccess }, SharedStorageDir, bSaveSharedTempStorage);		
 		}
 
 		public void UpdateEmailProperties(BuildNode NodeToDo, int LastSucceededCL, string FailedString, string FailCauserEMails, string EMailNote, bool SendSuccessForGreenAfterRed)
@@ -285,7 +285,8 @@ namespace AutomationTool
 			WriteJobSteps(OrderedToDo, bSkipTriggers);
 
 			// Export it as JSON too
-			WriteJson(OrderedToDo, bSkipTriggers);
+			Dictionary<BuildNode, int> FullNodeListSortKey = GetDisplayOrder(SortedNodes);
+			WriteJson(OrderedToDo, bSkipTriggers, ExplicitTrigger, FullNodeListSortKey.OrderBy(x => x.Value).Select(x => x.Key).ToList());
 		}
 
 		private void WriteProperties(IEnumerable<BuildNode> AllNodes, IEnumerable<AggregateNode> AllAggregates, List<BuildNode> OrderedToDo, List<BuildNode> SortedNodes, int TimeIndex, int TimeQuantum)
@@ -343,25 +344,29 @@ namespace AutomationTool
 				{
 					throw new AutomationException("Node {0} is sticky but has agent requirements.", NodeToDo.Name);
 				}
-				AgentGroup = "Sticky";
+				AgentGroup = "Startup";
 			}
 			return AgentGroup;
 		}
 
-		private void WriteJson(List<BuildNode> OrderedToDo, bool bSkipTriggers)
+		private void WriteJson(List<BuildNode> OrderedToDo, bool bSkipTriggers, TriggerNode ExplicitTrigger, List<BuildNode> AllNodes)
 		{
 			using(JsonWriter JsonWriter = new JsonWriter(CommandUtils.CombinePaths(CommandUtils.CmdEnv.LogFolder, "JobSteps.json")))
 			{
 				JsonWriter.WriteObjectStart();
-				JsonWriter.WriteArrayStart("JobSteps");
+				JsonWriter.WriteArrayStart("Groups");
 
 				string CurrentAgentGroup = "";
 				foreach (BuildNode NodeToDo in OrderedToDo)
 				{
-					if (!NodeToDo.IsComplete) // if something is already finished, we don't put it into EC  
+					if (!NodeToDo.IsComplete && (!(NodeToDo is TriggerNode) || NodeToDo == ExplicitTrigger)) // if something is already finished, we don't put it into EC  
 					{
 						// Write the agent group object
 						string AgentGroup = GetAgentGroupOrSticky(NodeToDo);
+						if(String.IsNullOrEmpty(AgentGroup))
+						{
+							AgentGroup = NodeToDo.Name + "_Agent";
+						}
 						if(AgentGroup != CurrentAgentGroup)
 						{
 							if(CurrentAgentGroup != "")
@@ -373,61 +378,19 @@ namespace AutomationTool
 							if(CurrentAgentGroup != "")
 							{
 								JsonWriter.WriteObjectStart();
-								JsonWriter.WriteValue("Group", AgentGroup);
-								JsonWriter.WriteValue("AgentPlatform", NodeToDo.AgentPlatform.ToString());
-								JsonWriter.WriteArrayStart("JobSteps");
+								JsonWriter.WriteValue("Name", AgentGroup);
+								JsonWriter.WriteArrayStart("Agent Types");
+								JsonWriter.WriteValue(NodeToDo.AgentPlatform.ToString());
+								JsonWriter.WriteArrayEnd();
+								JsonWriter.WriteArrayStart("Nodes");
 							}
 						}
 
 						// Add this node
 						JsonWriter.WriteObjectStart();
-						JsonWriter.WriteValue("Node", NodeToDo.Name);
-						if(CurrentAgentGroup == "")
-						{
-							JsonWriter.WriteValue("AgentPlatform", NodeToDo.AgentPlatform.ToString());
-						}
-						if(!String.IsNullOrEmpty(NodeToDo.AgentRequirements))
-						{
-							JsonWriter.WriteValue("AgentRequirements", NodeToDo.AgentRequirements);
-						}
-
-						// Write all the dependencies
-						StringBuilder DependsOnList = new StringBuilder();
-						foreach (BuildNode Dep in FindDirectOrderDependencies(NodeToDo))
-						{
-							if (!Dep.IsComplete && OrderedToDo.Contains(Dep)) // if something is already finished, we don't put it into EC
-							{
-								if (OrderedToDo.IndexOf(Dep) > OrderedToDo.IndexOf(NodeToDo))
-								{
-									throw new AutomationException("Topological sort error, node {0} has a dependency of {1} which sorted after it.", NodeToDo.Name, Dep.Name);
-								}
-								if(DependsOnList.Length > 0)
-								{
-									DependsOnList.Append(";");
-								}
-								string DepAgentGroup = GetAgentGroupOrSticky(Dep);
-								if(DepAgentGroup != "")
-								{
-									DependsOnList.AppendFormat("{0}/", DepAgentGroup);
-								}
-								DependsOnList.Append(Dep.Name);
-							}
-						}
-						JsonWriter.WriteValue("DependsOn", DependsOnList.ToString());
-
-						// Add any trigger-specific settings
-						TriggerNode TriggerNodeToDo = NodeToDo as TriggerNode;
-						if (TriggerNodeToDo != null && !TriggerNodeToDo.IsTriggered)
-						{
-							JsonWriter.WriteValue("TriggerState", TriggerNodeToDo.StateName);
-							JsonWriter.WriteValue("TriggerActionText", TriggerNodeToDo.ActionText);
-							JsonWriter.WriteValue("TriggerDescriptionText", TriggerNodeToDo.DescriptionText);
-							if (NodeToDo.RecipientsForFailureEmails.Length > 0)
-							{
-								JsonWriter.WriteValue("TriggerEmailRecipients", String.Join(" ", NodeToDo.RecipientsForFailureEmails));
-							}
-						}
-
+						JsonWriter.WriteValue("Name", NodeToDo.Name);
+						JsonWriter.WriteValue("DependsOn", GetDirectDependencyList(NodeToDo, OrderedToDo));
+						JsonWriter.WriteValue("CopyToSharedStorage", NodeToDo.CopyToSharedStorage);
 						JsonWriter.WriteObjectEnd();
 					}
 				}
@@ -437,10 +400,81 @@ namespace AutomationTool
 					JsonWriter.WriteObjectEnd();
 				}
 				JsonWriter.WriteArrayEnd();
+
+				JsonWriter.WriteArrayStart("Triggers");
+				foreach (TriggerNode NodeToDo in OrderedToDo.OfType<TriggerNode>())
+				{
+					if(NodeToDo != ExplicitTrigger)
+					{
+						JsonWriter.WriteObjectStart();
+						JsonWriter.WriteValue("Name", NodeToDo.Name);
+						JsonWriter.WriteValue("DirectDependencies", GetDirectDependencyList(NodeToDo, OrderedToDo));
+						JsonWriter.WriteValue("AllDependencies", GetCompleteDependencyList(NodeToDo, OrderedToDo));
+						JsonWriter.WriteValue("ActionText", NodeToDo.ActionText);
+						JsonWriter.WriteValue("DescriptionText", NodeToDo.DescriptionText);
+						if (NodeToDo.RecipientsForFailureEmails.Length > 0)
+						{
+							JsonWriter.WriteValue("Notify", String.Join(" ", NodeToDo.RecipientsForFailureEmails));
+						}
+						JsonWriter.WriteObjectEnd();
+					}
+				}
+				JsonWriter.WriteArrayEnd();
+
+				JsonWriter.WriteArrayStart("NodesInGraph");
+				foreach(BuildNode NodeToDo in AllNodes)
+				{
+					if(NodeToDo.ControllingTriggers.Length == 0 && !(NodeToDo is TriggerNode))
+					{
+						JsonWriter.WriteValue(NodeToDo.Name);
+					}
+				}
+				JsonWriter.WriteArrayEnd();
+
 				JsonWriter.WriteObjectEnd();
 			}
 		}
 
+		private string GetDirectDependencyList(BuildNode NodeToDo, List<BuildNode> OrderedToDo)
+		{
+			HashSet<BuildNode> DirectDependencies = FindDirectOrderDependencies(NodeToDo);
+
+			StringBuilder DependsOnList = new StringBuilder();
+			foreach(BuildNode Dep in OrderedToDo.Where(x => DirectDependencies.Contains(x)))
+			{
+				if (!Dep.IsComplete) // if something is already finished, we don't put it into EC
+				{
+					if (OrderedToDo.IndexOf(Dep) > OrderedToDo.IndexOf(NodeToDo))
+					{
+						throw new AutomationException("Topological sort error, node {0} has a dependency of {1} which sorted after it.", NodeToDo.Name, Dep.Name);
+					}
+					if(DependsOnList.Length > 0)
+					{
+						DependsOnList.Append(";");
+					}
+					DependsOnList.Append(Dep.Name);
+				}
+			}
+			return DependsOnList.ToString();
+		}
+
+		private static string GetCompleteDependencyList(BuildNode NodeToDo, List<BuildNode> OrderedToDo)
+		{
+			List<BuildNode> Nodes = new List<BuildNode>(NodeToDo.OrderDependencies);
+			for(int Idx = 0; Idx < Nodes.Count; Idx++)
+			{
+				BuildNode Node = Nodes[Idx];
+				foreach(BuildNode DependencyNode in Node.OrderDependencies)
+				{
+					if(!Nodes.Contains(DependencyNode))
+					{
+						Nodes.Add(DependencyNode);
+					}
+				}
+			}
+			return String.Join(";", Nodes.Where(x => OrderedToDo.Contains(x)).OrderBy(x => OrderedToDo.IndexOf(x)).Select(x => x.Name));
+		}
+
 		private void WriteJobSteps(List<BuildNode> OrderedToDo, bool bSkipTriggers)
 		{
 			BuildNode LastSticky = null;
@@ -629,6 +663,10 @@ namespace AutomationTool
 							JobStep MainStep = new JobStep(NodeParentPath, NodeToDo.Name, Procedure, DoParallel, PreCondition, RunCondition, (NodeToDo.IsSticky && NodeToDo == LastSticky) ? JobStepReleaseMode.Release : JobStepReleaseMode.Keep);
 							MainStep.ActualParameters.Add("NodeName", NodeToDo.Name);
 							MainStep.ActualParameters.Add("Sticky", NodeToDo.IsSticky ? "1" : "0");
+							if(!NodeToDo.CopyToSharedStorage)
+							{
+								MainStep.ActualParameters.Add("ExtraOptions", "-NoCopyToSharedStorage");
+							}
 							if (NodeToDo.AgentSharingGroup != "")
 							{
 								MainStep.ActualParameters.Add("AgentSharingGroup", NodeToDo.AgentSharingGroup);
diff --git a/Engine/Source/Programs/AutomationTool/GUBP/GUBP.Automation.csproj b/Engine/Source/Programs/AutomationTool/GUBP/GUBP.Automation.csproj
--- a/Engine/Source/Programs/AutomationTool/GUBP/GUBP.Automation.csproj
+++ b/Engine/Source/Programs/AutomationTool/GUBP/GUBP.Automation.csproj
@@ -25,7 +25,7 @@
   </PropertyGroup>
   <PropertyGroup Condition=" '$(Configuration)|$(Platform)' == 'Development|AnyCPU' ">
     <DebugType>pdbonly</DebugType>
-    <Optimize>true</Optimize>
+    <Optimize>false</Optimize>
     <OutputPath>..\..\..\..\Binaries\DotNET\AutomationScripts\</OutputPath>
     <DefineConstants>TRACE</DefineConstants>
     <ErrorReport>prompt</ErrorReport>
diff --git a/Engine/Source/Programs/AutomationTool/GUBP/GUBP.cs b/Engine/Source/Programs/AutomationTool/GUBP/GUBP.cs
--- a/Engine/Source/Programs/AutomationTool/GUBP/GUBP.cs
+++ b/Engine/Source/Programs/AutomationTool/GUBP/GUBP.cs
@@ -99,50 +99,14 @@ public partial class GUBP : BuildCommand
         string BranchName = P4Enabled ? P4Env.BuildRootP4 : ParseParamValue("BranchName", "");
         GUBPBranchHacker.BranchOptions BranchOptions = GetBranchOptions(BranchName);
 
-        // encapsulate the logic and temp vars to determine the job info.
-        // Ensures the temp vars don't escape the scope to be used later.
-        Func<JobInfo> GetJobInfo = () =>
-        {
-            int PreflightShelveCL = 0;
-            int PreflightUID = 0;
-            string PreflightShelveCLString = GetEnvVar("uebp_PreflightShelveCL");
-            // We must be passed a valid shelve and be a build machine (or be running a preflight test) to be a valid preflight.
-            if (!string.IsNullOrEmpty(PreflightShelveCLString) && (IsBuildMachine || ParseParam("PreflightTest")))
-            {
-                Log("**** Preflight shelve {0}", PreflightShelveCLString);
-                PreflightShelveCL = int.Parse(PreflightShelveCLString);
-                if (PreflightShelveCL < 2000000)
-                {
-                    throw new AutomationException("{0} does not look like a CL", PreflightShelveCL);
-                }
-                PreflightUID = ParseParamInt("PreflightUID", 0);
-            }
-            return new JobInfo(
-                BranchName,
-                P4Enabled ? P4Env.BuildRootEscaped : "NoP4",
-                BranchOptions.RootNameForTempStorage ?? "UE4",
-                P4Enabled ? ParseParamInt("CL", P4Env.Changelist) : 0,
-                PreflightShelveCL, 
-                PreflightUID);
-        };
-
-        var JobInfo = GetJobInfo();
-
-        bool LocalOnly = !P4Enabled || bFakeEC;
-
-        if (bSaveSharedTempStorage)
-        {
-            if (!TempStorage.IsSharedTempStorageAvailable(true))
-            {
-                throw new AutomationException("Request to save to shared temp storage, but shared temp storage is unavailable or does not exist.");
-            }
-        }
-        else if (!LocalOnly && !TempStorage.IsSharedTempStorageAvailable(false))
-        {
-            LogWarning("Looks like we want to use shared temp storage, but since we don't have it, we won't use it.");
-            LocalOnly = true;
-        }
-
+		string BranchNameForTempStorage = P4Enabled ? P4Env.BuildRootEscaped : "NoP4";
+		string RootNameForTempStorage = BranchOptions.RootNameForTempStorage ?? "UE4";
+		string PreflightSuffix = GetPreflightSuffix();
+		int ChangeNumber = P4Enabled ? ParseParamInt("CL", P4Env.Changelist) : 0;
+		bool bIsPreflight = PreflightSuffix != null;
+		string BuildFolder = string.Format("{0}{1}", ChangeNumber, PreflightSuffix);
+		string SharedStorageDir = GetSharedStorageDir(bSaveSharedTempStorage, bFakeEC, BranchNameForTempStorage, RootNameForTempStorage, BuildFolder);
+		
         bool CommanderSetup = ParseParam("CommanderJobSetupOnly");
 
 		string ExplicitTriggerName = "";
@@ -161,7 +125,8 @@ public partial class GUBP : BuildCommand
 		if(ScriptName == null)
 		{
 			// Automatically configure it from whatever is in the branch
-			AddNodesForBranch(HostPlatforms, JobInfo, BranchOptions, out BuildNodeDefinitions, out AggregateNodeDefinitions, ref TimeQuantum);
+			string BuildName = "CL-" + (P4Enabled? P4Env.ChangelistString : "Unknown") + (PreflightSuffix ?? "");
+			AddNodesForBranch(HostPlatforms, BranchName, new JobInfo(BuildName, PreflightSuffix != null), BranchOptions, out BuildNodeDefinitions, out AggregateNodeDefinitions, ref TimeQuantum);
 
 			// Check if we're exporting the graph as a starting point for manually configuring it
 			string OutputScriptName = ParseParamValue("ExportScript");
@@ -180,8 +145,9 @@ public partial class GUBP : BuildCommand
 		}
 
 		BuildGraph Graph = new BuildGraph(AggregateNodeDefinitions, BuildNodeDefinitions);
-		FindCompletionState(Graph.BuildNodes, JobInfo, LocalOnly);
+		FindCompletionState(Graph.BuildNodes, SharedStorageDir);
 		ComputeDependentFrequencies(Graph.BuildNodes);
+		ValidateTriggerDependencies(Graph.BuildNodes);
 
 		int TimeIndex = ParseParamInt("UserTimeIndex", -1);
 		if (TimeIndex == -1)
@@ -197,11 +163,24 @@ public partial class GUBP : BuildCommand
 		// Find the list of nodes which we're actually going to execute
 		HashSet<BuildNode> NodesToDo = ParseNodesToDo(Graph.BuildNodes, Graph.AggregateNodes);
 		CullNodesForTimeIndex(NodesToDo, TimeIndex);
-        if (JobInfo.IsPreflight)
+        if (bIsPreflight)
         {
             CullNodesForPreflight(NodesToDo);
         }
 
+		// Mark all the nodes whose build products need to be copied to temp storage. Nodes which appear in the same agent sharing group do not need to be 
+		// copied; we don't clean between builds on the same agent.
+		foreach(BuildNode NodeToDo in NodesToDo)
+		{
+			foreach(BuildNode InputDependency in NodeToDo.InputDependencies)
+			{
+				if(String.IsNullOrEmpty(InputDependency.AgentSharingGroup) || InputDependency.AgentSharingGroup != NodeToDo.AgentSharingGroup || NodeToDo.IsSticky != InputDependency.IsSticky)
+				{
+					InputDependency.CopyToSharedStorage = true;
+				}
+			}
+		}
+
 		// From each of those nodes, remove any order dependencies that we're not going to run. There are no circumstances in which they should be 
 		// considered dependencies from this point on.
 		foreach(BuildNode NodeToDo in NodesToDo)
@@ -260,7 +239,7 @@ public partial class GUBP : BuildCommand
 				throw new AutomationException("Couldn't find node {0}", ShowHistoryParam);
 			}
 
-            NodeHistory History = FindNodeHistory(Node, JobInfo);
+            NodeHistory History = FindNodeHistory(Node.Name, ChangeNumber, SharedStorageDir);
 			if(History == null)
 			{
 				throw new AutomationException("Couldn't get history for {0}", ShowHistoryParam);
@@ -273,7 +252,7 @@ public partial class GUBP : BuildCommand
 		{
 			if(CommanderSetup)
 			{
-				return DoCommanderSetup(EC, Graph.BuildNodes, Graph.AggregateNodes, OrderedToDo, TimeIndex, TimeQuantum, bSkipTriggers, bFake, bFakeEC, ExplicitTrigger, UnfinishedTriggers, JobInfo.IsPreflight);
+				return DoCommanderSetup(EC, Graph.BuildNodes, Graph.AggregateNodes, OrderedToDo, TimeIndex, TimeQuantum, bSkipTriggers, bFake, bFakeEC, ExplicitTrigger, UnfinishedTriggers, bIsPreflight);
 			}
 			else if(ParseParam("ListOnly"))
 			{
@@ -281,11 +260,63 @@ public partial class GUBP : BuildCommand
 			}
 			else
 			{
-				return ExecuteNodes(EC, OrderedToDo, bFake, bFakeEC, bSaveSharedTempStorage, JobInfo);
+				return ExecuteNodes(EC, OrderedToDo, bFake, bFakeEC, SharedStorageDir, bSaveSharedTempStorage, ChangeNumber);
 			}
 		}
 	}
 
+	/// <summary>
+	/// Get a suffix used to identify preflight builds.
+	/// </summary>
+	/// <returns>Suffix to uniquely identify this preflight. Null if the current build is not a preflight</returns>
+	string GetPreflightSuffix()
+	{
+        int PreflightShelveCL = 0;
+        int PreflightUID = 0;
+        string PreflightShelveCLString = GetEnvVar("uebp_PreflightShelveCL");
+        // We must be passed a valid shelve and be a build machine (or be running a preflight test) to be a valid preflight.
+        if (!string.IsNullOrEmpty(PreflightShelveCLString) && (IsBuildMachine || ParseParam("PreflightTest")))
+        {
+            Log("**** Preflight shelve {0}", PreflightShelveCLString);
+            PreflightShelveCL = int.Parse(PreflightShelveCLString);
+            if (PreflightShelveCL < 2000000)
+            {
+                throw new AutomationException("{0} does not look like a CL", PreflightShelveCL);
+            }
+            PreflightUID = ParseParamInt("PreflightUID", 0);
+        }
+		return (PreflightShelveCL > 0)? string.Format("-PF-{0}-{1}", PreflightShelveCL, PreflightUID) : null;
+	}
+
+	/// <summary>
+	/// Determine the directory shared build products should be saved to.
+	/// </summary>
+	string GetSharedStorageDir(bool bSaveSharedTempStorage, bool bFakeEC, string BranchNameForTempStorage, string RootNameForTempStorage, string BuildFolder)
+	{
+        bool LocalOnly = !P4Enabled || bFakeEC;
+
+        if (bSaveSharedTempStorage)
+        {
+            if (!TempStorage.IsSharedTempStorageAvailable(true))
+            {
+                throw new AutomationException("Request to save to shared temp storage, but shared temp storage is unavailable or does not exist.");
+            }
+        }
+        else if (!LocalOnly && !TempStorage.IsSharedTempStorageAvailable(false))
+        {
+            LogWarning("Looks like we want to use shared temp storage, but since we don't have it, we won't use it.");
+            LocalOnly = true;
+        }
+
+		string SharedStorageDir = null;
+		if(!LocalOnly)
+		{
+			string RootSharedStorageDir = TempStorage.ResolveSharedTempStorageDirectory(RootNameForTempStorage);
+			SharedStorageDir = CommandUtils.CombinePaths(RootSharedStorageDir, BranchNameForTempStorage, BuildFolder);
+		}
+		return SharedStorageDir;
+	}
+
 	/// <summary>
 	/// Get a string describing a time interval, in the format "XhYm".
 	/// </summary>
@@ -318,30 +349,45 @@ public partial class GUBP : BuildCommand
 		}
 	}
 
-    static void FindCompletionState(IEnumerable<BuildNode> NodesToDo, JobInfo JobInfo, bool LocalOnly)
+	/// <summary>
+	/// Check whether any nodes reference triggers through input dependencies. We're deprecating the treatment of trigger nodes as
+	/// build nodes, because they don't do anything useful when they run. The new EC framework bypasses the execution of trigger nodes
+	/// entirely, and just sends notifications once all the dependencies are met.
+	/// </summary>
+	/// <param name="Nodes">List of nodes to check trigger dependencies for</param>
+	static void ValidateTriggerDependencies(IReadOnlyList<BuildNode> Nodes)
 	{
-		foreach(BuildNode NodeToDo in NodesToDo)
+		foreach(BuildNode NodeToDo in Nodes)
 		{
-            // Construct the full temp storage node info
-            var TempStorageNodeInfo = new TempStorageNodeInfo(JobInfo, NodeToDo.Name);
-
-			if (LocalOnly)
+			foreach(TriggerNode TriggerNodeToDo in NodeToDo.InputDependencies.OfType<TriggerNode>())
 			{
-                NodeToDo.IsComplete = TempStorage.LocalTempStorageManifestExists(TempStorageNodeInfo, bQuiet: true);
+				LogWarning("'{0}' has an input dependency on '{1}'. Treating trigger nodes as build nodes is deprecated.", NodeToDo.Name, TriggerNodeToDo.Name);
 			}
-			else
+		}
+		foreach(TriggerNode TriggerNodeToDo in Nodes.OfType<TriggerNode>())
+		{
+			foreach(BuildNode NodeToDo in TriggerNodeToDo.InputDependencies)
 			{
-                NodeToDo.IsComplete = TempStorage.TempStorageExists(TempStorageNodeInfo, bLocalOnly: false, bQuiet: true);
+				LogWarning("'{0}' has an input dependency on '{1}'. Treating trigger nodes as build nodes is deprecated.", TriggerNodeToDo.Name, NodeToDo.Name);
 			}
+		}
+	}
+
+    static void FindCompletionState(IEnumerable<BuildNode> NodesToDo, string SharedStorageDir)
+	{
+		foreach(BuildNode NodeToDo in NodesToDo)
+		{
+            // Construct the full temp storage node info
+			NodeToDo.IsComplete = TempStorage.TempStorageExists(NodeToDo.Name, SharedStorageDir, bQuiet: true);
 
             LogVerbose("** {0}", NodeToDo.Name);
 			if (!NodeToDo.IsComplete)
 			{
-                LogVerbose("***** GUBP Trigger Node was already triggered {0} -> {1} : {2}", NodeToDo.Name, JobInfo.RootNameForTempStorage, TempStorageNodeInfo.GetRelativeDirectory());
+                LogVerbose("***** GUBP Trigger Node was already triggered {0}", NodeToDo.Name);
 			}
 			else
 			{
-                LogVerbose("***** GUBP Trigger Node was NOT yet triggered {0} -> {1} : {2}", NodeToDo.Name, JobInfo.RootNameForTempStorage, TempStorageNodeInfo.GetRelativeDirectory());
+                LogVerbose("***** GUBP Trigger Node was NOT yet triggered {0}", NodeToDo.Name);
 			}
 		}
     }
@@ -501,6 +547,10 @@ public partial class GUBP : BuildCommand
 			{
 				Builder.Append(" - (Sticky)");
 			}
+			if(!NodeToDo.CopyToSharedStorage)
+			{
+				Builder.Append(" - (Not Archived)");
+			}
 
             string Agent = NodeToDo.AgentRequirements;
 			if(ParseParamValue("AgentOverride") != "" && !NodeToDo.Name.Contains("Mac"))
@@ -781,27 +831,49 @@ public partial class GUBP : BuildCommand
         return OrderedToDo;
     }
 
-    static NodeHistory FindNodeHistory(BuildNode NodeToDo, JobInfo JobInfo)
+    static NodeHistory FindNodeHistory(string NodeName, int CurrentChangeNumber, string SharedStorageDir)
     {
 		NodeHistory History = null;
-
-        // Don't get node history on nodes that don't have a valid CL.
-        if (!(NodeToDo is TriggerNode) && JobInfo.Changelist > 0)
+        if (CurrentChangeNumber > 0)
         {
-            History = new NodeHistory
-            {
-                AllStarted = TempStorage.FindMatchingSharedTempStorageNodeCLs(new TempStorageNodeInfo(JobInfo, NodeToDo.Name + StartedTempStorageSuffix)),
-                AllSucceeded = TempStorage.FindMatchingSharedTempStorageNodeCLs(new TempStorageNodeInfo(JobInfo, NodeToDo.Name + SucceededTempStorageSuffix)),
-                AllFailed = TempStorage.FindMatchingSharedTempStorageNodeCLs(new TempStorageNodeInfo(JobInfo, NodeToDo.Name + FailedTempStorageSuffix))
-            };
+			History = new NodeHistory();
+
+			// Assume that the shared storage directory is a pattern containing the changelist for the current build, and look for directories with a similar pattern.
+			string SharedStorageLastName = Path.GetFileName(SharedStorageDir);
+			string ChangeString = CurrentChangeNumber.ToString();
+			int ChangeIndex = SharedStorageLastName.LastIndexOf(ChangeString);
+			if(ChangeIndex != -1)
+			{
+				string SharedStoragePattern = SharedStorageLastName.Substring(0, ChangeIndex) + "*" + SharedStorageLastName.Substring(ChangeIndex + ChangeString.Length);
+				foreach(DirectoryInfo Directory in new DirectoryInfo(SharedStorageDir).Parent.EnumerateDirectories(SharedStoragePattern))
+				{
+					int OtherCL;
+					if(int.TryParse(Directory.Name.Substring(ChangeIndex, Directory.Name.Length - (SharedStorageLastName.Length - ChangeIndex - ChangeString.Length)), out OtherCL))
+					{
+						if(TempStorage.SharedTempStorageManifestExists(Directory.FullName, NodeName + StartedTempStorageSuffix, true))
+						{
+							History.AllStarted.Add(OtherCL);
+						}
+						if(TempStorage.SharedTempStorageManifestExists(Directory.FullName, NodeName + SucceededTempStorageSuffix, true))
+						{
+							History.AllSucceeded.Add(OtherCL);
+						}
+						if(TempStorage.SharedTempStorageManifestExists(Directory.FullName, NodeName + FailedTempStorageSuffix, true))
+						{
+							History.AllFailed.Add(OtherCL);
+						}
+					}
+				}
+			}
 
+			// Filter everything relative to the current CL
             if (History.AllFailed.Count > 0)
             {
-                History.LastFailed = History.AllFailed.LastOrDefault(x => x < JobInfo.Changelist);
+                History.LastFailed = History.AllFailed.LastOrDefault(x => x < CurrentChangeNumber);
             }
             if (History.AllSucceeded.Count > 0)
             {
-                History.LastSucceeded = History.AllSucceeded.LastOrDefault(x => x < JobInfo.Changelist);
+                History.LastSucceeded = History.AllSucceeded.LastOrDefault(x => x < CurrentChangeNumber);
 
                 foreach (int Failed in History.AllFailed)
                 {
@@ -821,10 +893,9 @@ public partial class GUBP : BuildCommand
                 }
 			}
         }
-
 		return History;
     }
-    void GetFailureEmails(ElectricCommander EC, BuildNode NodeToDo, NodeHistory History, JobInfo JobInfo, bool OnlyLateUpdates = false)
+    void GetFailureEmails(ElectricCommander EC, BuildNode NodeToDo, NodeHistory History, string SharedStorageDir, bool OnlyLateUpdates = false)
 	{
         string FailCauserEMails = "";
         string EMailNote = "";
@@ -839,7 +910,7 @@ public partial class GUBP : BuildCommand
                 {
                     if (OnlyLateUpdates)
                     {
-                        LastNonDuplicateFail = FindLastNonDuplicateFail(NodeToDo, History, JobInfo);
+                        LastNonDuplicateFail = FindLastNonDuplicateFail(NodeToDo, History, SharedStorageDir);
                         if (LastNonDuplicateFail < P4Env.Changelist)
                         {
 							Log("*** Red-after-red spam reduction, changed CL {0} to CL {1} because the errors didn't change.", P4Env.Changelist, LastNonDuplicateFail);
@@ -929,11 +1000,11 @@ public partial class GUBP : BuildCommand
         return true;
     }
 
-    int FindLastNonDuplicateFail(BuildNode NodeToDo, NodeHistory History, JobInfo JobInfo)
+    int FindLastNonDuplicateFail(BuildNode NodeToDo, NodeHistory History, string SharedStorageDir)
     {
         int Result = P4Env.Changelist;
 
-        var TempStorageNodeInfo = new TempStorageNodeInfo(JobInfo, NodeToDo.Name + FailedTempStorageSuffix);
+        var TempStorageNodeInfo = NodeToDo.Name + FailedTempStorageSuffix;
 
         List<int> BackwardsFails = new List<int>(History.AllFailed);
         BackwardsFails.Add(P4Env.Changelist);
@@ -951,13 +1022,13 @@ public partial class GUBP : BuildCommand
                 break;
             }
             // Find any local temp storage manifest for this changelist and delete it.
-            var ThisTempStorageNodeInfo = new TempStorageNodeInfo(JobInfo.CreateWithNewChangelist(CL), TempStorageNodeInfo.NodeStorageName);
+            var ThisTempStorageNodeInfo = TempStorageNodeInfo;
             TempStorage.DeleteLocalTempStorageManifest(ThisTempStorageNodeInfo, true); // these all clash locally, which is fine we just retrieve them from shared
 
             List<string> Files = null;
             try
             {
-                Files = TempStorage.RetrieveFromTempStorage(ThisTempStorageNodeInfo, CmdEnv.LocalRoot); // this will fail on our CL if we didn't fail or we are just setting up the branch
+                Files = TempStorage.RetrieveFromTempStorage(SharedStorageDir, ThisTempStorageNodeInfo); // this will fail on our CL if we didn't fail or we are just setting up the branch
             }
             catch (Exception)
             {
@@ -1049,6 +1120,7 @@ public partial class GUBP : BuildCommand
 		foreach(BuildNode NodeToDo in NodesToDo)
 		{
 			RecursiveNodesToDo.UnionWith(NodeToDo.InputDependencies);
+			RecursiveNodesToDo.UnionWith(NodeToDo.ControllingTriggers);
 		}
 
 		return RecursiveNodesToDo;
@@ -1291,7 +1363,7 @@ public partial class GUBP : BuildCommand
 		return ExitCode.Success;
 	}
 
-    ExitCode ExecuteNodes(ElectricCommander EC, List<BuildNode> OrderedToDo, bool bFake, bool bFakeEC, bool bSaveSharedTempStorage, JobInfo JobInfo)
+    ExitCode ExecuteNodes(ElectricCommander EC, List<BuildNode> OrderedToDo, bool bFake, bool bFakeEC, string SharedStorageDir, bool bSaveSharedTempStorage, int ChangeNumber)
 	{
         Dictionary<string, BuildNode> BuildProductToNodeMap = new Dictionary<string, BuildNode>();
 		foreach (BuildNode NodeToDo in OrderedToDo)
@@ -1301,13 +1373,11 @@ public partial class GUBP : BuildCommand
                 throw new AutomationException("topological sort error");
             }
 
-            var TempStorageNodeInfo = new TempStorageNodeInfo(JobInfo, NodeToDo.Name);
-            
             if (NodeToDo.IsComplete)
             {
 				// Just fetch the build products from temp storage
-				Log("***** Retrieving GUBP Node {0} -> {1} : {2}", NodeToDo.Name, JobInfo.RootNameForTempStorage, TempStorageNodeInfo.GetRelativeDirectory());
-				NodeToDo.RetrieveBuildProducts(TempStorageNodeInfo);
+				Log("***** Retrieving GUBP Node {0}", NodeToDo.Name);
+				NodeToDo.RetrieveBuildProducts(SharedStorageDir);
             }
             else
             {
@@ -1318,7 +1388,7 @@ public partial class GUBP : BuildCommand
                 // Record that the node has started. We save our status to a new temp storage location specifically named with a suffix so we can find it later.
                 if (SaveSuccessRecords) 
                 {
-                    EC.SaveStatus(new TempStorageNodeInfo(JobInfo, NodeToDo.Name + StartedTempStorageSuffix), bSaveSharedTempStorage);
+                    EC.SaveStatus(SharedStorageDir, NodeToDo.Name + StartedTempStorageSuffix, bSaveSharedTempStorage);
                 }
 
 				// Execute the node
@@ -1329,7 +1399,7 @@ public partial class GUBP : BuildCommand
 				// Archive the build products if the node succeeded
 				if(bResult)
 				{
-					NodeToDo.ArchiveBuildProducts(TempStorageNodeInfo, !bSaveSharedTempStorage);
+					NodeToDo.ArchiveBuildProducts(SharedStorageDir, (bSaveSharedTempStorage && !ParseParam("NoCopyToSharedStorage")));
 				}
 
 				// Record that the node has finished
@@ -1338,11 +1408,14 @@ public partial class GUBP : BuildCommand
                 {
 					using(TelemetryStopwatch UpdateNodeHistoryStopwatch = new TelemetryStopwatch("UpdateNodeHistory"))
 					{
-                        History = FindNodeHistory(NodeToDo, JobInfo);
+						if(!(NodeToDo is TriggerNode))
+						{
+							History = FindNodeHistory(NodeToDo.Name, ChangeNumber, SharedStorageDir);
+						}
 					}
 					using(TelemetryStopwatch SaveNodeStatusStopwatch = new TelemetryStopwatch("SaveNodeStatus"))
 					{
-                        EC.SaveStatus(new TempStorageNodeInfo(JobInfo, NodeToDo.Name + (bResult? SucceededTempStorageSuffix : FailedTempStorageSuffix)), bSaveSharedTempStorage, ParseParamValue("MyJobStepId"));
+                        EC.SaveStatus(SharedStorageDir, NodeToDo.Name + (bResult? SucceededTempStorageSuffix : FailedTempStorageSuffix), bSaveSharedTempStorage, ParseParamValue("MyJobStepId"));
 					}
 					using(TelemetryStopwatch UpdateECPropsStopwatch = new TelemetryStopwatch("UpdateECProps"))
 					{
@@ -1352,7 +1425,7 @@ public partial class GUBP : BuildCommand
 					{
 						using(TelemetryStopwatch GetFailEmailsStopwatch = new TelemetryStopwatch("GetFailEmails"))
 						{
-                            GetFailureEmails(EC, NodeToDo, History, JobInfo);
+                            GetFailureEmails(EC, NodeToDo, History, SharedStorageDir);
 						}
 					}
 					EC.UpdateECBuildTime(NodeToDo, ElapsedTime.TotalSeconds);
diff --git a/Engine/Source/Programs/AutomationTool/GUBP/JobInfo.cs b/Engine/Source/Programs/AutomationTool/GUBP/JobInfo.cs
--- a/Engine/Source/Programs/AutomationTool/GUBP/JobInfo.cs
+++ b/Engine/Source/Programs/AutomationTool/GUBP/JobInfo.cs
@@ -19,71 +19,15 @@ namespace AutomationTool
     /// </remarks>
     public class JobInfo
     {
-        /// <summary>
-        /// BranchName for this job. P4Env.BuildRootP4 if P4Enabled. Otherwise, uses the -BranchName= parameter. Can be empty string.
-        /// </summary>
-        public string BranchName { get; private set; }
-
-        /// <summary>
-        /// The root name beneath <see cref="CommandUtils.RootBuildStorageDirectory"/> where all temp storage will placed for this job.
-        /// </summary>
-        public string RootNameForTempStorage { get; private set; }
-
-        /// <summary>
-        /// Used as the prefix for temp storage for the job.
-        /// 
-        /// If P4 is enabled, this is P4Env.BuildRootEscaped, ie the branch name safe to create a file or folder from. If P4 is not enabled, it is NoP4, 
-        /// but we never write to the shared temp storage location in those cases.
-        /// For testing, arbitrary values can be used.
-        /// </summary>
-        public string BranchNameForTempStorage { get; private set; }
-
-        /// <summary>
-        /// The changelist associated with this job (ie, the code CL that was sync'ed to run this job.
-        /// The -CL parameter can be used to override the P4 environment for testing.
-        /// If P4 is not enabled, this value will be 0.
-        /// </summary>
-        public int Changelist { get; private set; }
-
-        /// <summary>
-        /// If this is a preflight vuild, this is the changelist that contains the shelve being preflighted. 
-        /// Otherwise, will be 0.
-        /// </summary>
-        public int PreflightShelveChangelist { get; private set; }
+		/// <summary>
+		/// Unique name used to identify the build within the branch. Useful for creating unique folders for build artifacts.
+		/// </summary>
+		public string BuildName { get; private set; }
 
         /// <summary>
-        /// If this is a preflight build, contains a unique ID to identify the preflight job.
-        /// Used to disambiguate multiple preflight runs from the same changelist and shelve, since the
-        /// code in the shelve can change from run to run.
+        /// Will be true if this job is a preflight run.
         /// </summary>
-        public int PreflightUID { get; private set; }
-
-        /// <summary>
-        /// Will be true if this job is a preflight run. Shortcut for checking if the shelve is > 0
-        /// </summary>
-        public bool IsPreflight { get { return PreflightShelveChangelist > 0; } }
-
-        /// <summary>
-        /// If this is a preflight run, returns a suffix used to describe the preflight shelve and UID. Used
-        /// to uniquely identify a temp storage block for a particular preflight run.
-        /// </summary>
-        /// <returns></returns>
-        public string GetPreflightSuffix()
-        {
-            return IsPreflight ? string.Format("-PF-{0}-{1}", PreflightShelveChangelist, PreflightUID) : "";
-        }
-
-        /// <summary>
-        /// Quick method to create a new JobInfo using all the same attributes, but a different CL.
-        /// 
-        /// Used by some GUBP code to search for and use previous temp storage blocks when checking for node history.
-        /// </summary>
-        /// <param name="NewCL">The CL to use instead of the current Job's CL.</param>
-        /// <returns>A new JobInfo instance with all the same attributes as the current instance, but a different CL.</returns>
-        public JobInfo CreateWithNewChangelist(int NewCL)
-        {
-            return new JobInfo(BranchName, RootNameForTempStorage, BranchNameForTempStorage, NewCL, PreflightShelveChangelist, PreflightUID);
-        }
+        public bool IsPreflight { get; private set; }
 
         /// <summary>
         /// Public ctor. This is an immutable class, so all attributes must be set a construction time.
@@ -94,14 +38,10 @@ namespace AutomationTool
         /// <param name="Changelist">See similarly named attribute.</param>
         /// <param name="PreflightShelveChangelist">See similarly named attribute.</param>
         /// <param name="PreflightUID">See similarly named attribute.</param>
-        public JobInfo(string BranchName, string BranchNameForTempStorage, string RootNameForTempStorage, int Changelist, int PreflightShelveChangelist, int PreflightUID)
+        public JobInfo(string BuildName, bool IsPreflight)
         {
-            this.BranchName = BranchName;
-            this.BranchNameForTempStorage = BranchNameForTempStorage;
-            this.RootNameForTempStorage = RootNameForTempStorage;
-            this.Changelist = Changelist;
-            this.PreflightShelveChangelist = PreflightShelveChangelist;
-            this.PreflightUID = PreflightUID;
+			this.BuildName = BuildName;
+            this.IsPreflight = IsPreflight;
         }
     }
 }
diff --git a/Engine/Source/Programs/AutomationTool/GUBP/LegacyBranchSetup.cs b/Engine/Source/Programs/AutomationTool/GUBP/LegacyBranchSetup.cs
--- a/Engine/Source/Programs/AutomationTool/GUBP/LegacyBranchSetup.cs
+++ b/Engine/Source/Programs/AutomationTool/GUBP/LegacyBranchSetup.cs
@@ -15,7 +15,7 @@ partial class GUBP
 	public class GUBPBranchConfig
 	{
 		public List<UnrealTargetPlatform> HostPlatforms;
-        public string BranchName { get { return JobInfo.BranchName; } }
+        public string BranchName;
 		public BranchInfo Branch;
 		public GUBPBranchHacker.BranchOptions BranchOptions;
 		public bool bForceIncrementalCompile;
@@ -24,10 +24,11 @@ partial class GUBP
 		public Dictionary<string, GUBPNode> GUBPNodes = new Dictionary<string, GUBPNode>();
 		public Dictionary<string, GUBPAggregateNode> GUBPAggregates = new Dictionary<string, GUBPAggregateNode>();
 
-        public GUBPBranchConfig(IEnumerable<UnrealTargetPlatform> InHostPlatforms, BranchInfo InBranch, GUBPBranchHacker.BranchOptions InBranchOptions, bool bInForceIncrementalCompile, JobInfo JobInfo)
+        public GUBPBranchConfig(IEnumerable<UnrealTargetPlatform> InHostPlatforms, BranchInfo InBranch, string InBranchName, GUBPBranchHacker.BranchOptions InBranchOptions, bool bInForceIncrementalCompile, JobInfo JobInfo)
 		{
 			HostPlatforms = new List<UnrealTargetPlatform>(InHostPlatforms);
 			Branch = InBranch;
+			BranchName = InBranchName;
 			BranchOptions = InBranchOptions;
 			bForceIncrementalCompile = bInForceIncrementalCompile;
 			this.JobInfo = JobInfo;
@@ -372,7 +373,7 @@ partial class GUBP
         return AltHostPlatform;
     }
 
-    void AddNodesForBranch(List<UnrealTargetPlatform> InitialHostPlatforms, JobInfo JobInfo, GUBPBranchHacker.BranchOptions BranchOptions, out List<BuildNodeDefinition> AllNodeDefinitions, out List<AggregateNodeDefinition> AllAggregateDefinitions, ref int TimeQuantum)
+    void AddNodesForBranch(List<UnrealTargetPlatform> InitialHostPlatforms, string BranchName, JobInfo JobInfo, GUBPBranchHacker.BranchOptions BranchOptions, out List<BuildNodeDefinition> AllNodeDefinitions, out List<AggregateNodeDefinition> AllAggregateDefinitions, ref int TimeQuantum)
 	{
 		DateTime StartTime = DateTime.UtcNow;
 
@@ -467,7 +468,7 @@ partial class GUBP
             }
         }
 
-		GUBPBranchConfig BranchConfig = new GUBPBranchConfig(HostPlatforms, Branch, BranchOptions, bForceIncrementalCompile, JobInfo);
+		GUBPBranchConfig BranchConfig = new GUBPBranchConfig(HostPlatforms, Branch, BranchName, BranchOptions, bForceIncrementalCompile, JobInfo);
 
         BranchConfig.AddNode(new VersionFilesNode());
 
diff --git a/Engine/Source/Programs/AutomationTool/GUBP/LegacyNodes.cs b/Engine/Source/Programs/AutomationTool/GUBP/LegacyNodes.cs
--- a/Engine/Source/Programs/AutomationTool/GUBP/LegacyNodes.cs
+++ b/Engine/Source/Programs/AutomationTool/GUBP/LegacyNodes.cs
@@ -397,10 +397,6 @@ partial class GUBP
             string AddArgs = "-CopyAppBundleBackToDevice";
 
             Agenda.AddTargets(new string[] { "UnrealHeaderTool" }, HostPlatform, UnrealTargetConfiguration.Development, InAddArgs: AddArgs);
-			if (HostPlatform == UnrealTargetPlatform.Win64)
-			{
-				Agenda.AddTargets(new string[] { "ParallelExecutor" }, HostPlatform, UnrealTargetConfiguration.Development, InAddArgs: AddArgs);
-			}
             return Agenda;
         }
         public override void PostBuild(GUBP bp, UE4Build UE4Build)
@@ -1084,7 +1080,7 @@ partial class GUBP
 
 			if(IsSample(BranchConfig, InGameProj))
 			{
-				AddDependency(WaitToPackageSamplesNode.StaticGetFullName());
+				AddPseudodependency(WaitToPackageSamplesNode.StaticGetFullName());
 			}
 
             if (InGameProj.GameName != BranchConfig.Branch.BaseEngineProject.GameName && GameProj.Properties.Targets.ContainsKey(TargetRules.TargetType.Editor))
@@ -1523,7 +1519,7 @@ partial class GUBP
         {
 			foreach(UnrealTargetPlatform HostPlatform in HostPlatforms)
 			{
-				AddDependency(RootEditorNode.StaticGetFullName(HostPlatform));
+				AddPseudodependency(RootEditorNode.StaticGetFullName(HostPlatform));
 			}
 		}
 
@@ -1637,7 +1633,7 @@ partial class GUBP
 
 			if(GamePlatformMonolithicsNode.IsSample(BranchConfig, GameProj))
 			{
-				AddDependency(WaitToPackageSamplesNode.StaticGetFullName());
+				AddPseudodependency(WaitToPackageSamplesNode.StaticGetFullName());
 				AgentSharingGroup = "SampleCooks" + StaticGetHostPlatformSuffix(HostPlatform);
 			}
         }
@@ -2014,8 +2010,8 @@ partial class GUBP
             }
             string BaseDir = CommandUtils.CombinePaths(CommandUtils.RootBuildStorageDirectory(), BuildShareName);
             string NodeName = StaticGetFullName(InGameProj, InHostPlatform, InClientTargetPlatforms, InClientConfigs, InServerTargetPlatforms, InServerConfigs, InClientNotGame);
-            string Inner = P4Env.BuildRootEscaped + "-CL-" + P4Env.ChangelistString + BranchConfig.JobInfo.GetPreflightSuffix();
-            string ArchiveDirectory = CombinePaths(BaseDir, NodeName, Inner);
+            string Inner = P4Env.BuildRootEscaped + "-" + BranchConfig.JobInfo.BuildName;
+			string ArchiveDirectory = CombinePaths(BaseDir, NodeName, Inner);
             return ArchiveDirectory;
         }
         public override void DoBuild(GUBP bp)
@@ -2812,7 +2808,7 @@ partial class GUBP
         public CleanSharedTempStorageNode(GUBP bp, GUBPBranchConfig BranchConfig)
         {
             var ToolsNode = BranchConfig.FindNode(ToolsForCompileNode.StaticGetFullName(UnrealTargetPlatform.Win64));
-            RootNameForTempStorage = BranchConfig.JobInfo.RootNameForTempStorage;
+            RootNameForTempStorage = BranchConfig.BranchOptions.RootNameForTempStorage ?? "UE4";
             AgentSharingGroup = ToolsNode.AgentSharingGroup;
         }
         public override float Priority()
diff --git a/Engine/Source/Programs/AutomationTool/GUBP/Nodes/BuildNode.cs b/Engine/Source/Programs/AutomationTool/GUBP/Nodes/BuildNode.cs
--- a/Engine/Source/Programs/AutomationTool/GUBP/Nodes/BuildNode.cs
+++ b/Engine/Source/Programs/AutomationTool/GUBP/Nodes/BuildNode.cs
@@ -109,6 +109,7 @@ namespace AutomationTool
 		public bool IsParallelAgentShareEditor;
 		public bool IsSticky;
 		public bool IsTest;
+		public bool CopyToSharedStorage;
 		public string DisplayGroupName;
 
 		public List<string> BuildProducts;
@@ -132,17 +133,17 @@ namespace AutomationTool
 			DisplayGroupName = Definition.DisplayGroupName;
 		}
 
-		public virtual void ArchiveBuildProducts(TempStorageNodeInfo TempStorageNodeInfo, bool bLocalOnly)
+		public virtual void ArchiveBuildProducts(string SharedStorageDir, bool bWriteToSharedStorage)
 		{
-			TempStorage.StoreToTempStorage(TempStorageNodeInfo, BuildProducts, bLocalOnly, CommandUtils.CmdEnv.LocalRoot);
+			TempStorage.StoreToTempStorage(Name, BuildProducts, SharedStorageDir, bWriteToSharedStorage);
 		}
 
-		public virtual void RetrieveBuildProducts(TempStorageNodeInfo TempStorageNodeInfo)
+		public virtual void RetrieveBuildProducts(string SharedStorageDir)
 		{
-			CommandUtils.Log("***** Retrieving GUBP Node {0} from {1}", Name, TempStorageNodeInfo.GetRelativeDirectory());
+			CommandUtils.Log("***** Retrieving GUBP Node {0}", Name);
 			try
 			{
-				BuildProducts = TempStorage.RetrieveFromTempStorage(TempStorageNodeInfo, CommandUtils.CmdEnv.LocalRoot);
+				BuildProducts = TempStorage.RetrieveFromTempStorage(SharedStorageDir, Name);
 			}
 			catch (Exception Ex)
 			{
diff --git a/Engine/Source/Programs/AutomationTool/GUBP/Nodes/LegacyNode.cs b/Engine/Source/Programs/AutomationTool/GUBP/Nodes/LegacyNode.cs
--- a/Engine/Source/Programs/AutomationTool/GUBP/Nodes/LegacyNode.cs
+++ b/Engine/Source/Programs/AutomationTool/GUBP/Nodes/LegacyNode.cs
@@ -57,9 +57,9 @@ namespace AutomationTool
 			Node = Definition.Node;
 		}
 
-		public override void RetrieveBuildProducts(TempStorageNodeInfo TempStorageNodeInfo)
+		public override void RetrieveBuildProducts(string SharedStorageDir)
 		{
-			base.RetrieveBuildProducts(TempStorageNodeInfo);
+			base.RetrieveBuildProducts(SharedStorageDir);
 
 			Node.BuildProducts = BuildProducts;
 		}
