Copying //UE4/Dev-Build to //UE4/Main

diff --git a/Engine/Source/Programs/AutomationTool/GUBP/TempStorage.cs b/Engine/Source/Programs/AutomationTool/GUBP/TempStorage.cs
--- a/Engine/Source/Programs/AutomationTool/GUBP/TempStorage.cs
+++ b/Engine/Source/Programs/AutomationTool/GUBP/TempStorage.cs
@@ -14,1364 +14,1243 @@ using System.IO.Compression;
 
 namespace AutomationTool
 {
-    /// <summary>
-    /// Describes temp storage for a particular GUBP Node being run. A GUBP job consists of many nodes organized in a dependency graph that run in a distributed fashion.
-    /// Each node has the ability to store files into a shared temp storage location specifically for that node.
-    /// When a later node runs, it can retrieve those temp storage files and use them to bootstrap itself so it doesn't have to redo the work of earlier nodes.
-    /// A Node's temp storage location is determined from the global GUBP Job Info along with the node's "storage name", which is typically the name of the GUBP node.
-    /// However, some nodes create additional "pseudo-nodes" to store things like completion status of the node, so future jobs can find their history.
-    /// Immutable.
-    /// </summary>
-    /// <remarks>
-    /// This class is immutable to ensure no side-effect code an tamper with it after the node starts.
-    /// 
-    /// @todo: This storage naming convention relies on the build system preventing jobs from being run more than once for a given CL. We should be putting a Job UID in here or something so a job can run twice!!
-    /// </remarks>
-    public class TempStorageNodeInfo
-    {
-        /// <summary>
-        /// The GUBP global job attributes.
-        /// </summary>
-        public JobInfo JobInfo { get; private set; }
-
-        /// <summary>
-        /// The storage name associated with this node. It is usually the name of the GUBP node running, but sometimes certain suffixes are added to indicate node completion status, etc.
-        /// </summary>
-        public string NodeStorageName { get; private set; }
-
-        /// <summary>
-        /// Constructor. This class is immutable, so all attributes must be fully specified here.
-        /// </summary>
-        /// <param name="JobInfo"></param>
-        /// <param name="NodeStorageName"></param>
-        public TempStorageNodeInfo(JobInfo JobInfo, string NodeStorageName)
-        {
-            this.JobInfo = JobInfo;
-            this.NodeStorageName = NodeStorageName;
-        }
-
-        /// <summary>
-        /// Cache of the legacy string. These values are not super cheap to construct, and can get rather long, so we only cache the value once it is asked for.
-        /// </summary>
-        private string CachedString;
-
-        /// <summary>
-        /// Cache of the relative directory. These values are not super cheap to construct, and can get rather long, so we only cache the value once it is asked for.
-        /// </summary>
-        private string CachedDir;
-
-        /// <summary>
-        /// Cache of the filename. These values are not super cheap to construct, and can get rather long, so we only cache the value once it is asked for.
-        /// </summary>
-        private string CachedFile;
-
-        /// <summary>
-        /// Used by anything that needs the block info as a string. This should generally be legacy stuff, as we should only be using the temp storage path for temp storage!
-        /// </summary>
-        /// <returns>{BranchNameForTempStorage}-{Changelist}{PreflightInfo}-{NodeStorageName}</returns>
-        public string GetLegacyString()
-        {
-            if (string.IsNullOrEmpty(CachedString))
-            {
-                CachedString = string.Format("{0}-{1}{2}-{3}", JobInfo.BranchNameForTempStorage, JobInfo.Changelist, JobInfo.GetPreflightSuffix(), NodeStorageName);
-            }
-            return CachedString;
-        }
-
-        /// <summary>
-        /// Turns the block info into a relative path to a folder where temp storage should be placed.
-        /// NOTE: <see cref="TempStorage.FindMatchingSharedTempStorageNodeCLs"/> is dependent on this directory layout!!!
-        /// If this layout is changed, TempStorage.FindMatchingSharedTempStorageNodeCLs needs to be changed as well.
-        /// </summary>
-        /// <returns>{BranchNameForTempStorage}/{Changelist}{PreflightInfo}/{NodeStorageName}</returns>
-        public string GetRelativeDirectory()
-        {
-            if (string.IsNullOrEmpty(CachedDir))
-            {
-                CachedDir = CommandUtils.CombinePaths(JobInfo.BranchNameForTempStorage, string.Format("{0}{1}", JobInfo.Changelist, JobInfo.GetPreflightSuffix()), NodeStorageName);
-            }
-            return CachedDir;
-        }
-
-        /// <summary>
-        /// Turns the block info into a file name that will be used to store the temp manifest for that block.
-        /// </summary>
-        /// <returns>{NodeStorageName}.TempManifest</returns>
-        public string GetManifestFilename()
-        {
-            if (string.IsNullOrEmpty(CachedFile))
-            {
-                CachedFile = string.Format("{0}.TempManifest", NodeStorageName);
-            }
-            return CachedFile;
-        }
-    }
-
-    /// <summary>
-    /// Handles temp storage duties for GUBP. Temp storage is a centralized location where build products from a node are stored
-    /// so that dependent nodes can reuse them, even if that node is not the same machine that built them. The main entry points 
-    /// are <see cref="StoreToTempStorage"/> and <see cref="RetrieveFromTempStorage"/>. The former is called when a node is complete
-    /// to store its build products, and the latter is called when a node starts up for all its dependent nodes to get the previous
-    /// build products.
-    /// 
-    /// Each node writes it's temp storage to a folder whose path is defined by certain attributes of the job being run and the name of the node.
-    /// The contents are the temp data itself (stored in a single .zip file unless -NoZipTempStorage is used), and a .TempManifest file that describes all the files stored
-    /// for this node, along with their expected timestamps and sizes. This is used to verify the files retrieved and stored match what is expected.
-    /// 
-    /// To support scenarios where the node running the step is the same node as the dependent one (agent sharing groups exist to ensure this),
-    /// the system always leaves a copy of the .TempManifest in "local storage", which is a local folder (see <see cref="LocalTempStorageManifestDirectory"/>).
-    /// When asked to retrieve a node from temp storage, the system first checks that local storage to see if the manifest exists there.
-    /// If it does, it assumes this machine already has the required temp storage. Beyond just verifying the files in the manifest exist properly, nothing
-    /// is copied in this case. The function them returns an out parameter (WasLocal) that tells the system whether it successfully determined
-    /// that the local files can be used.
-    /// </summary>
-    public class TempStorage
-    {
-        /// <summary>
-        /// Structure used to store a file entry that will be written to the temp file cache. essentially stores a relative path to the file, the last write time (UTC), and the file size in bytes.
-        /// </summary>
-        private class TempStorageFileInfo
-        {
-            public string Name;
-            public DateTime Timestamp;
-            public long Size;
-
-            /// <summary>
-            /// Compares with another temp storage manifest file. Allows certain files to compare differently because our build system requires it.
-            /// Also uses a relaxed timestamp matching to allow for filesystems with limited granularity in their timestamps.
-            /// </summary>
-            /// <param name="Other"></param>
-            /// <param name="bCheckTimestamps">If true, checks timestamps of the two manifests. This is done when checking the local manifest against tampering.</param>
-            /// <returns></returns>
-            public bool Compare(TempStorageFileInfo Other, bool bCheckTimestamps)
-            {
-                bool bOk = true;
-                if (!Name.Equals(Other.Name, StringComparison.InvariantCultureIgnoreCase))
-                {
-                    CommandUtils.LogError("File name mismatch {0} {1}", Name, Other.Name);
-                    bOk = false;
-                }
-                else
-                {
-                    // this is a bit of a hack, but UAT itself creates these, so we need to allow them to be 
-                    bool bOkToBeDifferent = Name.Contains("Engine/Binaries/DotNET/");
-                    // Lets just allow all mac warnings to be silent
-                    bOkToBeDifferent = bOkToBeDifferent || Name.Contains("Engine/Binaries/Mac/");
-                    // this is a problem with mac compiles
-                    bOkToBeDifferent = bOkToBeDifferent || Name.EndsWith("MacOS/libogg.dylib");
-                    bOkToBeDifferent = bOkToBeDifferent || Name.EndsWith("MacOS/libvorbis.dylib");
-                    bOkToBeDifferent = bOkToBeDifferent || Name.EndsWith("Contents/MacOS/UE4Editor");
-
-                    LogEventType LogType = bOkToBeDifferent ? LogEventType.Console : LogEventType.Warning;
-
-                    // on FAT filesystems writetime has a two seconds resolution
-                    // cf. http://msdn.microsoft.com/en-us/library/windows/desktop/ms724290%28v=vs.85%29.aspx
-                    if (bCheckTimestamps && !((Timestamp - Other.Timestamp).TotalSeconds < 2 && (Timestamp - Other.Timestamp).TotalSeconds > -2))
-                    {
-                        CommandUtils.LogWithVerbosity(LogType, "File date mismatch {0} {1} {2} {3}", Name, Timestamp, Other.Name, Other.Timestamp);
-                        // WRH 2015/08/19 - don't actually fail any builds due to a timestamp check. there's lots of reasons why the timestamps are not preserved right now, mostly due to zipfile stuff.
-                        //bOk = bOkToBeDifferent;
-                    }
-                    if (Size != Other.Size)
-                    {
-                        CommandUtils.LogWithVerbosity(LogType, "File size mismatch {0} {1} {2} {3}", Name, Size, Other.Name, Other.Size);
-                        bOk = bOkToBeDifferent;
-                    }
-                }
-                return bOk;
-            }
-
-            public override string ToString()
-            {
-                return String.IsNullOrEmpty(Name) ? "" : Name;
-            }
-        }
-
-        /// <summary>
-        /// Represents a temp storage manifest file, or a listing of all the files stored in a temp storage folder.
-        /// Essentially stores a mapping of folder names to a list of file infos that describe each file.
-        /// Can be created from list if directory to file mappings, saved to an XML file, and loaded from an XML file.
-        /// All files and directories are represented as relative to the root folder in which the manifest is saved in.
-        /// </summary>
-        public class TempStorageManifest
-        {
-            private const string RootElementName = "tempstorage";
-            private const string DirectoryElementName = "directory";
-            private const string FileElementName = "file";
-            private const string NameAttributeName = "name";
-            private const string TimestampAttributeName = "timestamp";
-            private const string SizeAttributeName = "size";
-
-            /// <summary>
-            /// A mapping of relative directory names to a list of file infos inside that directory, also stored as relative paths to the manifest file.
-            /// </summary>
-            private readonly Dictionary<string, List<TempStorageFileInfo>> Directories;
-
-            /// <summary>
-            /// Creates an manifest from a given directory to file list mapping. Internal helper function.
-            /// </summary>
-            /// <param name="Directories">Used to initialize the Directories member.</param>
-            private TempStorageManifest(Dictionary<string, List<TempStorageFileInfo>> Directories)
-            {
-                this.Directories = Directories;
-            }
-
-            /// <summary>
-            /// Creates a manifest from a flat list of files (in many folders) and a BaseFolder from which they are rooted.
-            /// </summary>
-            /// <param name="InFiles">List of full file paths</param>
-            /// <param name="RootDir">Root folder for all the files. All files must be relative to this RootDir.</param>
-            /// <returns>The newly create TempStorageManifest, if all files exist. Otherwise throws an AutomationException.</returns>
-            public static TempStorageManifest Create(List<string> InFiles, string RootDir)
-            {
-                var Directories = new Dictionary<string, List<TempStorageFileInfo>>();
-
-                foreach (string Filename in InFiles)
-                {
-                    // use this to warm up the file shared on Mac.
-                    InternalUtils.Robust_FileExists(true, Filename, "Could not add {0} to manifest because it does not exist");
-                    // should exist now, let's get the size and timestamp
-                    // We also use this to get the OS's fullname in a consistent way to ensure we handle case sensitivity issues when comparing values below.
-                    var FileInfo = new FileInfo(Filename);
-
-                    // Strip the root dir off the file, as we only store relative paths in the manifest.
-                    // Manifest only stores path with slashes for consistency, so ensure we convert them here.
-                    var RelativeFile = CommandUtils.ConvertSeparators(PathSeparator.Slash, CommandUtils.StripBaseDirectory(FileInfo.FullName, RootDir));
-                    var RelativeDir = CommandUtils.ConvertSeparators(PathSeparator.Slash, Path.GetDirectoryName(RelativeFile));
-
-                    // add the file entry, adding a directory entry along the way if necessary.
-                    List<TempStorageFileInfo> ManifestDirectory;
-                    if (Directories.TryGetValue(RelativeDir, out ManifestDirectory) == false)
-                    {
-                        ManifestDirectory = new List<TempStorageFileInfo>();
-                        Directories.Add(RelativeDir, ManifestDirectory);
-                    }
-                    ManifestDirectory.Add(new TempStorageFileInfo { Name = RelativeFile, Timestamp = FileInfo.LastWriteTimeUtc, Size = FileInfo.Length });
-                }
+	/// <summary>
+	/// Handles temp storage duties for GUBP. Temp storage is a centralized location where build products from a node are stored
+	/// so that dependent nodes can reuse them, even if that node is not the same machine that built them. The main entry points 
+	/// are <see cref="StoreToTempStorage"/> and <see cref="RetrieveFromTempStorage"/>. The former is called when a node is complete
+	/// to store its build products, and the latter is called when a node starts up for all its dependent nodes to get the previous
+	/// build products.
+	/// 
+	/// Each node writes it's temp storage to a folder whose path is defined by certain attributes of the job being run and the name of the node.
+	/// The contents are the temp data itself (stored in a single .zip file unless -NoZipTempStorage is used), and a .TempManifest file that describes all the files stored
+	/// for this node, along with their expected timestamps and sizes. This is used to verify the files retrieved and stored match what is expected.
+	/// 
+	/// To support scenarios where the node running the step is the same node as the dependent one (agent sharing groups exist to ensure this),
+	/// the system always leaves a copy of the .TempManifest in "local storage", which is a local folder (see <see cref="LocalTempStorageManifestDirectory"/>).
+	/// When asked to retrieve a node from temp storage, the system first checks that local storage to see if the manifest exists there.
+	/// If it does, it assumes this machine already has the required temp storage. Beyond just verifying the files in the manifest exist properly, nothing
+	/// is copied in this case. The function them returns an out parameter (WasLocal) that tells the system whether it successfully determined
+	/// that the local files can be used.
+	/// </summary>
+	public class TempStorage
+	{
+		/// <summary>
+		/// Structure used to store a file entry that will be written to the temp file cache. essentially stores a relative path to the file, the last write time (UTC), and the file size in bytes.
+		/// </summary>
+		private class TempStorageFileInfo
+		{
+			public string Name;
+			public DateTime Timestamp;
+			public long Size;
+
+			/// <summary>
+			/// Compares with another temp storage manifest file. Allows certain files to compare differently because our build system requires it.
+			/// Also uses a relaxed timestamp matching to allow for filesystems with limited granularity in their timestamps.
+			/// </summary>
+			/// <param name="Other"></param>
+			/// <param name="bCheckTimestamps">If true, checks timestamps of the two manifests. This is done when checking the local manifest against tampering.</param>
+			/// <returns></returns>
+			public bool Compare(TempStorageFileInfo Other, bool bCheckTimestamps)
+			{
+				bool bOk = true;
+				if (!Name.Equals(Other.Name, StringComparison.InvariantCultureIgnoreCase))
+				{
+					CommandUtils.LogError("File name mismatch {0} {1}", Name, Other.Name);
+					bOk = false;
+				}
+				else
+				{
+					// this is a bit of a hack, but UAT itself creates these, so we need to allow them to be 
+					bool bOkToBeDifferent = Name.Contains("Engine/Binaries/DotNET/");
+					// Lets just allow all mac warnings to be silent
+					bOkToBeDifferent = bOkToBeDifferent || Name.Contains("Engine/Binaries/Mac/");
+					// this is a problem with mac compiles
+					bOkToBeDifferent = bOkToBeDifferent || Name.EndsWith("MacOS/libogg.dylib");
+					bOkToBeDifferent = bOkToBeDifferent || Name.EndsWith("MacOS/libvorbis.dylib");
+					bOkToBeDifferent = bOkToBeDifferent || Name.EndsWith("Contents/MacOS/UE4Editor");
+
+					LogEventType LogType = bOkToBeDifferent ? LogEventType.Console : LogEventType.Warning;
+
+					// on FAT filesystems writetime has a two seconds resolution
+					// cf. http://msdn.microsoft.com/en-us/library/windows/desktop/ms724290%28v=vs.85%29.aspx
+					if (bCheckTimestamps && !((Timestamp - Other.Timestamp).TotalSeconds < 2 && (Timestamp - Other.Timestamp).TotalSeconds > -2))
+					{
+						CommandUtils.LogWithVerbosity(LogType, "File date mismatch {0} {1} {2} {3}", Name, Timestamp, Other.Name, Other.Timestamp);
+						// WRH 2015/08/19 - don't actually fail any builds due to a timestamp check. there's lots of reasons why the timestamps are not preserved right now, mostly due to zipfile stuff.
+						//bOk = bOkToBeDifferent;
+					}
+					if (Size != Other.Size)
+					{
+						CommandUtils.LogWithVerbosity(LogType, "File size mismatch {0} {1} {2} {3}", Name, Size, Other.Name, Other.Size);
+						bOk = bOkToBeDifferent;
+					}
+				}
+				return bOk;
+			}
+
+			public override string ToString()
+			{
+				return String.IsNullOrEmpty(Name) ? "" : Name;
+			}
+		}
 
-                return new TempStorageManifest(Directories);
-            }
-
-            /// <summary>
-            /// Compares this temp manifest with another one. It allows certain files that we know will differ to differe, while requiring the rest of the important details remain the same.
-            /// </summary>
-            /// <param name="Other"></param>
-            /// <param name="bCheckTimestamps">If true, checks timestamps of the two manifests. This is done when checking the local manifest against tampering.</param>
-            /// <returns></returns>
-            public bool Compare(TempStorageManifest Other, bool bCheckTimestamps)
-            {
-                if (Directories.Count != Other.Directories.Count)
-                {
-                    CommandUtils.LogError("Directory count mismatch {0} {1}", Directories.Count, Other.Directories.Count);
-                    foreach (KeyValuePair<string, List<TempStorageFileInfo>> Directory in Directories)
-                    {
-                        List<TempStorageFileInfo> OtherDirectory;
-                        if (Other.Directories.TryGetValue(Directory.Key, out OtherDirectory) == false)
-                        {
-                            CommandUtils.LogError("Missing Directory {0}", Directory.Key);
-                            return false;
-                        }
-                    }
-                    foreach (KeyValuePair<string, List<TempStorageFileInfo>> Directory in Other.Directories)
-                    {
-                        List<TempStorageFileInfo> OtherDirectory;
-                        if (Directories.TryGetValue(Directory.Key, out OtherDirectory) == false)
-                        {
-                            CommandUtils.LogError("Missing Other Directory {0}", Directory.Key);
-                            return false;
-                        }
-                    }
-                    return false;
-                }
+		/// <summary>
+		/// Represents a temp storage manifest file, or a listing of all the files stored in a temp storage folder.
+		/// Essentially stores a mapping of folder names to a list of file infos that describe each file.
+		/// Can be created from list if directory to file mappings, saved to an XML file, and loaded from an XML file.
+		/// All files and directories are represented as relative to the root folder in which the manifest is saved in.
+		/// </summary>
+		public class TempStorageManifest
+		{
+			private const string RootElementName = "tempstorage";
+			private const string DirectoryElementName = "directory";
+			private const string FileElementName = "file";
+			private const string NameAttributeName = "name";
+			private const string TimestampAttributeName = "timestamp";
+			private const string SizeAttributeName = "size";
+
+			/// <summary>
+			/// A mapping of relative directory names to a list of file infos inside that directory, also stored as relative paths to the manifest file.
+			/// </summary>
+			private readonly Dictionary<string, List<TempStorageFileInfo>> Directories;
+
+			/// <summary>
+			/// Creates an manifest from a given directory to file list mapping. Internal helper function.
+			/// </summary>
+			/// <param name="Directories">Used to initialize the Directories member.</param>
+			private TempStorageManifest(Dictionary<string, List<TempStorageFileInfo>> Directories)
+			{
+				this.Directories = Directories;
+			}
+
+			/// <summary>
+			/// Creates a manifest from a flat list of files (in many folders) and a BaseFolder from which they are rooted.
+			/// </summary>
+			/// <param name="InFiles">List of full file paths</param>
+			/// <param name="RootDir">Root folder for all the files. All files must be relative to this RootDir.</param>
+			/// <returns>The newly create TempStorageManifest, if all files exist. Otherwise throws an AutomationException.</returns>
+			public static TempStorageManifest Create(List<string> InFiles, string RootDir)
+			{
+				var Directories = new Dictionary<string, List<TempStorageFileInfo>>();
+
+				foreach (string Filename in InFiles)
+				{
+					// use this to warm up the file shared on Mac.
+					InternalUtils.Robust_FileExists(true, Filename, "Could not add {0} to manifest because it does not exist");
+					// should exist now, let's get the size and timestamp
+					// We also use this to get the OS's fullname in a consistent way to ensure we handle case sensitivity issues when comparing values below.
+					var FileInfo = new FileInfo(Filename);
+
+					// Strip the root dir off the file, as we only store relative paths in the manifest.
+					// Manifest only stores path with slashes for consistency, so ensure we convert them here.
+					var RelativeFile = CommandUtils.ConvertSeparators(PathSeparator.Slash, CommandUtils.StripBaseDirectory(FileInfo.FullName, RootDir));
+					var RelativeDir = CommandUtils.ConvertSeparators(PathSeparator.Slash, Path.GetDirectoryName(RelativeFile));
+
+					// add the file entry, adding a directory entry along the way if necessary.
+					List<TempStorageFileInfo> ManifestDirectory;
+					if (Directories.TryGetValue(RelativeDir, out ManifestDirectory) == false)
+					{
+						ManifestDirectory = new List<TempStorageFileInfo>();
+						Directories.Add(RelativeDir, ManifestDirectory);
+					}
+					ManifestDirectory.Add(new TempStorageFileInfo { Name = RelativeFile, Timestamp = FileInfo.LastWriteTimeUtc, Size = FileInfo.Length });
+				}
+
+				return new TempStorageManifest(Directories);
+			}
+
+			/// <summary>
+			/// Compares this temp manifest with another one. It allows certain files that we know will differ to differe, while requiring the rest of the important details remain the same.
+			/// </summary>
+			/// <param name="Other"></param>
+			/// <param name="bCheckTimestamps">If true, checks timestamps of the two manifests. This is done when checking the local manifest against tampering.</param>
+			/// <returns></returns>
+			public bool Compare(TempStorageManifest Other, bool bCheckTimestamps)
+			{
+				if (Directories.Count != Other.Directories.Count)
+				{
+					CommandUtils.LogError("Directory count mismatch {0} {1}", Directories.Count, Other.Directories.Count);
+					foreach (KeyValuePair<string, List<TempStorageFileInfo>> Directory in Directories)
+					{
+						List<TempStorageFileInfo> OtherDirectory;
+						if (Other.Directories.TryGetValue(Directory.Key, out OtherDirectory) == false)
+						{
+							CommandUtils.LogError("Missing Directory {0}", Directory.Key);
+							return false;
+						}
+					}
+					foreach (KeyValuePair<string, List<TempStorageFileInfo>> Directory in Other.Directories)
+					{
+						List<TempStorageFileInfo> OtherDirectory;
+						if (Directories.TryGetValue(Directory.Key, out OtherDirectory) == false)
+						{
+							CommandUtils.LogError("Missing Other Directory {0}", Directory.Key);
+							return false;
+						}
+					}
+					return false;
+				}
+
+				bool bResult = true;
+				foreach (KeyValuePair<string, List<TempStorageFileInfo>> Directory in Directories)
+				{
+					List<TempStorageFileInfo> OtherDirectory;
+					if (Other.Directories.TryGetValue(Directory.Key, out OtherDirectory) == false)
+					{
+						CommandUtils.LogError("Missing Directory {0}", Directory.Key);
+						return false;
+					}
+					if (OtherDirectory.Count != Directory.Value.Count)
+					{
+						CommandUtils.LogError("File count mismatch {0} {1} {2}", Directory.Key, OtherDirectory.Count, Directory.Value.Count);
+						for (int FileIndex = 0; FileIndex < Directory.Value.Count; ++FileIndex)
+						{
+							CommandUtils.Log("Manifest1: {0}", Directory.Value[FileIndex].Name);
+						}
+						for (int FileIndex = 0; FileIndex < OtherDirectory.Count; ++FileIndex)
+						{
+							CommandUtils.Log("Manifest2: {0}", OtherDirectory[FileIndex].Name);
+						}
+						return false;
+					}
+					for (int FileIndex = 0; FileIndex < Directory.Value.Count; ++FileIndex)
+					{
+						TempStorageFileInfo File = Directory.Value[FileIndex];
+						TempStorageFileInfo OtherFile = OtherDirectory[FileIndex];
+						if (File.Compare(OtherFile, bCheckTimestamps) == false)
+						{
+							bResult = false;
+						}
+					}
+				}
+				return bResult;
+			}
+
+			/// <summary>
+			/// Loads a manifest from a file
+			/// </summary>
+			/// <param name="Filename">Full path to the manifest.</param>
+			/// <returns>The newly created manifest instance</returns>
+			public static TempStorageManifest Load(string Filename)
+			{
+				try
+				{
+					// Load the manifest XML file.
+					var Directories = XDocument.Load(Filename).Root
+						// Get the directory array of child elements
+						.Elements(DirectoryElementName)
+						// convert it to a dictionary of directory names to a list of TempStorageFileInfo elements for each file in the directory.
+						.ToDictionary(
+							DirElement => DirElement.Attribute(NameAttributeName).Value,
+							DirElement => DirElement.Elements(FileElementName).Select(FileElement => new TempStorageFileInfo
+											{
+												Timestamp = new DateTime(long.Parse(FileElement.Attribute(TimestampAttributeName).Value)/*, DateTimeKind.Utc*/),
+												Size = long.Parse(FileElement.Attribute(SizeAttributeName).Value),
+												Name = FileElement.Value,
+											})
+											.ToList()
+					);
+
+					// The manifest must have at least one file with length > 0
+					if (IsEmptyManifest(Directories))
+					{
+						throw new AutomationException("Attempt to load empty manifest.");
+					}
+
+					return new TempStorageManifest(Directories);
+				}
+				catch (Exception Ex)
+				{
+					throw new AutomationException(Ex, "Failed to load manifest file {0}", Filename);
+				}
+			}
+
+			/// <summary>
+			/// Returns true if the manifest represented by the dictionary is empty (no files with non-zero length)
+			/// </summary>
+			/// <param name="Directories">Dictionary of directory to file mappings that define the manifest</param>
+			/// <returns></returns>
+			private static bool IsEmptyManifest(Dictionary<string, List<TempStorageFileInfo>> Directories)
+			{
+				return !Directories.SelectMany(Pair => Pair.Value).Any(FileInfo => FileInfo.Size > 0);
+			}
+
+			/// <summary>
+			/// Saves the manifest to the given filename.
+			/// </summary>
+			/// <param name="Filename">Full path to the file to save to.</param>
+			public void Save(string Filename)
+			{
+				if (IsEmptyManifest(Directories))
+				{
+					throw new AutomationException("Attempt to save empty manifest.");
+				}
+
+				new XElement(RootElementName,
+					from Dir in Directories
+					select new XElement(DirectoryElementName,
+						new XAttribute(NameAttributeName, Dir.Key),
+						from File in Dir.Value
+						select new XElement(FileElementName,
+							new XAttribute(TimestampAttributeName, File.Timestamp.Ticks),
+							new XAttribute(SizeAttributeName, File.Size),
+							File.Name)
+					)
+				).Save(Filename);
+				CommandUtils.Log("Saved temp manifest {0} with {1} files and total size {2}", Filename, Directories.SelectMany(Dir => Dir.Value).Count(), Directories.SelectMany(Dir => Dir.Value).Sum(File => File.Size));
+			}
+
+			/// <summary>
+			/// Returns the sum of filesizes in the manifest in bytes.
+			/// </summary>
+			/// <returns>Returns the sum of filesizes in the manifest in bytes.</returns>
+			public long GetTotalSize()
+			{
+				return Directories.SelectMany(Dir => Dir.Value).Sum(FileInfo => FileInfo.Size);
+			}
+
+			/// <summary>
+			/// Gets a flat list of files in the manifest, converting to full path names rooted at the given base dir.
+			/// </summary>
+			/// <param name="RootDir">Root dir to prepend to all the files in the manifest.</param>
+			/// <returns>Flat list of all files in the manifest re-rooted at RootDir.</returns>
+			public List<string> GetFiles(string RootDir)
+			{
+				// flatten the list of directories, pull the files out, set their root path, and ensure the path is not too long.
+				return Directories
+					.SelectMany(Dir => Dir.Value)
+					.Select(FileInfo => CommandUtils.CombinePaths(RootDir, FileInfo.Name))
+					.ToList();
+			}
+		}
 
-                bool bResult = true;
-                foreach (KeyValuePair<string, List<TempStorageFileInfo>> Directory in Directories)
-                {
-                    List<TempStorageFileInfo> OtherDirectory;
-                    if (Other.Directories.TryGetValue(Directory.Key, out OtherDirectory) == false)
-                    {
-                        CommandUtils.LogError("Missing Directory {0}", Directory.Key); 
-                        return false;
-                    }
-                    if (OtherDirectory.Count != Directory.Value.Count)
-                    {
-                        CommandUtils.LogError("File count mismatch {0} {1} {2}", Directory.Key, OtherDirectory.Count, Directory.Value.Count);
-                        for (int FileIndex = 0; FileIndex < Directory.Value.Count; ++FileIndex)
-                        {
-                            CommandUtils.Log("Manifest1: {0}", Directory.Value[FileIndex].Name);
-                        }
-                        for (int FileIndex = 0; FileIndex < OtherDirectory.Count; ++FileIndex)
-                        {
-                            CommandUtils.Log("Manifest2: {0}", OtherDirectory[FileIndex].Name);
-                        }
-                        return false;
-                    }
-                    for (int FileIndex = 0; FileIndex < Directory.Value.Count; ++FileIndex)
-                    {
-                        TempStorageFileInfo File = Directory.Value[FileIndex];
-                        TempStorageFileInfo OtherFile = OtherDirectory[FileIndex];
-                        if (File.Compare(OtherFile, bCheckTimestamps) == false)
-                        {
-                            bResult = false;
-                        }
-                    }
-                }
-                return bResult;
-            }
-
-            /// <summary>
-            /// Loads a manifest from a file
-            /// </summary>
-            /// <param name="Filename">Full path to the manifest.</param>
-            /// <returns>The newly created manifest instance</returns>
-            public static TempStorageManifest Load(string Filename)
-            {
-                try
-                {
-                    // Load the manifest XML file.
-                    var Directories = XDocument.Load(Filename).Root
-                        // Get the directory array of child elements
-                        .Elements(DirectoryElementName)
-                        // convert it to a dictionary of directory names to a list of TempStorageFileInfo elements for each file in the directory.
-                        .ToDictionary(
-                            DirElement => DirElement.Attribute(NameAttributeName).Value,
-                            DirElement => DirElement.Elements(FileElementName).Select(FileElement => new TempStorageFileInfo
-                                            {
-                                                Timestamp = new DateTime(long.Parse(FileElement.Attribute(TimestampAttributeName).Value)/*, DateTimeKind.Utc*/),
-                                                Size = long.Parse(FileElement.Attribute(SizeAttributeName).Value),
-                                                Name = FileElement.Value,
-                                            })
-                                            .ToList()
-                    );
-
-                    // The manifest must have at least one file with length > 0
-                    if (IsEmptyManifest(Directories))
-                    {
-                        throw new AutomationException("Attempt to load empty manifest.");
-                    }
-                    
-                    return new TempStorageManifest(Directories);
-                }
-                catch (Exception Ex)
-                {
-                    throw new AutomationException(Ex, "Failed to load manifest file {0}", Filename);
-                }
-            }
-
-            /// <summary>
-            /// Returns true if the manifest represented by the dictionary is empty (no files with non-zero length)
-            /// </summary>
-            /// <param name="Directories">Dictionary of directory to file mappings that define the manifest</param>
-            /// <returns></returns>
-            private static bool IsEmptyManifest(Dictionary<string, List<TempStorageFileInfo>> Directories)
-            {
-                return !Directories.SelectMany(Pair => Pair.Value).Any(FileInfo => FileInfo.Size > 0);
-            }
-
-            /// <summary>
-            /// Saves the manifest to the given filename.
-            /// </summary>
-            /// <param name="Filename">Full path to the file to save to.</param>
-            public void Save(string Filename)
-            {
-                if (IsEmptyManifest(Directories))
-                {
-                    throw new AutomationException("Attempt to save empty manifest.");
-                }
+		/// <summary>
+		/// returns "[LocalRoot]\Engine\Saved\TmpStore" - the location where temp storage manifests will be stored before copying them to their final temp storage location.
+		/// </summary>
+		/// <returns>returns "[LocalRoot]\Engine\Saved\TmpStore"</returns>
+		private static string LocalTempStorageManifestDirectory()
+		{
+			return CommandUtils.CombinePaths(CommandUtils.CmdEnv.LocalRoot, "Engine", "Saved", TempStorageSubdirectoryName);
+		}
 
-                new XElement(RootElementName,
-                    from Dir in Directories
-                    select new XElement(DirectoryElementName,
-                        new XAttribute(NameAttributeName, Dir.Key),
-                        from File in Dir.Value
-                        select new XElement(FileElementName,
-                            new XAttribute(TimestampAttributeName, File.Timestamp.Ticks),
-                            new XAttribute(SizeAttributeName, File.Size),
-                            File.Name)
-                    )
-                ).Save(Filename);
-                CommandUtils.Log("Saved temp manifest {0} with {1} files and total size {2}", Filename, Directories.SelectMany(Dir=>Dir.Value).Count(), Directories.SelectMany(Dir=>Dir.Value).Sum(File=>File.Size));
-            }
-
-            /// <summary>
-            /// Returns the sum of filesizes in the manifest in bytes.
-            /// </summary>
-            /// <returns>Returns the sum of filesizes in the manifest in bytes.</returns>
-            public long GetTotalSize()
-            {
-                return Directories.SelectMany(Dir => Dir.Value).Sum(FileInfo => FileInfo.Size);
-            }
-
-            /// <summary>
-            /// Gets a flat list of files in the manifest, converting to full path names rooted at the given base dir.
-            /// </summary>
-            /// <param name="RootDir">Root dir to prepend to all the files in the manifest.</param>
-            /// <returns>Flat list of all files in the manifest re-rooted at RootDir.</returns>
-            public List<string> GetFiles(string RootDir)
-            {
-                // flatten the list of directories, pull the files out, set their root path, and ensure the path is not too long.
-                return Directories
-                    .SelectMany(Dir => Dir.Value)
-                    .Select(FileInfo => CommandUtils.CombinePaths(RootDir, FileInfo.Name))
-                    .ToList();
-            }
-        }
-
-        /// <summary>
-        /// returns "[LocalRoot]\Engine\Saved\TmpStore" - the location where temp storage manifests will be stored before copying them to their final temp storage location.
-        /// </summary>
-        /// <returns>returns "[LocalRoot]\Engine\Saved\TmpStore"</returns>
-        private static string LocalTempStorageManifestDirectory()
-        {
-            return CommandUtils.CombinePaths(CommandUtils.CmdEnv.LocalRoot, "Engine", "Saved", TempStorageSubdirectoryName);
-        }
-
-        /// <summary>
-        /// Legayc code to clean all temp storage data from the given directory that is older than a certain threshold (defined directly in the code).
-        /// </summary>
-        /// <param name="TopDirectory">Fully qualified path of the folder that is the root of a bunch of temp storage folders. Should be a string returned by <see cref="ResolveSharedTempStorageDirectory"/> </param>
-        private static void CleanSharedTempStorageLegacy(string TopDirectory)
-        {
-            // This is a hack to clean out the old temp storage in the old folder name.
-            TopDirectory = TopDirectory.Replace(Path.DirectorySeparatorChar + TempStorageSubdirectoryName + Path.DirectorySeparatorChar, Path.DirectorySeparatorChar + "GUBP" + Path.DirectorySeparatorChar);
-
-            using (var TelemetryStopwatch = new TelemetryStopwatch("CleanSharedTempStorageLegacy"))
-            {
-                const int MaximumDaysToKeepLegacyTempStorage = 1;
-                var StartTimeDir = DateTime.UtcNow;
-                DirectoryInfo DirInfo = new DirectoryInfo(TopDirectory);
-                var TopLevelDirs = DirInfo.GetDirectories();
-                {
-                    var BuildDuration = (DateTime.UtcNow - StartTimeDir).TotalMilliseconds;
-                    CommandUtils.Log("Took {0}s to enumerate {1} directories.", BuildDuration / 1000, TopLevelDirs.Length);
-                }
-                foreach (var TopLevelDir in TopLevelDirs)
-                {
-                    if (CommandUtils.DirectoryExists_NoExceptions(TopLevelDir.FullName))
-                    {
-                        bool bOld = false;
-                        foreach (var ThisFile in CommandUtils.FindFiles_NoExceptions(true, "*.TempManifest", false, TopLevelDir.FullName))
-                        {
-                            FileInfo Info = new FileInfo(ThisFile);
-
-                            if ((DateTime.UtcNow - Info.LastWriteTimeUtc).TotalDays > MaximumDaysToKeepLegacyTempStorage)
-                            {
-                                bOld = true;
-                            }
-                        }
-                        if (bOld)
-                        {
-                            CommandUtils.Log("Deleting temp storage directory {0}, because it is more than {1} days old.", TopLevelDir.FullName, MaximumDaysToKeepLegacyTempStorage);
-                            var StartTime = DateTime.UtcNow;
-                            try
-                            {
-                                if (Directory.Exists(TopLevelDir.FullName))
-                                {
-                                    // try the direct approach first
-                                    Directory.Delete(TopLevelDir.FullName, true);
-                                }
-                            }
-                            catch
-                            {
-                            }
-                            CommandUtils.DeleteDirectory_NoExceptions(true, TopLevelDir.FullName);
-                            var BuildDuration = (DateTime.UtcNow - StartTime).TotalMilliseconds;
-                            CommandUtils.Log("Took {0}s to delete {1}.", BuildDuration / 1000, TopLevelDir.FullName);
-                        }
-                    }
-                }
-            }
-        }
-
-        /// <summary>
-        /// Cleans all temp storage data from the given directory that is older than a certain threshold (defined directly in the code).
-        /// This will clean up builds that are NOT part of the the job for this branch. This is intentional because once a branch is shut down,
-        /// there would be nothing left to clean it. This will only clean the supplied folder, but it is assumed this folder is an actual
-        /// Temp Storage folder with .TempManifest files in it.
-        /// </summary>
-        /// <param name="TopDirectory">Fully qualified path of the folder that is the root of a bunch of temp storage folders. Should be a string returned by <see cref="ResolveSharedTempStorageDirectory"/> </param>
-        private static void CleanSharedTempStorage(string TopDirectory)
-        {
-            CleanSharedTempStorageLegacy(TopDirectory);
-
-            Action<string> TryToDeletePossiblyEmptyFolder = (string FolderName) =>
-                {
-                    try
-                    {
-                        Directory.Delete(FolderName);
-                    }
-                    catch (IOException)
-                    {
-                        // only catch "directory is not empty type exceptions, if possible. Best we can do is check for IOException.
-                    }
-                    catch (Exception Ex)
-                    {
-                        CommandUtils.LogWarning("Unexpected failure trying to delete possibly empty temp folder {0}: {1}", FolderName, Ex);
-                    }
-                };
-
-            CommandUtils.Log("Cleaning temp storage for {0}...", TopDirectory);
-            int FoldersCleaned = 0;
-            int FoldersManuallyCleaned = 0;
-            int FoldersFailedCleaned = 0;
-            using (var TelemetryStopwatch = new TelemetryStopwatch("CleanSharedTempStorage"))
-            {
-                const double MaximumDaysToKeepTempStorage = 3;
-                var Now = DateTime.UtcNow;
-                foreach (var OldManifestInfo in
-                    // First subdirectory is the branch name
-                    from BranchDir in Directory.EnumerateDirectories(TopDirectory)
-                    // Second subdirectory is the changelist
-                    from CLDir in Directory.EnumerateDirectories(BranchDir)
-                    // third subdirectory is the node storage name
-                    from NodeNameDir in Directory.EnumerateDirectories(CLDir)
-                    // now look for manifest files.
-                    let ManifestFile = Directory.EnumerateFiles(NodeNameDir, "*.TempManifest").SingleOrDefault()
-                    where ManifestFile != null
-                    // only choose ones that are old enough.
-                    where (Now - new FileInfo(ManifestFile).LastWriteTimeUtc).TotalDays > MaximumDaysToKeepTempStorage
-                    select new { BranchDir, CLDir, NodeNameDir })
-                {
-                    try
-                    {
-                        CommandUtils.Log("Deleting folder with old temp storage {0}...", OldManifestInfo.NodeNameDir);
-                        Directory.Delete(OldManifestInfo.NodeNameDir, true);
-                        FoldersCleaned++;
-                    }
-                    catch (Exception Ex)
-                    {
-                        CommandUtils.LogWarning("Failed to delete old manifest folder '{0}', will try one file at a time: {1}", OldManifestInfo.NodeNameDir, Ex);
-                        if (CommandUtils.DeleteDirectory_NoExceptions(true, OldManifestInfo.NodeNameDir))
-                        {
-                            FoldersManuallyCleaned++;
-                        }
-                        else
-                        {
-                            FoldersFailedCleaned++;
-                        }
-                    }
-                    // Once we are done, try to delete the CLDir and BranchDir (will fail if not empty, so no worries).
-                    TryToDeletePossiblyEmptyFolder(OldManifestInfo.CLDir);
-                    TryToDeletePossiblyEmptyFolder(OldManifestInfo.BranchDir);
-                }
-                TelemetryStopwatch.Finish(string.Format("CleanSharedTempStorage.{0}.{1}.{2}", FoldersCleaned, FoldersManuallyCleaned, FoldersFailedCleaned));
-            }
-        }
-
-        /// <summary>
-        /// Cleans the shared temp storage folder for the given root name. This code will ensure that any duplicate games
-        /// or games that resolve to use the same temp folder are not cleaned twice.
-        /// </summary>
-        /// <param name="RootNameForTempStorage">List of game names to clean temp folders for.</param>
-        public static void CleanSharedTempStorageDirectory(string RootNameForTempStorage)
-        {
-            if (!CommandUtils.IsBuildMachine || Utils.IsRunningOnMono)  // saw a hang on this, anyway it isn't necessary to clean with macs, they are slow anyway
-            {
-                return;
-            }
-            var TempStorageFolder = ResolveSharedTempStorageDirectory(RootNameForTempStorage);
-            try
-            {
-                CleanSharedTempStorage(TempStorageFolder);
-            }
-            catch (Exception Ex)
-            {
-                CommandUtils.LogWarning("Unable to Clean Temp Directory {0}. Exception: {1}", TempStorageFolder, Ex);
-            }
-        }
-
-        /// "TmpStore", the folder name where all GUBP temp storage files are added inside a game's network storage location.
-        private const string TempStorageSubdirectoryName = "TmpStore";
-
-        /// <summary>
-        /// Determines if the root UE4 temp storage directory exists.
-        /// </summary>
-        /// <param name="ForSaving">If true, confirm that the directory exists and is writeable.</param>
-        /// <returns>true if the UE4 temp storage directory exists (and is writeable if requested).</returns>
-        public static bool IsSharedTempStorageAvailable(bool ForSaving)
-        {
-            var UE4TempStorageDir = CommandUtils.CombinePaths(CommandUtils.RootBuildStorageDirectory(), "UE4", TempStorageSubdirectoryName);
-            if (ForSaving)
-            {
-                int Retries = 0;
-                while (Retries < 24)
-                {
-                    if (InternalUtils.Robust_DirectoryExistsAndIsWritable_NoExceptions(UE4TempStorageDir))
-                    {
-                        return true;
-                    }
-                    CommandUtils.FindDirectories_NoExceptions(false, "*", false, UE4TempStorageDir); // there is some internet evidence to suggest this might perk up the mac share
-                    Thread.Sleep(5000);
-                    Retries++;
-                }
-            }
-            else if (InternalUtils.Robust_DirectoryExists_NoExceptions(UE4TempStorageDir, "Could not find {0}"))
-            {
-                return true;
-            }
-            return false;
-        }
-
-
-        /// <summary>
-        /// Cache of previously resolved build directories for games. Since this requires a few filesystem checks,
-        /// we cache the results so future lookups are fast.
-        /// </summary>
-        static readonly Dictionary<string, string> SharedBuildDirectoryResolveCache = new Dictionary<string, string>();
-
-        /// <summary>
-        /// Determines full path of the shared build folder for the given game.
-        /// If the shared build folder does not already exist, will try to use UE4's folder instead.
-        /// This is because the root build folders are independently owned by each game team and managed by IT, so cannot be created on the fly.
-        /// </summary>
-        /// The root name beneath <see cref="CommandUtils.RootBuildStorageDirectory"/> where all temp storage will placed for this job.
-        /// <returns>The full path of the shared build directory name, or throws an exception if none is found and the UE4 folder is not found either.</returns>
-        public static string ResolveSharedBuildDirectory(string RootNameForTempStorage)
-        {
-            if (RootNameForTempStorage == null) throw new ArgumentNullException("RootNameForTempStorage");
-
-            if (SharedBuildDirectoryResolveCache.ContainsKey(RootNameForTempStorage))
-            {
-                return SharedBuildDirectoryResolveCache[RootNameForTempStorage];
-            }
-            string Root = CommandUtils.RootBuildStorageDirectory();
-            string Result = CommandUtils.CombinePaths(Root, RootNameForTempStorage);
-            if (!InternalUtils.Robust_DirectoryExistsAndIsWritable_NoExceptions(Result))
-            {
-                throw new AutomationException("Could not find an appropriate shared temp folder {0}", Result);
-            }
-            SharedBuildDirectoryResolveCache.Add(RootNameForTempStorage, Result);
-            return Result;
-        }
-
-        /// <summary>
-        /// Returns the full path of the temp storage directory for the given game name (ie, P:\Builds\RootNameForTempStorage\TmpStore).
-        /// If the game's build folder does not exist, will use the temp storage folder in the UE4 build folder (ie, P:\Builds\UE4\TmpStore).
-        /// If the directory does not exist, we will try to create it.
-        /// </summary>
-        /// <param name="RootNameForTempStorage">The root name beneath <see cref="CommandUtils.RootBuildStorageDirectory"/> where all temp storage will placed for this job.</param>
-        /// <returns>Essentially adds the TmpStore/ subdirectory to the SharedBuildDirectory for the game.</returns>
-        public static string ResolveSharedTempStorageDirectory(string RootNameForTempStorage)
-        {
-            string Result = CommandUtils.CombinePaths(ResolveSharedBuildDirectory(RootNameForTempStorage), TempStorageSubdirectoryName);
-
-            if (!InternalUtils.Robust_DirectoryExists_NoExceptions(Result, "Could not find {0}"))
-            {
-                CommandUtils.CreateDirectory_NoExceptions(Result);
-            }
-            if (!InternalUtils.Robust_DirectoryExists_NoExceptions(Result, "Could not find {0}"))
-            {
-                throw new AutomationException("Could not create an appropriate shared temp folder {0} for {1}", Result, RootNameForTempStorage);
-            }
-            return Result;
-        }
-
-        /// <summary>
-        /// Returns the full path to the temp storage directory for the given temp storage node and game (ie, P:\Builds\RootNameForTempStorage\TmpStore\NodeInfoDirectory)
-        /// </summary>
-        /// <param name="TempStorageNodeInfo">Node info descibing the block of temp storage (essentially used to identify a subdirectory insides the game's temp storage folder).</param>
-        /// <returns>The full path to the temp storage directory for the given storage block name for the given game.</returns>
-        private static string SharedTempStorageDirectory(TempStorageNodeInfo TempStorageNodeInfo)
-        {
-            return CommandUtils.CombinePaths(ResolveSharedTempStorageDirectory(TempStorageNodeInfo.JobInfo.RootNameForTempStorage), TempStorageNodeInfo.GetRelativeDirectory());
-        }
-
-        /// <summary>
-        /// Gets the name of the local temp storage manifest file for the given temp storage node (ie, Engine\Saved\TmpStore\NodeInfoDirectory\NodeInfoFilename.TempManifest)
-        /// </summary>
-        /// <param name="TempStorageNodeInfo">Node info descibing the block of temp storage (essentially used to identify a subdirectory insides the game's temp storage folder).</param>
-        /// <returns>The name of the temp storage manifest file for the given storage block name for the given game.</returns>
-        private static string LocalTempStorageManifestFilename(TempStorageNodeInfo TempStorageNodeInfo)
-        {
-            return CommandUtils.CombinePaths(LocalTempStorageManifestDirectory(), TempStorageNodeInfo.GetManifestFilename());
-        }
-
-        /// <summary>
-        /// Saves the list of fully qualified files rooted at RootDir to a local temp storage manifest with the given temp storage node.
-        /// </summary>
-        /// <param name="RootDir">Folder that all the given files are rooted from.</param>
-        /// <param name="TempStorageNodeInfo">Node info descibing the block of temp storage (essentially used to identify a subdirectory insides the game's temp storage folder).</param>
-        /// <param name="Files">Fully qualified names of files to reference in the manifest file.</param>
-        /// <returns>The created manifest instance (which has already been saved to disk).</returns>
-        private static TempStorageManifest SaveLocalTempStorageManifest(string RootDir, TempStorageNodeInfo TempStorageNodeInfo, List<string> Files)
-        {
-            string FinalFilename = LocalTempStorageManifestFilename(TempStorageNodeInfo);
-            var Saver = TempStorageManifest.Create(Files, RootDir);
-            CommandUtils.CreateDirectory(true, Path.GetDirectoryName(FinalFilename));
-            Saver.Save(FinalFilename);
-            return Saver;
-        }
-
-        /// <summary>
-        /// Gets the name of the temp storage manifest file for the given temp storage node and game (ie, P:\Builds\RootNameForTempStorage\TmpStore\NodeInfoDirectory\NodeInfoFilename.TempManifest)
-        /// </summary>
-        /// <param name="TempStorageNodeInfo">Node info descibing the block of temp storage (essentially used to identify a subdirectory insides the game's temp storage folder).</param>
-        /// <returns>The name of the temp storage manifest file for the given storage block name for the given game.</returns>
-        private static string SharedTempStorageManifestFilename(TempStorageNodeInfo TempStorageNodeInfo)
-        {
-            return CommandUtils.CombinePaths(SharedTempStorageDirectory(TempStorageNodeInfo), TempStorageNodeInfo.GetManifestFilename());
-        }
-
-        /// <summary>
-        /// Deletes all temp storage manifests from the local storage location (Engine\Saved\TmpStore).
-        /// Local temp storage logic only stores manifests, so this effectively deletes any local temp storage work, while not actually deleting the local files, which the temp 
-        /// storage system doesn't really own.
-        /// </summary>
-        public static void DeleteLocalTempStorage()
-        {
-            CommandUtils.DeleteDirectory(true, LocalTempStorageManifestDirectory());
-        }
-
-        /// <summary>
-        /// Deletes all shared temp storage (file and manifest) for the given temp storage node and game.
-        /// </summary>
-        /// <param name="TempStorageNodeInfo">Node info descibing the block of temp storage (essentially used to identify a subdirectory insides the game's temp storage folder).</param>
-        private static void DeleteSharedTempStorage(TempStorageNodeInfo TempStorageNodeInfo)
-        {
-            CommandUtils.DeleteDirectory(true, SharedTempStorageDirectory(TempStorageNodeInfo));
-        }
-
-        /// <summary>
-        /// Checks if a local temp storage manifest exists for the given temp storage node.
-        /// </summary>
-        /// <param name="TempStorageNodeInfo">Node info descibing the block of temp storage (essentially used to identify a subdirectory insides the game's temp storage folder).</param>
-        /// <param name="bQuiet">True to suppress logging during the operation.</param>
-        /// <returns>true if the local temp storage manifest file exists</returns>
-        public static bool LocalTempStorageManifestExists(TempStorageNodeInfo TempStorageNodeInfo, bool bQuiet = false)
-        {
-            var LocalManifest = LocalTempStorageManifestFilename(TempStorageNodeInfo);
-            return CommandUtils.FileExists_NoExceptions(bQuiet, LocalManifest);
-        }
-
-        /// <summary>
-        /// Deletes a local temp storage manifest file for the given temp storage node if it exists.
-        /// </summary>
-        /// <param name="TempStorageNodeInfo">Node info descibing the block of temp storage (essentially used to identify a subdirectory insides the game's temp storage folder).</param>
-        /// <param name="bQuiet">True to suppress logging during the operation.</param>
-        public static void DeleteLocalTempStorageManifest(TempStorageNodeInfo TempStorageNodeInfo, bool bQuiet = false)
-        {
-            var LocalManifest = LocalTempStorageManifestFilename(TempStorageNodeInfo);
-            CommandUtils.DeleteFile(bQuiet, LocalManifest);
-        }
-
-        /// <summary>
-        /// Checks if a shared temp storage manifest exists for the given temp storage node and game.
-        /// </summary>
-        /// <param name="TempStorageNodeInfo">Node info descibing the block of temp storage (essentially used to identify a subdirectory insides the game's temp storage folder).</param>
-        /// <param name="bQuiet">True to suppress logging during the operation.</param>
-        /// <returns>true if the shared temp storage manifest file exists</returns>
-        private static bool SharedTempStorageManifestExists(TempStorageNodeInfo TempStorageNodeInfo, bool bQuiet)
-        {
-            var SharedManifest = SharedTempStorageManifestFilename(TempStorageNodeInfo);
-            if (CommandUtils.FileExists_NoExceptions(bQuiet, SharedManifest))
-            {
-                return true;
-            }
-            return false;
-        }
-
-        /// <summary>
-        /// Checks if a temp storage manifest exists, either locally or in the shared folder (depending on the value of bLocalOnly) for the given temp storage node and game.
-        /// </summary>
-        /// <param name="TempStorageNodeInfo">Node info descibing the block of temp storage (essentially used to identify a subdirectory insides the game's temp storage folder).</param>
-        /// <param name="bLocalOnly">If true, only ensures that the local temp storage manifest exists. Otherwise, checks if the manifest exists in either local or shared storage.</param>
-        /// <param name="bQuiet">True to suppress logging during the operation.</param>
-        /// <returns>true if the shared temp storage manifest file exists</returns>
-        public static bool TempStorageExists(TempStorageNodeInfo TempStorageNodeInfo, bool bLocalOnly, bool bQuiet)
-        {
-            return LocalTempStorageManifestExists(TempStorageNodeInfo, bQuiet) || (!bLocalOnly && SharedTempStorageManifestExists(TempStorageNodeInfo, bQuiet));
-        }
-
-        /// <summary>
-        /// Stores some result info from a parallel zip or unzip operation.
-        /// </summary>
-        class ParallelZipResult
-        {
-            /// <summary>
-            /// Time taken for the zip portion of the operation (actually, everything BUT the copy time).
-            /// </summary>
-            public readonly TimeSpan ZipTime;
-
-            /// <summary>
-            /// Total size of the zip files that are created.
-            /// </summary>
-            public readonly long ZipFilesTotalSize;
-
-            public ParallelZipResult(TimeSpan ZipTime, long ZipFilesTotalSize)
-            {
-                this.ZipTime = ZipTime;
-                this.ZipFilesTotalSize = ZipFilesTotalSize;
-            }
-        }
-
-        /// <summary>
-        /// Zips a set of files (that must be rooted at the given RootDir) to a set of zip files in the given OutputDir. The files will be prefixed with the given basename.
-        /// </summary>
-        /// <param name="Files">Fully qualified list of files to zip (must be rooted at RootDir).</param>
-        /// <param name="RootDir">Root Directory where all files will be extracted.</param>
-        /// <param name="OutputDir">Location to place the set of zip files created.</param>
-        /// <param name="StagingDir">Location to create zip files before copying them to the OutputDir. If the OutputDir is on a remote file share, staging may be more efficient. Use null to avoid using a staging copy.</param>
-        /// <param name="ZipBasename">The basename of the set of zip files.</param>
-        /// <returns>Some metrics about the zip process.</returns>
-        /// <remarks>
-        /// This function tries to zip the files in parallel as fast as it can. It makes no guarantees about how many zip files will be created or which files will be in which zip,
-        /// but it does try to reasonably balance the file sizes.
-        /// </remarks>
-        private static ParallelZipResult ParallelZipFiles(IEnumerable<string> Files, string RootDir, string OutputDir, string StagingDir, string ZipBasename)
-        {
-            var ZipTimer = DateTimeStopwatch.Start();
-            // First get the sizes of all the files. We won't parallelize if there isn't enough data to keep the number of zips down.
-            var FilesInfo = Files
-                .Select(File => new { File, FileSize = new FileInfo(File).Length })
-                .ToList();
-
-            // Profiling results show that we can zip 100MB quite fast and it is not worth parallelizing that case and creating a bunch of zips that are relatively small.
-            const long MinFileSizeToZipInParallel = 1024 * 1024 * 100L;
-            var bZipInParallel = FilesInfo.Sum(FileInfo => FileInfo.FileSize) >= MinFileSizeToZipInParallel;
-
-            // order the files in descending order so our threads pick up the biggest ones first.
-            // We want to end with the smaller files to more effectively fill in the gaps
-            var FilesToZip = new ConcurrentQueue<string>(FilesInfo.OrderByDescending(FileInfo => FileInfo.FileSize).Select(FileInfo => FileInfo.File));
-
-            long ZipFilesTotalSize = 0L;
-            // We deliberately avoid Parallel.ForEach here because profiles have shown that dynamic partitioning creates
-            // too many zip files, and they can be of too varying size, creating uneven work when unzipping later,
-            // as ZipFile cannot unzip files in parallel from a single archive.
-            // We can safely assume the build system will not be doing more important things at the same time, so we simply use all our logical cores,
-            // which has shown to be optimal via profiling, and limits the number of resulting zip files to the number of logical cores.
-            // 
-            // Sadly, mono implemention of System.IO.Compression is really poor (as of 2015/Aug), causing OOM when parallel zipping a large set of files.
-            // However, Ionic is MUCH slower than .NET's native implementation (2x+ slower in our build farm), so we stick to the faster solution on PC.
-            // The code duplication in the threadprocs is unfortunate here, and hopefully we can settle on .NET's implementation on both platforms eventually.
-            List<Thread> ZipThreads;
-
-            if (Utils.IsRunningOnMono)
-            {
-                ZipThreads = (
-                    from CoreNum in Enumerable.Range(0, bZipInParallel ? Environment.ProcessorCount : 1)
-                    let ZipFileName = Path.Combine(StagingDir ?? OutputDir, string.Format("{0}{1}.zip", ZipBasename, bZipInParallel ? "-" + CoreNum.ToString("00") : ""))
-                    select new Thread(() =>
-                    {
-                        // don't create the zip unless we have at least one file to add
-                        string File;
-                        if (FilesToZip.TryDequeue(out File))
-                        {
-                            // Create one zip per thread using the given basename
-                            using (var ZipArchive = new Ionic.Zip.ZipFile(ZipFileName) { CompressionLevel = Ionic.Zlib.CompressionLevel.BestSpeed })
-                            {
-
-                                // pull from the queue until we are out of files.
-                                do
-                                {
-                                    // use fastest compression. In our best case we are CPU bound, so this is a good tradeoff,
-                                    // cutting overall time by 2/3 while only modestly increasing the compression ratio (22.7% -> 23.8% for RootEditor PDBs).
-                                    // This is in cases of a super hot cache, so the operation was largely CPU bound.
-                                    ZipArchive.AddFile(File, CommandUtils.ConvertSeparators(PathSeparator.Slash, Path.GetDirectoryName(CommandUtils.StripBaseDirectory(File, RootDir))));
-                                } while (FilesToZip.TryDequeue(out File));
-                                ZipArchive.Save();
-                            }
-                            Interlocked.Add(ref ZipFilesTotalSize, new FileInfo(ZipFileName).Length);
-                            // if we are using a staging dir, copy to the final location and delete the staged copy.
-                            if (StagingDir != null)
-                            {
-                                CommandUtils.CopyFile(ZipFileName, CommandUtils.MakeRerootedFilePath(ZipFileName, StagingDir, OutputDir));
-                                CommandUtils.DeleteFile(true, ZipFileName);
-                            }
-                        }
-                    })).ToList();
-            }
-            else
-            {
-                ZipThreads = (
-                    from CoreNum in Enumerable.Range(0, bZipInParallel ? Environment.ProcessorCount : 1)
-                    let ZipFileName = Path.Combine(StagingDir ?? OutputDir, string.Format("{0}{1}.zip", ZipBasename, bZipInParallel ? "-" + CoreNum.ToString("00") : ""))
-                    select new Thread(() =>
-                    {
-                        // don't create the zip unless we have at least one file to add
-                        string File;
-                        if (FilesToZip.TryDequeue(out File))
-                        {
-                            // Create one zip per thread using the given basename
-                            using (var ZipArchive = System.IO.Compression.ZipFile.Open(ZipFileName, System.IO.Compression.ZipArchiveMode.Create))
-                            {
-
-                                // pull from the queue until we are out of files.
-                                do
-                                {
-                                    // use fastest compression. In our best case we are CPU bound, so this is a good tradeoff,
-                                    // cutting overall time by 2/3 while only modestly increasing the compression ratio (22.7% -> 23.8% for RootEditor PDBs).
-                                    // This is in cases of a super hot cache, so the operation was largely CPU bound.
-                                    // Also, sadly, mono appears to have a bug where nothing you can do will properly set the LastWriteTime on the created entry,
-                                    // so we have to ignore timestamps on files extracted from a zip, since it may have been created on a Mac.
-                                    ZipArchive.CreateEntryFromFile(File, CommandUtils.ConvertSeparators(PathSeparator.Slash, CommandUtils.StripBaseDirectory(File, RootDir)), System.IO.Compression.CompressionLevel.Fastest);
-                                } while (FilesToZip.TryDequeue(out File));
-                            }
-                            Interlocked.Add(ref ZipFilesTotalSize, new FileInfo(ZipFileName).Length);
-                            // if we are using a staging dir, copy to the final location and delete the staged copy.
-                            if (StagingDir != null)
-                            {
-                                CommandUtils.CopyFile(ZipFileName, CommandUtils.MakeRerootedFilePath(ZipFileName, StagingDir, OutputDir));
-                                CommandUtils.DeleteFile(true, ZipFileName);
-                            }
-                        }
-                    })).ToList();
-            }
-            ZipThreads.ForEach(thread => thread.Start());
-            ZipThreads.ForEach(thread => thread.Join());
-            return new ParallelZipResult(ZipTimer.ElapsedTime, ZipFilesTotalSize);
-        }
-
-        /// <summary>
-        /// Unzips a set of zip files with a given basename in a given folder to a given RootDir.
-        /// </summary>
-        /// <param name="RootDir">Root Directory where all files will be extracted.</param>
-        /// <param name="FolderWithZipFiles">Folder containing the zip files to unzip. None of the zips should have the same file path in them.</param>
-        /// <param name="ZipBasename">The basename of the set of zip files to unzip.</param>
-        /// <returns>Some metrics about the unzip process.</returns>
-        /// <remarks>
-        /// The code is expected to be the used as the symmetrical inverse of <see cref="ParallelZipFiles"/>, but could be used independently, as long as the files in the zip do not overlap.
-        /// </remarks>
-        private static ParallelZipResult ParallelUnZipFiles(string RootDir, string FolderWithZipFiles, string ZipBasename)
-        {
-            var UnzipTimer = DateTimeStopwatch.Start();
-            long ZipFilesTotalSize = 0L;
-            // Sadly, mono implemention of System.IO.Compression is really poor (as of 2015/Aug), causing OOM when parallel zipping a large set of files.
-            // However, Ionic is MUCH slower than .NET's native implementation (2x+ slower in our build farm), so we stick to the faster solution on PC.
-            // The code duplication in the threadprocs is unfortunate here, and hopefully we can settle on .NET's implementation on both platforms eventually.
-            if (Utils.IsRunningOnMono)
-            {
-                Parallel.ForEach(Directory.EnumerateFiles(FolderWithZipFiles, ZipBasename + "*.zip").ToList(),
-                    (ZipFilename) =>
-                    {
-                        Interlocked.Add(ref ZipFilesTotalSize, new FileInfo(ZipFilename).Length);
-                        using (var ZipArchive = Ionic.Zip.ZipFile.Read(ZipFilename))
-                        {
-                            ZipArchive.ExtractAll(RootDir, Ionic.Zip.ExtractExistingFileAction.OverwriteSilently);
-                        }
-                    });
-            }
-            else
-            {
-                Parallel.ForEach(Directory.EnumerateFiles(FolderWithZipFiles, ZipBasename + "*.zip").ToList(),
-                    (ZipFilename) =>
-                    {
-                        Interlocked.Add(ref ZipFilesTotalSize, new FileInfo(ZipFilename).Length);
-                        // unzip the files manually instead of caling ZipFile.ExtractToDirectory() because we need to overwrite readonly files. Because of this, creating the directories is up to us as well.
-                        using (var ZipArchive = System.IO.Compression.ZipFile.OpenRead(ZipFilename))
-                        {
-                            foreach (var Entry in ZipArchive.Entries)
-                            {
-                                // Use CommandUtils.CombinePaths to ensure directory separators get converted correctly. On mono on *nix, if the path has backslashes it will not convert it.
-                                var ExtractedFilename = CommandUtils.CombinePaths(RootDir, Entry.FullName);
-                                // Zips can contain empty dirs. Ours usually don't have them, but we should support it.
-                                if (Path.GetFileName(ExtractedFilename).Length == 0)
-                                {
-                                    Directory.CreateDirectory(ExtractedFilename);
-                                }
-                                else
-                                {
-                                    // We must delete any existing file, even if it's readonly. .Net does not do this by default.
-                                    if (File.Exists(ExtractedFilename))
-                                    {
-                                        InternalUtils.SafeDeleteFile(ExtractedFilename, true);
-                                    }
-                                    else
-                                    {
-                                        Directory.CreateDirectory(Path.GetDirectoryName(ExtractedFilename));
-                                    }
-                                    Entry.ExtractToFile(ExtractedFilename, true);
-                                }
-                            }
-                        }
-                    });
-            }
-            return new ParallelZipResult(UnzipTimer.ElapsedTime, ZipFilesTotalSize);
-        }
-
-        /// <summary>
-        /// Saves the list of fully qualified files (that should be rooted at the shared temp storage location for the game) to a shared temp storage manifest with the given temp storage node and game.
-        /// </summary>
-        /// <param name="TempStorageNodeInfo">Node info descibing the block of temp storage (essentially used to identify a subdirectory insides the game's temp storage folder).</param>
-        /// <param name="Files">Fully qualified names of files to reference in the manifest file.</param>
-        /// <param name="bLocalOnly">If true, only ensures that the local temp storage manifest exists. Otherwise, checks if the manifest exists in either local or shared storage.</param>
-        /// <param name="RootDir">Folder that all the given files are rooted from. If null or empty, CmdEnv.LocalRoot is used.</param>
-        /// <returns>The created manifest instance (which has already been saved to disk).</returns>
-        public static void StoreToTempStorage(TempStorageNodeInfo TempStorageNodeInfo, List<string> Files, bool bLocalOnly, string RootDir)
-        {
-            using (var TelemetryStopwatch = new TelemetryStopwatch("StoreToTempStorage"))
-            {
-                // use LocalRoot if one is not specified
-                if (string.IsNullOrEmpty(RootDir))
-                {
-                    RootDir = CommandUtils.CmdEnv.LocalRoot;
-                }
+		/// <summary>
+		/// File extension for manifests in temp storage.
+		/// </summary>
+		private const string TempManifestExtension = ".TempManifest";
 
-                // save the manifest to local temp storage.
-                var Local = SaveLocalTempStorageManifest(RootDir, TempStorageNodeInfo, Files);
-                var LocalTotalSize = Local.GetTotalSize();
-                if (bLocalOnly)
-                {
-                    TelemetryStopwatch.Finish(string.Format("StoreToTempStorage.{0}.{1}.{2}.Local.{3}.{4}.{5}", Files.Count, LocalTotalSize, 0L, 0L, 0L, TempStorageNodeInfo.NodeStorageName));
-                }
-                else
-                {
-                    var SharedStorageNodeDir = SharedTempStorageDirectory(TempStorageNodeInfo);
-                    CommandUtils.Log("Storing to {0}", SharedStorageNodeDir);
-                    // this folder should not already exist, else we have concurrency or job duplication problems.
-                    if (CommandUtils.DirectoryExists_NoExceptions(SharedStorageNodeDir))
-                    {
-                        throw new AutomationException("Storage Block Already Exists! {0}", SharedStorageNodeDir);
-                    }
-                    CommandUtils.CreateDirectory(true, SharedStorageNodeDir);
-
-                    var LocalManifestFilename = LocalTempStorageManifestFilename(TempStorageNodeInfo);
-                    var SharedManifestFilename = SharedTempStorageManifestFilename(TempStorageNodeInfo);
-                    var StagingDir = Path.GetDirectoryName(LocalManifestFilename);
-                    var ZipBasename = Path.GetFileNameWithoutExtension(LocalManifestFilename);
-                    // initiate the parallel zip operation.
-                    var ZipResult = ParallelZipFiles(Files, RootDir, SharedStorageNodeDir, StagingDir, ZipBasename);
-
-                    // copy the local manifest to the shared location. We have to assume the zip is a good copy.
-                    CommandUtils.CopyFile(LocalManifestFilename, SharedManifestFilename);
-                    TelemetryStopwatch.Finish(string.Format("StoreToTempStorage.{0}.{1}.{2}.Remote.{3}.{4}.{5}", Files.Count, LocalTotalSize, ZipResult.ZipFilesTotalSize, 0, (long)ZipResult.ZipTime.TotalMilliseconds, TempStorageNodeInfo.NodeStorageName));
-                }
-            }
-        }
-
-        /// <summary>
-        /// Inverse of <see cref="StoreToTempStorage"/>.
-        /// Copies a block of files from a temp storage location given by a temp storage node and game to local storage rooted at the given root dir. 
-        /// If the temp manifest for this block is found locally, the copy is skipped, as we assume this is the same machine that created the temp storage and the files are still there.
-        /// </summary>
-        /// <param name="TempStorageNodeInfo">Node info descibing the block of temp storage (essentially used to identify a subdirectory insides the game's temp storage folder).</param>
-        /// <param name="RootDir">Folder that all the retrieved files should be rooted from. If null or empty, CmdEnv.LocalRoot is used.</param>
-        /// <returns>List of fully qualified paths to all the files that were retrieved. This is returned even if we skip the copy (set WasLocal = true) .</returns>
-        public static List<string> RetrieveFromTempStorage(TempStorageNodeInfo TempStorageNodeInfo, string RootDir)
-        {
-            using (var TelemetryStopwatch = new TelemetryStopwatch("RetrieveFromTempStorage"))
-            {
-                // use LocalRoot if one is not specified
-                if (String.IsNullOrEmpty(RootDir))
-                {
-                    RootDir = CommandUtils.CmdEnv.LocalRoot;
-                }
+		/// <summary>
+		/// Legayc code to clean all temp storage data from the given directory that is older than a certain threshold (defined directly in the code).
+		/// </summary>
+		/// <param name="TopDirectory">Fully qualified path of the folder that is the root of a bunch of temp storage folders. Should be a string returned by <see cref="ResolveSharedTempStorageDirectory"/> </param>
+		private static void CleanSharedTempStorageLegacy(string TopDirectory)
+		{
+			// This is a hack to clean out the old temp storage in the old folder name.
+			TopDirectory = TopDirectory.Replace(Path.DirectorySeparatorChar + TempStorageSubdirectoryName + Path.DirectorySeparatorChar, Path.DirectorySeparatorChar + "GUBP" + Path.DirectorySeparatorChar);
+
+			using (var TelemetryStopwatch = new TelemetryStopwatch("CleanSharedTempStorageLegacy"))
+			{
+				const int MaximumDaysToKeepLegacyTempStorage = 1;
+				var StartTimeDir = DateTime.UtcNow;
+				DirectoryInfo DirInfo = new DirectoryInfo(TopDirectory);
+				var TopLevelDirs = DirInfo.GetDirectories();
+				{
+					var BuildDuration = (DateTime.UtcNow - StartTimeDir).TotalMilliseconds;
+					CommandUtils.Log("Took {0}s to enumerate {1} directories.", BuildDuration / 1000, TopLevelDirs.Length);
+				}
+				foreach (var TopLevelDir in TopLevelDirs)
+				{
+					if (CommandUtils.DirectoryExists_NoExceptions(TopLevelDir.FullName))
+					{
+						bool bOld = false;
+						foreach (var ThisFile in CommandUtils.FindFiles_NoExceptions(true, "*.TempManifest", false, TopLevelDir.FullName))
+						{
+							FileInfo Info = new FileInfo(ThisFile);
+
+							if ((DateTime.UtcNow - Info.LastWriteTimeUtc).TotalDays > MaximumDaysToKeepLegacyTempStorage)
+							{
+								bOld = true;
+							}
+						}
+						if (bOld)
+						{
+							CommandUtils.Log("Deleting temp storage directory {0}, because it is more than {1} days old.", TopLevelDir.FullName, MaximumDaysToKeepLegacyTempStorage);
+							var StartTime = DateTime.UtcNow;
+							try
+							{
+								if (Directory.Exists(TopLevelDir.FullName))
+								{
+									// try the direct approach first
+									Directory.Delete(TopLevelDir.FullName, true);
+								}
+							}
+							catch
+							{
+							}
+							CommandUtils.DeleteDirectory_NoExceptions(true, TopLevelDir.FullName);
+							var BuildDuration = (DateTime.UtcNow - StartTime).TotalMilliseconds;
+							CommandUtils.Log("Took {0}s to delete {1}.", BuildDuration / 1000, TopLevelDir.FullName);
+						}
+					}
+				}
+			}
+		}
 
-                // First see if the local manifest is there.
-                // If it is, then we must be on the same node as the one that originally created the temp storage.
-                // In that case, we just verify all the files exist as described in the manifest and use that.
-                // If there was any tampering, abort immediately because we never accept tampering, and it signifies a build problem.
-                var LocalManifest = LocalTempStorageManifestFilename(TempStorageNodeInfo);
-                if (CommandUtils.FileExists_NoExceptions(LocalManifest))
-                {
-                    CommandUtils.Log("Found local manifest {0}", LocalManifest);
-                    var Local = TempStorageManifest.Load(LocalManifest);
-                    var Files = Local.GetFiles(RootDir);
-                    var LocalTest = TempStorageManifest.Create(Files, RootDir);
-                    if (!Local.Compare(LocalTest, true))
-                    {
-                        throw new AutomationException("Local files in manifest {0} were tampered with.", LocalManifest);
-                    }
-                    TelemetryStopwatch.Finish(string.Format("RetrieveFromTempStorage.{0}.{1}.{2}.Local.{3}.{4}.{5}", Files.Count, Local.GetTotalSize(), 0L, 0L, 0L, TempStorageNodeInfo.NodeStorageName));
-                    return Files;
-                }
+		/// <summary>
+		/// Cleans all temp storage data from the given directory that is older than a certain threshold (defined directly in the code).
+		/// This will clean up builds that are NOT part of the the job for this branch. This is intentional because once a branch is shut down,
+		/// there would be nothing left to clean it. This will only clean the supplied folder, but it is assumed this folder is an actual
+		/// Temp Storage folder with .TempManifest files in it.
+		/// </summary>
+		/// <param name="TopDirectory">Fully qualified path of the folder that is the root of a bunch of temp storage folders. Should be a string returned by <see cref="ResolveSharedTempStorageDirectory"/> </param>
+		private static void CleanSharedTempStorage(string TopDirectory)
+		{
+			CleanSharedTempStorageLegacy(TopDirectory);
+
+			Action<string> TryToDeletePossiblyEmptyFolder = (string FolderName) =>
+				{
+					try
+					{
+						Directory.Delete(FolderName);
+					}
+					catch (IOException)
+					{
+						// only catch "directory is not empty type exceptions, if possible. Best we can do is check for IOException.
+					}
+					catch (Exception Ex)
+					{
+						CommandUtils.LogWarning("Unexpected failure trying to delete possibly empty temp folder {0}: {1}", FolderName, Ex);
+					}
+				};
+
+			CommandUtils.Log("Cleaning temp storage for {0}...", TopDirectory);
+			int FoldersCleaned = 0;
+			int FoldersManuallyCleaned = 0;
+			int FoldersFailedCleaned = 0;
+			using (var TelemetryStopwatch = new TelemetryStopwatch("CleanSharedTempStorage"))
+			{
+				const double MaximumDaysToKeepTempStorage = 3;
+				var Now = DateTime.UtcNow;
+				foreach (var OldManifestInfo in
+					// First subdirectory is the branch name
+					from BranchDir in Directory.EnumerateDirectories(TopDirectory)
+					// Second subdirectory is the changelist
+					from CLDir in Directory.EnumerateDirectories(BranchDir)
+					// third subdirectory is the node storage name
+					from NodeNameDir in Directory.EnumerateDirectories(CLDir)
+					// now look for manifest files.
+					let ManifestFile = Directory.EnumerateFiles(NodeNameDir, "*.TempManifest").SingleOrDefault()
+					where ManifestFile != null
+					// only choose ones that are old enough.
+					where (Now - new FileInfo(ManifestFile).LastWriteTimeUtc).TotalDays > MaximumDaysToKeepTempStorage
+					select new { BranchDir, CLDir, NodeNameDir })
+				{
+					try
+					{
+						CommandUtils.Log("Deleting folder with old temp storage {0}...", OldManifestInfo.NodeNameDir);
+						Directory.Delete(OldManifestInfo.NodeNameDir, true);
+						FoldersCleaned++;
+					}
+					catch (Exception Ex)
+					{
+						CommandUtils.LogWarning("Failed to delete old manifest folder '{0}', will try one file at a time: {1}", OldManifestInfo.NodeNameDir, Ex);
+						if (CommandUtils.DeleteDirectory_NoExceptions(true, OldManifestInfo.NodeNameDir))
+						{
+							FoldersManuallyCleaned++;
+						}
+						else
+						{
+							FoldersFailedCleaned++;
+						}
+					}
+					// Once we are done, try to delete the CLDir and BranchDir (will fail if not empty, so no worries).
+					TryToDeletePossiblyEmptyFolder(OldManifestInfo.CLDir);
+					TryToDeletePossiblyEmptyFolder(OldManifestInfo.BranchDir);
+				}
+				TelemetryStopwatch.Finish(string.Format("CleanSharedTempStorage.{0}.{1}.{2}", FoldersCleaned, FoldersManuallyCleaned, FoldersFailedCleaned));
+			}
+		}
 
-                // We couldn't find the node storage locally, so get it from the shared location.
-                var SharedStorageNodeDir = SharedTempStorageDirectory(TempStorageNodeInfo);
+		/// <summary>
+		/// Cleans the shared temp storage folder for the given root name. This code will ensure that any duplicate games
+		/// or games that resolve to use the same temp folder are not cleaned twice.
+		/// </summary>
+		/// <param name="RootNameForTempStorage">List of game names to clean temp folders for.</param>
+		public static void CleanSharedTempStorageDirectory(string RootNameForTempStorage)
+		{
+			if (!CommandUtils.IsBuildMachine || Utils.IsRunningOnMono)  // saw a hang on this, anyway it isn't necessary to clean with macs, they are slow anyway
+			{
+				return;
+			}
+			var TempStorageFolder = ResolveSharedTempStorageDirectory(RootNameForTempStorage);
+			try
+			{
+				CleanSharedTempStorage(TempStorageFolder);
+			}
+			catch (Exception Ex)
+			{
+				CommandUtils.LogWarning("Unable to Clean Temp Directory {0}. Exception: {1}", TempStorageFolder, Ex);
+			}
+		}
 
-                CommandUtils.Log("Attempting to retrieve from {0}", SharedStorageNodeDir);
-                if (!CommandUtils.DirectoryExists_NoExceptions(SharedStorageNodeDir))
-                {
-                    throw new AutomationException("Storage Block Does Not Exists! {0}", SharedStorageNodeDir);
-                }
-                var SharedManifest = SharedTempStorageManifestFilename(TempStorageNodeInfo);
-                InternalUtils.Robust_FileExists(SharedManifest, "Storage Block Manifest Does Not Exists! {0}");
-
-                var Shared = TempStorageManifest.Load(SharedManifest);
-                var SharedFiles = Shared.GetFiles(SharedStorageNodeDir);
-
-                // We know the source files exist and are under RootDir because we created the manifest, which verifies it.
-                // Now create the list of target files
-                var DestFiles = SharedFiles.Select(Filename =>
-                {
-                    var DestFile = CommandUtils.MakeRerootedFilePath(Filename, SharedStorageNodeDir, RootDir);
-                    // create a FileInfo using the file, which will help us catch path too long exceptions early and propagate a friendly error containing the full path.
-                    try
-                    {
-                        return new FileInfo(DestFile).FullName;
-                    }
-                    catch (PathTooLongException Ex)
-                    {
-                        throw new AutomationException(Ex, "Path too long ... failed to create FileInfo for {0}", DestFile);
-                    }
-                }).ToList();
-
-                var ZipBasename = Path.GetFileNameWithoutExtension(LocalManifest);
-
-                // now unzip in parallel, overwriting any existing local file.
-                var ZipResult = ParallelUnZipFiles(RootDir, SharedStorageNodeDir, ZipBasename);
-
-                var NewLocal = SaveLocalTempStorageManifest(RootDir, TempStorageNodeInfo, DestFiles);
-                // Now compare the created local files to ensure their attributes match the one we copied from the network.
-                if (!NewLocal.Compare(Shared, true))
-                {
-                    // we will rename this so it can't be used, but leave it around for inspection
-                    CommandUtils.RenameFile_NoExceptions(LocalManifest, LocalManifest + ".broken");
-                    throw new AutomationException("Shared and Local manifest mismatch.");
-                }
+		/// "TmpStore", the folder name where all GUBP temp storage files are added inside a game's network storage location.
+		private const string TempStorageSubdirectoryName = "TmpStore";
 
-                // Handle unix permissions/chmod issues. This will touch the timestamp we check on the file, so do this after we've compared with the manifest attributes.
-                if (Utils.IsRunningOnMono)
-                {
-                    foreach (string DestFile in DestFiles)
-                    {
-                        CommandUtils.FixUnixFilePermissions(DestFile);
-                    }
-                }
+		/// <summary>
+		/// Determines if the root UE4 temp storage directory exists.
+		/// </summary>
+		/// <param name="ForSaving">If true, confirm that the directory exists and is writeable.</param>
+		/// <returns>true if the UE4 temp storage directory exists (and is writeable if requested).</returns>
+		public static bool IsSharedTempStorageAvailable(bool ForSaving)
+		{
+			var UE4TempStorageDir = CommandUtils.CombinePaths(CommandUtils.RootBuildStorageDirectory(), "UE4", TempStorageSubdirectoryName);
+			if (ForSaving)
+			{
+				int Retries = 0;
+				while (Retries < 24)
+				{
+					if (InternalUtils.Robust_DirectoryExistsAndIsWritable_NoExceptions(UE4TempStorageDir))
+					{
+						return true;
+					}
+					CommandUtils.FindDirectories_NoExceptions(false, "*", false, UE4TempStorageDir); // there is some internet evidence to suggest this might perk up the mac share
+					Thread.Sleep(5000);
+					Retries++;
+				}
+			}
+			else if (InternalUtils.Robust_DirectoryExists_NoExceptions(UE4TempStorageDir, "Could not find {0}"))
+			{
+				return true;
+			}
+			return false;
+		}
 
-                TelemetryStopwatch.Finish(string.Format("RetrieveFromTempStorage.{0}.{1}.{2}.Remote.{3}.{4}.{5}", DestFiles.Count, Shared.GetTotalSize(), ZipResult.ZipFilesTotalSize, 0, (long)ZipResult.ZipTime.TotalMilliseconds, TempStorageNodeInfo.NodeStorageName));
-                return DestFiles;
-            }
-        }
-
-        /// <summary>
-        /// Finds temp storage manifests that match the block name (where the block name has embedded wildcards to find all block names of a given pattern).
-        /// This is used by <see cref="GUBP.FindNodeHistory"/> to search for any temp storage manifests that match a certain GUBP node name.
-        /// This is used to construct a "history" of that node which is then used to log info and generate failure emails.
-        /// This method is used multiple times to find nodes matching a specific name with a suffix of _Started _Failed and _Succeeded. 
-        /// The CL# is pulled from those names and used to generate a P4 history.
-        /// Main problem with this is that temp storage is ephemeral, so after a few days, this will typically not find any history for a node.
-        /// 
-        /// NOTE: This entire routine is dependent on how TempStorageNodeInfo lays out its directories!!!
-        /// If this layout is changed, this code needs to be changed as well.
-        /// @todo: Someday Make FindNodeHistory query the build database directly to get a true history of the node.
-        /// </summary>
-        /// <param name="TempStorageNodeInfo">Node info descibing the block of temp storage (essentially used to identify a subdirectory insides the game's temp storage folder).</param>
-        /// <returns></returns>
-        public static List<int> FindMatchingSharedTempStorageNodeCLs(TempStorageNodeInfo TempStorageNodeInfo)
-        {
-            // Find the shared temp storage folder for this game and branch
-            var SharedTempStorageDirectoryForGameAndBranch = CommandUtils.CombinePaths(ResolveSharedTempStorageDirectory(TempStorageNodeInfo.JobInfo.RootNameForTempStorage), TempStorageNodeInfo.JobInfo.BranchNameForTempStorage);
-            int dummy;
-            return (
-                // Look for all folders underneath this, it should be a CL, or CL-PreflightInfo
-                from CLDir in Directory.EnumerateDirectories(SharedTempStorageDirectoryForGameAndBranch)
-                // only accept folders that are plain numbers. Any suffixes ('-') imply preflight builds or something else.
-                let PossibleCL = Path.GetFileName(CLDir)
-                where int.TryParse(PossibleCL, out dummy)
-                let CL = int.Parse(PossibleCL)
-                // if we have a real CL, try to find the node name we are looking for.
-                where File.Exists(CommandUtils.CombinePaths(CLDir, TempStorageNodeInfo.NodeStorageName, TempStorageNodeInfo.GetManifestFilename()))
-                select CL)
-                .OrderBy(CL => CL).ToList();
-        }
-
-        /// <summary>
-        /// Runs a test of the temp storage system. This function is part of the class so it doesn't have to expose internals just to run the test.
-        /// </summary>
-        /// <param name="CmdEnv">The environment to use.</param>
-        static internal void TestTempStorage(CommandEnvironment CmdEnv)
-        {
-            // We are not a real GUBP job, so fake the values.
-            var TestTempStorageInfo = new TempStorageNodeInfo(new JobInfo("", "Test", "UE4", 0, 0, 0), "Test");
-
-            // Delete any local and shared temp storage that may exist.
-            DeleteLocalTempStorage();
-            DeleteSharedTempStorage(TestTempStorageInfo);
-            if (TempStorageExists(TestTempStorageInfo, false, false))
-            {
-                throw new AutomationException("storage should not exist");
-            }
-
-            // Create some test files in various locations in the local root with unique content.
-            var TestFiles = new[]
-                { 
-                    CommandUtils.CombinePaths(CmdEnv.LocalRoot, "Engine", "Build", "Batchfiles", "TestFile.Txt"),
-                    CommandUtils.CombinePaths(CmdEnv.LocalRoot, "TestFile2.Txt"),
-                    CommandUtils.CombinePaths(CmdEnv.LocalRoot, "engine", "plugins", "TestFile3.Txt"),
-                }
-                .Select(TestFile => new 
-                { 
-                    FileName = TestFile, 
-                    FileContents = string.Format("{0} - {1}", TestFile, DateTime.UtcNow) 
-                })
-                .ToList();
-            foreach (var TestFile in TestFiles)
-            {
-                CommandUtils.DeleteFile(TestFile.FileName);
-                CommandUtils.Log("Test file {0}", TestFile.FileName);
-                File.WriteAllText(TestFile.FileName, TestFile.FileContents);
-                // we should be able to overwrite readonly files.
-                File.SetAttributes(TestFile.FileName, FileAttributes.ReadOnly);
-            }
-            
-            // wrap the operation so we are sure to clean up afterward.
-            try
-            {
-                // Store the test file to temp storage.
-                StoreToTempStorage(TestTempStorageInfo, TestFiles.Select(TestFile => TestFile.FileName).ToList(), false, CmdEnv.LocalRoot);
-                // The manifest should exist locally.
-                if (!LocalTempStorageManifestExists(TestTempStorageInfo))
-                {
-                    throw new AutomationException("local storage should exist");
-                }
-                // The manifest should exist on the shared drive.
-                if (!SharedTempStorageManifestExists(TestTempStorageInfo, false))
-                {
-                    throw new AutomationException("shared storage should exist");
-                }
-                // Now delete the local manifest
-                DeleteLocalTempStorage();
-                // It should no longer be there.
-                if (LocalTempStorageManifestExists(TestTempStorageInfo))
-                {
-                    throw new AutomationException("local storage should not exist");
-                }
-                // But the shared storage should still exist.
-                if (!TempStorageExists(TestTempStorageInfo, false, false))
-                {
-                    throw new AutomationException("some storage should exist");
-                }
 
-                // Now we should be able to retrieve the test files from shared storage, and it should overwrite our read-only files.
-                RetrieveFromTempStorage(TestTempStorageInfo, CmdEnv.LocalRoot);
-                // Now delete the local manifest so we can try again with no files present (the usual case for restoring from temp storage).
-                DeleteLocalTempStorage();
+		/// <summary>
+		/// Cache of previously resolved build directories for games. Since this requires a few filesystem checks,
+		/// we cache the results so future lookups are fast.
+		/// </summary>
+		static readonly Dictionary<string, string> SharedBuildDirectoryResolveCache = new Dictionary<string, string>();
+
+		/// <summary>
+		/// Determines full path of the shared build folder for the given game.
+		/// If the shared build folder does not already exist, will try to use UE4's folder instead.
+		/// This is because the root build folders are independently owned by each game team and managed by IT, so cannot be created on the fly.
+		/// </summary>
+		/// The root name beneath <see cref="CommandUtils.RootBuildStorageDirectory"/> where all temp storage will placed for this job.
+		/// <returns>The full path of the shared build directory name, or throws an exception if none is found and the UE4 folder is not found either.</returns>
+		public static string ResolveSharedBuildDirectory(string RootNameForTempStorage)
+		{
+			if (RootNameForTempStorage == null) throw new ArgumentNullException("RootNameForTempStorage");
+
+			if (SharedBuildDirectoryResolveCache.ContainsKey(RootNameForTempStorage))
+			{
+				return SharedBuildDirectoryResolveCache[RootNameForTempStorage];
+			}
+			string Root = CommandUtils.RootBuildStorageDirectory();
+			string Result = CommandUtils.CombinePaths(Root, RootNameForTempStorage);
+			if (!InternalUtils.Robust_DirectoryExistsAndIsWritable_NoExceptions(Result))
+			{
+				throw new AutomationException("Could not find an appropriate shared temp folder {0}", Result);
+			}
+			SharedBuildDirectoryResolveCache.Add(RootNameForTempStorage, Result);
+			return Result;
+		}
 
-                // Ok, delete our test files locally.
-                foreach (var TestFile in TestFiles)
-                {
-                    CommandUtils.DeleteFile(TestFile.FileName);
-                }
+		/// <summary>
+		/// Returns the full path of the temp storage directory for the given game name (ie, P:\Builds\RootNameForTempStorage\TmpStore).
+		/// If the game's build folder does not exist, will use the temp storage folder in the UE4 build folder (ie, P:\Builds\UE4\TmpStore).
+		/// If the directory does not exist, we will try to create it.
+		/// </summary>
+		/// <param name="RootNameForTempStorage">The root name beneath <see cref="CommandUtils.RootBuildStorageDirectory"/> where all temp storage will placed for this job.</param>
+		/// <returns>Essentially adds the TmpStore/ subdirectory to the SharedBuildDirectory for the game.</returns>
+		public static string ResolveSharedTempStorageDirectory(string RootNameForTempStorage)
+		{
+			string Result = CommandUtils.CombinePaths(ResolveSharedBuildDirectory(RootNameForTempStorage), TempStorageSubdirectoryName);
+
+			if (!InternalUtils.Robust_DirectoryExists_NoExceptions(Result, "Could not find {0}"))
+			{
+				CommandUtils.CreateDirectory_NoExceptions(Result);
+			}
+			if (!InternalUtils.Robust_DirectoryExists_NoExceptions(Result, "Could not find {0}"))
+			{
+				throw new AutomationException("Could not create an appropriate shared temp folder {0} for {1}", Result, RootNameForTempStorage);
+			}
+			return Result;
+		}
 
-                // Now we should be able to retrieve the test files from shared storage.
-                RetrieveFromTempStorage(TestTempStorageInfo, CmdEnv.LocalRoot);
-                // the local manifest should be there, since we just retrieved from shared storage.
-                if (!LocalTempStorageManifestExists(TestTempStorageInfo))
-                {
-                    throw new AutomationException("local storage should exist");
-                }
-                // The shared manifest should also still be there.
-                if (!SharedTempStorageManifestExists(TestTempStorageInfo, false))
-                {
-                    throw new AutomationException("shared storage should exist");
-                }
-                // Verify the retrieved files have the correct content.
-                foreach (var TestFile in TestFiles)
-                {
-                    if (File.ReadAllText(TestFile.FileName) != TestFile.FileContents)
-                    {
-                        throw new AutomationException("Contents of the test file {0} was changed after restoring from shared temp storage.", TestFile.FileName);
-                    }
-                }
-                // Now delete the shared temp storage
-                DeleteSharedTempStorage(TestTempStorageInfo);
-                // Shared temp storage manifest should no longer exist.
-                if (SharedTempStorageManifestExists(TestTempStorageInfo, false))
-                {
-                    throw new AutomationException("shared storage should not exist");
-                }
-                // Retrieving temp storage should now just retrieve from local
-                RetrieveFromTempStorage(TestTempStorageInfo, CmdEnv.LocalRoot);
-                if (!LocalTempStorageManifestExists(TestTempStorageInfo))
-                {
-                    throw new AutomationException("local storage should exist");
-                }
+		/// <summary>
+		/// Returns the full path to the temp storage directory for the given temp storage node and game (ie, P:\Builds\RootNameForTempStorage\TmpStore\NodeInfoDirectory)
+		/// </summary>
+		/// <param name="SharedStorageDir">Base directory for temp storage in the current build</param>
+		/// <param name="BlockName">Name of the block of temp storage (essentially used to identify a subdirectory insides the game's temp storage folder).</param>
+		/// <returns>The full path to the temp storage directory for the given storage block name for the given game.</returns>
+		private static string SharedTempStorageDirectory(string SharedStorageDir, string BlockName)
+		{
+			return CommandUtils.CombinePaths(SharedStorageDir, BlockName);
+		}
 
-                // and now lets test tampering
-                DeleteLocalTempStorage();
-                {
-                    bool bFailedProperly = false;
-                    var MissingFile = new List<string>(TestFiles.Select(TestFile => TestFile.FileName));
-                    // add a file to the manifest that shouldn't be there.
-                    MissingFile.Add(CommandUtils.CombinePaths(CmdEnv.LocalRoot, "Engine", "SomeFileThatDoesntExist.txt"));
-                    try
-                    {
-                        StoreToTempStorage(TestTempStorageInfo, MissingFile, false, CmdEnv.LocalRoot);
-                    }
-                    catch (AutomationException)
-                    {
-                        bFailedProperly = true;
-                    }
-                    if (!bFailedProperly)
-                    {
-                        throw new AutomationException("Missing file did not fail.");
-                    }
-                }
+		/// <summary>
+		/// Gets the name of the local temp storage manifest file for the given temp storage node (ie, Engine\Saved\TmpStore\NodeInfoDirectory\NodeInfoFilename.TempManifest)
+		/// </summary>
+		/// <param name="BlockName">Name of the block of temp storage (essentially used to identify a subdirectory insides the game's temp storage folder).</param>
+		/// <returns>The name of the temp storage manifest file for the given storage block name for the given game.</returns>
+		private static string LocalTempStorageManifestFilename(string BlockName)
+		{
+			return CommandUtils.CombinePaths(LocalTempStorageManifestDirectory(), BlockName + TempManifestExtension);
+		}
 
-                // clear the shared temp storage again.
-                DeleteSharedTempStorage(TestTempStorageInfo);
-                // store the test files to shared temp storage again.
-                StoreToTempStorage(TestTempStorageInfo, TestFiles.Select(TestFile => TestFile.FileName).ToList(), false, CmdEnv.LocalRoot);
-                // force a load from shared by deleting the local manifest
-                DeleteLocalTempStorage();
-                // delete our test files locally.
-                foreach (var TestFile in TestFiles)
-                {
-                    CommandUtils.DeleteFile(TestFile.FileName);
-                }
+		/// <summary>
+		/// Saves the list of fully qualified files rooted at RootDir to a local temp storage manifest with the given temp storage node.
+		/// </summary>
+		/// <param name="RootDir">Folder that all the given files are rooted from.</param>
+		/// <param name="BlockName">Name of the block of temp storage (essentially used to identify a subdirectory insides the game's temp storage folder).</param>
+		/// <param name="Files">Fully qualified names of files to reference in the manifest file.</param>
+		/// <returns>The created manifest instance (which has already been saved to disk).</returns>
+		private static TempStorageManifest SaveLocalTempStorageManifest(string RootDir, string BlockName, List<string> Files)
+		{
+			string FinalFilename = LocalTempStorageManifestFilename(BlockName);
+			var Saver = TempStorageManifest.Create(Files, RootDir);
+			CommandUtils.CreateDirectory(true, Path.GetDirectoryName(FinalFilename));
+			Saver.Save(FinalFilename);
+			return Saver;
+		}
 
-                // now test that retrieving from shared temp storage properly balks that a file is missing.
-                {
-                    // tamper with the shared files.
-                    var RandomSharedZipFile = Directory.EnumerateFiles(SharedTempStorageDirectory(TestTempStorageInfo), "*.zip").First();
-                    // delete the shared file.
-                    CommandUtils.DeleteFile(RandomSharedZipFile);
-
-                    bool bFailedProperly = false;
-                    try
-                    {
-                        RetrieveFromTempStorage(TestTempStorageInfo, CmdEnv.LocalRoot);
-                    }
-                    catch (AutomationException)
-                    {
-                        bFailedProperly = true;
-                    }
-                    if (!bFailedProperly)
-                    {
-                        throw new AutomationException("Did not fail to load from missing file.");
-                    }
-                }
-                // recreate our temp files.
-                foreach (var TestFile in TestFiles)
-                {
-                    File.WriteAllText(TestFile.FileName, TestFile.FileContents);
-                }
+		/// <summary>
+		/// Gets the name of the temp storage manifest file for the given temp storage node and game (ie, P:\Builds\RootNameForTempStorage\TmpStore\NodeInfoDirectory\NodeInfoFilename.TempManifest)
+		/// </summary>
+		/// <param name="SharedStorageDir">Base directory for temp storage in the current build</param>
+		/// <param name="BlockName">Name of the block of temp storage (essentially used to identify a subdirectory insides the game's temp storage folder).</param>
+		/// <returns>The name of the temp storage manifest file for the given storage block name for the given game.</returns>
+		private static string SharedTempStorageManifestFilename(string SharedStorageDir, string BlockName)
+		{
+			return CommandUtils.CombinePaths(SharedTempStorageDirectory(SharedStorageDir, BlockName), BlockName + TempManifestExtension);
+		}
 
-                // clear the shared temp storage again.
-                DeleteSharedTempStorage(TestTempStorageInfo);
-                // Copy the files to temp storage.
-                StoreToTempStorage(TestTempStorageInfo, TestFiles.Select(TestFile => TestFile.FileName).ToList(), false, CmdEnv.LocalRoot);
-                // Delete a local file.
-                CommandUtils.DeleteFile(TestFiles[0].FileName);
-                // retrieving from temp storage should use WasLocal, but should balk because a local file was deleted.
-                {
-                    bool bFailedProperly = false;
-                    try
-                    {
-                        RetrieveFromTempStorage(TestTempStorageInfo, CmdEnv.LocalRoot);
-                    }
-                    catch (AutomationException)
-                    {
-                        bFailedProperly = true;
-                    }
-                    if (!bFailedProperly)
-                    {
-                        throw new AutomationException("Did not fail to load from missing local file.");
-                    }
-                }
-            }
-            finally
-            {
-                // Do a final cleanup.
-                DeleteSharedTempStorage(TestTempStorageInfo);
-                DeleteLocalTempStorage();
-                foreach (var TestFile in TestFiles)
-                {
-                    CommandUtils.DeleteFile(TestFile.FileName);
+		/// <summary>
+		/// Deletes all temp storage manifests from the local storage location (Engine\Saved\TmpStore).
+		/// Local temp storage logic only stores manifests, so this effectively deletes any local temp storage work, while not actually deleting the local files, which the temp 
+		/// storage system doesn't really own.
+		/// </summary>
+		public static void DeleteLocalTempStorage()
+		{
+			CommandUtils.DeleteDirectory(true, LocalTempStorageManifestDirectory());
+		}
+
+		/// <summary>
+		/// Deletes all shared temp storage (file and manifest) for the given temp storage node and game.
+		/// </summary>
+		/// <param name="SharedStorageDir">Base directory for temp storage in the current build</param>
+		/// <param name="BlockName">Name of the block of temp storage (essentially used to identify a subdirectory insides the game's temp storage folder).</param>
+		private static void DeleteSharedTempStorage(string SharedStorageDir, string BlockName)
+		{
+			CommandUtils.DeleteDirectory(true, SharedTempStorageDirectory(SharedStorageDir, BlockName));
+		}
+
+		/// <summary>
+		/// Checks if a local temp storage manifest exists for the given temp storage node.
+		/// </summary>
+		/// <param name="BlockName">Name of the block of temp storage (essentially used to identify a subdirectory insides the game's temp storage folder).</param>
+		/// <param name="bQuiet">True to suppress logging during the operation.</param>
+		/// <returns>true if the local temp storage manifest file exists</returns>
+		public static bool LocalTempStorageManifestExists(string BlockName, bool bQuiet = false)
+		{
+			var LocalManifest = LocalTempStorageManifestFilename(BlockName);
+			return CommandUtils.FileExists_NoExceptions(bQuiet, LocalManifest);
+		}
+
+		/// <summary>
+		/// Deletes a local temp storage manifest file for the given temp storage node if it exists.
+		/// </summary>
+		/// <param name="BlockName">Name of the block of temp storage (essentially used to identify a subdirectory insides the game's temp storage folder).</param>
+		/// <param name="bQuiet">True to suppress logging during the operation.</param>
+		public static void DeleteLocalTempStorageManifest(string BlockName, bool bQuiet = false)
+		{
+			var LocalManifest = LocalTempStorageManifestFilename(BlockName);
+			CommandUtils.DeleteFile(bQuiet, LocalManifest);
+		}
+
+		/// <summary>
+		/// Checks if a shared temp storage manifest exists for the given temp storage node and game.
+		/// </summary>
+		/// <param name="SharedStorageDir">Base directory for temp storage in the current build</param>
+		/// <param name="BlockName">Name of the block of temp storage (essentially used to identify a subdirectory insides the game's temp storage folder).</param>
+		/// <param name="bQuiet">True to suppress logging during the operation.</param>
+		/// <returns>true if the shared temp storage manifest file exists</returns>
+		public static bool SharedTempStorageManifestExists(string SharedStorageDir, string BlockName, bool bQuiet)
+		{
+			var SharedManifest = SharedTempStorageManifestFilename(SharedStorageDir, BlockName);
+			if (CommandUtils.FileExists_NoExceptions(bQuiet, SharedManifest))
+			{
+				return true;
+			}
+			return false;
+		}
+
+		/// <summary>
+		/// Checks if a temp storage manifest exists, either locally or in the shared folder (depending on the value of bLocalOnly) for the given temp storage node and game.
+		/// </summary>
+		/// <param name="BlockName">Name of the block of temp storage (essentially used to identify a subdirectory insides the game's temp storage folder).</param>
+		/// <param name="SharedStorageDir">The root directory for temp storage shared by different agents; typically on a network share. May be null.</param>
+		/// <param name="bQuiet">True to suppress logging during the operation.</param>
+		/// <returns>true if the shared temp storage manifest file exists</returns>
+		public static bool TempStorageExists(string BlockName, string SharedStorageDir, bool bQuiet)
+		{
+			return LocalTempStorageManifestExists(BlockName, bQuiet) || (SharedStorageDir != null && SharedTempStorageManifestExists(SharedStorageDir, BlockName, bQuiet));
+		}
+
+		/// <summary>
+		/// Stores some result info from a parallel zip or unzip operation.
+		/// </summary>
+		class ParallelZipResult
+		{
+			/// <summary>
+			/// Time taken for the zip portion of the operation (actually, everything BUT the copy time).
+			/// </summary>
+			public readonly TimeSpan ZipTime;
+
+			/// <summary>
+			/// Total size of the zip files that are created.
+			/// </summary>
+			public readonly long ZipFilesTotalSize;
+
+			public ParallelZipResult(TimeSpan ZipTime, long ZipFilesTotalSize)
+			{
+				this.ZipTime = ZipTime;
+				this.ZipFilesTotalSize = ZipFilesTotalSize;
+			}
+		}
+
+		/// <summary>
+		/// Zips a set of files (that must be rooted at the given RootDir) to a set of zip files in the given OutputDir. The files will be prefixed with the given basename.
+		/// </summary>
+		/// <param name="Files">Fully qualified list of files to zip (must be rooted at RootDir).</param>
+		/// <param name="RootDir">Root Directory where all files will be extracted.</param>
+		/// <param name="OutputDir">Location to place the set of zip files created.</param>
+		/// <param name="StagingDir">Location to create zip files before copying them to the OutputDir. If the OutputDir is on a remote file share, staging may be more efficient. Use null to avoid using a staging copy.</param>
+		/// <param name="ZipBasename">The basename of the set of zip files.</param>
+		/// <returns>Some metrics about the zip process.</returns>
+		/// <remarks>
+		/// This function tries to zip the files in parallel as fast as it can. It makes no guarantees about how many zip files will be created or which files will be in which zip,
+		/// but it does try to reasonably balance the file sizes.
+		/// </remarks>
+		private static ParallelZipResult ParallelZipFiles(IEnumerable<string> Files, string RootDir, string OutputDir, string StagingDir, string ZipBasename)
+		{
+			var ZipTimer = DateTimeStopwatch.Start();
+			// First get the sizes of all the files. We won't parallelize if there isn't enough data to keep the number of zips down.
+			var FilesInfo = Files
+				.Select(File => new { File, FileSize = new FileInfo(File).Length })
+				.ToList();
+
+			// Profiling results show that we can zip 100MB quite fast and it is not worth parallelizing that case and creating a bunch of zips that are relatively small.
+			const long MinFileSizeToZipInParallel = 1024 * 1024 * 100L;
+			var bZipInParallel = FilesInfo.Sum(FileInfo => FileInfo.FileSize) >= MinFileSizeToZipInParallel;
+
+			// order the files in descending order so our threads pick up the biggest ones first.
+			// We want to end with the smaller files to more effectively fill in the gaps
+			var FilesToZip = new ConcurrentQueue<string>(FilesInfo.OrderByDescending(FileInfo => FileInfo.FileSize).Select(FileInfo => FileInfo.File));
+
+			long ZipFilesTotalSize = 0L;
+			// We deliberately avoid Parallel.ForEach here because profiles have shown that dynamic partitioning creates
+			// too many zip files, and they can be of too varying size, creating uneven work when unzipping later,
+			// as ZipFile cannot unzip files in parallel from a single archive.
+			// We can safely assume the build system will not be doing more important things at the same time, so we simply use all our logical cores,
+			// which has shown to be optimal via profiling, and limits the number of resulting zip files to the number of logical cores.
+			// 
+			// Sadly, mono implemention of System.IO.Compression is really poor (as of 2015/Aug), causing OOM when parallel zipping a large set of files.
+			// However, Ionic is MUCH slower than .NET's native implementation (2x+ slower in our build farm), so we stick to the faster solution on PC.
+			// The code duplication in the threadprocs is unfortunate here, and hopefully we can settle on .NET's implementation on both platforms eventually.
+			List<Thread> ZipThreads;
+
+			if (Utils.IsRunningOnMono)
+			{
+				ZipThreads = (
+					from CoreNum in Enumerable.Range(0, bZipInParallel ? Environment.ProcessorCount : 1)
+					let ZipFileName = Path.Combine(StagingDir ?? OutputDir, string.Format("{0}{1}.zip", ZipBasename, bZipInParallel ? "-" + CoreNum.ToString("00") : ""))
+					select new Thread(() =>
+					{
+						// don't create the zip unless we have at least one file to add
+						string File;
+						if (FilesToZip.TryDequeue(out File))
+						{
+							// Create one zip per thread using the given basename
+							using (var ZipArchive = new Ionic.Zip.ZipFile(ZipFileName) { CompressionLevel = Ionic.Zlib.CompressionLevel.BestSpeed })
+							{
+
+								// pull from the queue until we are out of files.
+								do
+								{
+									// use fastest compression. In our best case we are CPU bound, so this is a good tradeoff,
+									// cutting overall time by 2/3 while only modestly increasing the compression ratio (22.7% -> 23.8% for RootEditor PDBs).
+									// This is in cases of a super hot cache, so the operation was largely CPU bound.
+									ZipArchive.AddFile(File, CommandUtils.ConvertSeparators(PathSeparator.Slash, Path.GetDirectoryName(CommandUtils.StripBaseDirectory(File, RootDir))));
+								} while (FilesToZip.TryDequeue(out File));
+								ZipArchive.Save();
+							}
+							Interlocked.Add(ref ZipFilesTotalSize, new FileInfo(ZipFileName).Length);
+							// if we are using a staging dir, copy to the final location and delete the staged copy.
+							if (StagingDir != null)
+							{
+								CommandUtils.CopyFile(ZipFileName, CommandUtils.MakeRerootedFilePath(ZipFileName, StagingDir, OutputDir));
+								CommandUtils.DeleteFile(true, ZipFileName);
+							}
+						}
+					})).ToList();
+			}
+			else
+			{
+				ZipThreads = (
+					from CoreNum in Enumerable.Range(0, bZipInParallel ? Environment.ProcessorCount : 1)
+					let ZipFileName = Path.Combine(StagingDir ?? OutputDir, string.Format("{0}{1}.zip", ZipBasename, bZipInParallel ? "-" + CoreNum.ToString("00") : ""))
+					select new Thread(() =>
+					{
+						// don't create the zip unless we have at least one file to add
+						string File;
+						if (FilesToZip.TryDequeue(out File))
+						{
+							// Create one zip per thread using the given basename
+							using (var ZipArchive = System.IO.Compression.ZipFile.Open(ZipFileName, System.IO.Compression.ZipArchiveMode.Create))
+							{
+
+								// pull from the queue until we are out of files.
+								do
+								{
+									// use fastest compression. In our best case we are CPU bound, so this is a good tradeoff,
+									// cutting overall time by 2/3 while only modestly increasing the compression ratio (22.7% -> 23.8% for RootEditor PDBs).
+									// This is in cases of a super hot cache, so the operation was largely CPU bound.
+									// Also, sadly, mono appears to have a bug where nothing you can do will properly set the LastWriteTime on the created entry,
+									// so we have to ignore timestamps on files extracted from a zip, since it may have been created on a Mac.
+									ZipArchive.CreateEntryFromFile(File, CommandUtils.ConvertSeparators(PathSeparator.Slash, CommandUtils.StripBaseDirectory(File, RootDir)), System.IO.Compression.CompressionLevel.Fastest);
+								} while (FilesToZip.TryDequeue(out File));
+							}
+							Interlocked.Add(ref ZipFilesTotalSize, new FileInfo(ZipFileName).Length);
+							// if we are using a staging dir, copy to the final location and delete the staged copy.
+							if (StagingDir != null)
+							{
+								CommandUtils.CopyFile(ZipFileName, CommandUtils.MakeRerootedFilePath(ZipFileName, StagingDir, OutputDir));
+								CommandUtils.DeleteFile(true, ZipFileName);
+							}
+						}
+					})).ToList();
+			}
+			ZipThreads.ForEach(thread => thread.Start());
+			ZipThreads.ForEach(thread => thread.Join());
+			return new ParallelZipResult(ZipTimer.ElapsedTime, ZipFilesTotalSize);
+		}
+
+		/// <summary>
+		/// Unzips a set of zip files with a given basename in a given folder to a given RootDir.
+		/// </summary>
+		/// <param name="RootDir">Root Directory where all files will be extracted.</param>
+		/// <param name="FolderWithZipFiles">Folder containing the zip files to unzip. None of the zips should have the same file path in them.</param>
+		/// <param name="ZipBasename">The basename of the set of zip files to unzip.</param>
+		/// <returns>Some metrics about the unzip process.</returns>
+		/// <remarks>
+		/// The code is expected to be the used as the symmetrical inverse of <see cref="ParallelZipFiles"/>, but could be used independently, as long as the files in the zip do not overlap.
+		/// </remarks>
+		private static ParallelZipResult ParallelUnZipFiles(string RootDir, string FolderWithZipFiles, string ZipBasename)
+		{
+			var UnzipTimer = DateTimeStopwatch.Start();
+			long ZipFilesTotalSize = 0L;
+			// Sadly, mono implemention of System.IO.Compression is really poor (as of 2015/Aug), causing OOM when parallel zipping a large set of files.
+			// However, Ionic is MUCH slower than .NET's native implementation (2x+ slower in our build farm), so we stick to the faster solution on PC.
+			// The code duplication in the threadprocs is unfortunate here, and hopefully we can settle on .NET's implementation on both platforms eventually.
+			if (Utils.IsRunningOnMono)
+			{
+				Parallel.ForEach(Directory.EnumerateFiles(FolderWithZipFiles, ZipBasename + "*.zip").ToList(),
+					(ZipFilename) =>
+					{
+						Interlocked.Add(ref ZipFilesTotalSize, new FileInfo(ZipFilename).Length);
+						using (var ZipArchive = Ionic.Zip.ZipFile.Read(ZipFilename))
+						{
+							ZipArchive.ExtractAll(RootDir, Ionic.Zip.ExtractExistingFileAction.OverwriteSilently);
+						}
+					});
+			}
+			else
+			{
+				Parallel.ForEach(Directory.EnumerateFiles(FolderWithZipFiles, ZipBasename + "*.zip").ToList(),
+					(ZipFilename) =>
+					{
+						Interlocked.Add(ref ZipFilesTotalSize, new FileInfo(ZipFilename).Length);
+						// unzip the files manually instead of caling ZipFile.ExtractToDirectory() because we need to overwrite readonly files. Because of this, creating the directories is up to us as well.
+						using (var ZipArchive = System.IO.Compression.ZipFile.OpenRead(ZipFilename))
+						{
+							foreach (var Entry in ZipArchive.Entries)
+							{
+								// Use CommandUtils.CombinePaths to ensure directory separators get converted correctly. On mono on *nix, if the path has backslashes it will not convert it.
+								var ExtractedFilename = CommandUtils.CombinePaths(RootDir, Entry.FullName);
+								// Zips can contain empty dirs. Ours usually don't have them, but we should support it.
+								if (Path.GetFileName(ExtractedFilename).Length == 0)
+								{
+									Directory.CreateDirectory(ExtractedFilename);
+								}
+								else
+								{
+									// We must delete any existing file, even if it's readonly. .Net does not do this by default.
+									if (File.Exists(ExtractedFilename))
+									{
+										InternalUtils.SafeDeleteFile(ExtractedFilename, true);
+									}
+									else
+									{
+										Directory.CreateDirectory(Path.GetDirectoryName(ExtractedFilename));
+									}
+									Entry.ExtractToFile(ExtractedFilename, true);
+								}
+							}
+						}
+					});
+			}
+			return new ParallelZipResult(UnzipTimer.ElapsedTime, ZipFilesTotalSize);
+		}
+
+		/// <summary>
+		/// Saves the list of fully qualified files (that should be rooted at the shared temp storage location for the game) to a shared temp storage manifest with the given temp storage node and game.
+		/// </summary>
+		/// <param name="BlockName">Name of the block of temp storage (essentially used to identify a subdirectory insides the game's temp storage folder).</param>
+		/// <param name="Files">Fully qualified names of files to reference in the manifest file.</param>
+		/// <param name="SharedStorageDir">Base directory for shared temp storage for the current build. May be null.</param>
+		/// <param name="bWriteToSharedStorage>Whether writes are permitted to shared storage</param>
+		/// <returns>The created manifest instance (which has already been saved to disk).</returns>
+		public static void StoreToTempStorage(string BlockName, List<string> Files, string SharedStorageDir, bool bWriteToSharedStorage)
+		{
+			using (var TelemetryStopwatch = new TelemetryStopwatch("StoreToTempStorage"))
+			{
+				// use LocalRoot if one is not specified
+				string RootDir = CommandUtils.CmdEnv.LocalRoot;
+
+				// save the manifest to local temp storage.
+				var Local = SaveLocalTempStorageManifest(RootDir, BlockName, Files);
+				var LocalTotalSize = Local.GetTotalSize();
+				if (SharedStorageDir == null || !bWriteToSharedStorage)
+				{
+					TelemetryStopwatch.Finish(string.Format("StoreToTempStorage.{0}.{1}.{2}.Local.{3}.{4}.{5}", Files.Count, LocalTotalSize, 0L, 0L, 0L, BlockName));
+				}
+				else
+				{
+					var SharedStorageNodeDir = SharedTempStorageDirectory(SharedStorageDir, BlockName);
+					CommandUtils.Log("Storing to {0}", SharedStorageNodeDir);
+					// this folder should not already exist, else we have concurrency or job duplication problems.
+					if (CommandUtils.DirectoryExists_NoExceptions(SharedStorageNodeDir))
+					{
+						throw new AutomationException("Storage Block Already Exists! {0}", SharedStorageNodeDir);
+					}
+					CommandUtils.CreateDirectory(true, SharedStorageNodeDir);
+
+					var LocalManifestFilename = LocalTempStorageManifestFilename(BlockName);
+					var SharedManifestFilename = SharedTempStorageManifestFilename(SharedStorageDir, BlockName);
+					var StagingDir = Path.GetDirectoryName(LocalManifestFilename);
+					var ZipBasename = Path.GetFileNameWithoutExtension(LocalManifestFilename);
+					// initiate the parallel zip operation.
+					var ZipResult = ParallelZipFiles(Files, RootDir, SharedStorageNodeDir, StagingDir, ZipBasename);
+
+					// copy the local manifest to the shared location. We have to assume the zip is a good copy.
+					CommandUtils.CopyFile(LocalManifestFilename, SharedManifestFilename);
+					TelemetryStopwatch.Finish(string.Format("StoreToTempStorage.{0}.{1}.{2}.Remote.{3}.{4}.{5}", Files.Count, LocalTotalSize, ZipResult.ZipFilesTotalSize, 0, (long)ZipResult.ZipTime.TotalMilliseconds, BlockName));
+				}
+			}
+		}
+
+		/// <summary>
+		/// Inverse of <see cref="StoreToTempStorage"/>.
+		/// Copies a block of files from a temp storage location given by a temp storage node and game to local storage rooted at the given root dir. 
+		/// If the temp manifest for this block is found locally, the copy is skipped, as we assume this is the same machine that created the temp storage and the files are still there.
+		/// </summary>
+		/// <param name="SharedStorageDir">Base directory for temp storage in the current build</param>
+		/// <param name="BlockName">Node info descibing the block of temp storage (essentially used to identify a subdirectory insides the game's temp storage folder).</param>
+		/// <returns>List of fully qualified paths to all the files that were retrieved. This is returned even if we skip the copy (set WasLocal = true) .</returns>
+		public static List<string> RetrieveFromTempStorage(string SharedStorageDir, string BlockName)
+		{
+			using (var TelemetryStopwatch = new TelemetryStopwatch("RetrieveFromTempStorage"))
+			{
+				string RootDir = CommandUtils.CmdEnv.LocalRoot;
+
+				// First see if the local manifest is there.
+				// If it is, then we must be on the same node as the one that originally created the temp storage.
+				// In that case, we just verify all the files exist as described in the manifest and use that.
+				// If there was any tampering, abort immediately because we never accept tampering, and it signifies a build problem.
+				var LocalManifest = LocalTempStorageManifestFilename(BlockName);
+				if (CommandUtils.FileExists_NoExceptions(LocalManifest))
+				{
+					CommandUtils.Log("Found local manifest {0}", LocalManifest);
+					var Local = TempStorageManifest.Load(LocalManifest);
+					var Files = Local.GetFiles(RootDir);
+					var LocalTest = TempStorageManifest.Create(Files, RootDir);
+					if (!Local.Compare(LocalTest, true))
+					{
+						throw new AutomationException("Local files in manifest {0} were tampered with.", LocalManifest);
+					}
+					TelemetryStopwatch.Finish(string.Format("RetrieveFromTempStorage.{0}.{1}.{2}.Local.{3}.{4}.{5}", Files.Count, Local.GetTotalSize(), 0L, 0L, 0L, BlockName));
+					return Files;
+				}
+
+				// We couldn't find the node storage locally, so get it from the shared location.
+				var SharedStorageNodeDir = SharedTempStorageDirectory(SharedStorageDir, BlockName);
+
+				CommandUtils.Log("Attempting to retrieve from {0}", SharedStorageNodeDir);
+				if (!CommandUtils.DirectoryExists_NoExceptions(SharedStorageNodeDir))
+				{
+					throw new AutomationException("Storage Block Does Not Exists! {0}", SharedStorageNodeDir);
+				}
+				var SharedManifest = SharedTempStorageManifestFilename(SharedStorageDir, BlockName);
+				InternalUtils.Robust_FileExists(SharedManifest, "Storage Block Manifest Does Not Exists! {0}");
+
+				var Shared = TempStorageManifest.Load(SharedManifest);
+				var SharedFiles = Shared.GetFiles(SharedStorageNodeDir);
+
+				// We know the source files exist and are under RootDir because we created the manifest, which verifies it.
+				// Now create the list of target files
+				var DestFiles = SharedFiles.Select(Filename =>
+				{
+					var DestFile = CommandUtils.MakeRerootedFilePath(Filename, SharedStorageNodeDir, RootDir);
+					// create a FileInfo using the file, which will help us catch path too long exceptions early and propagate a friendly error containing the full path.
+					try
+					{
+						return new FileInfo(DestFile).FullName;
+					}
+					catch (PathTooLongException Ex)
+					{
+						throw new AutomationException(Ex, "Path too long ... failed to create FileInfo for {0}", DestFile);
+					}
+				}).ToList();
+
+				var ZipBasename = Path.GetFileNameWithoutExtension(LocalManifest);
+
+				// now unzip in parallel, overwriting any existing local file.
+				var ZipResult = ParallelUnZipFiles(RootDir, SharedStorageNodeDir, ZipBasename);
+
+				var NewLocal = SaveLocalTempStorageManifest(RootDir, BlockName, DestFiles);
+				// Now compare the created local files to ensure their attributes match the one we copied from the network.
+				if (!NewLocal.Compare(Shared, true))
+				{
+					// we will rename this so it can't be used, but leave it around for inspection
+					CommandUtils.RenameFile_NoExceptions(LocalManifest, LocalManifest + ".broken");
+					throw new AutomationException("Shared and Local manifest mismatch.");
+				}
+
+				// Handle unix permissions/chmod issues. This will touch the timestamp we check on the file, so do this after we've compared with the manifest attributes.
+				if (Utils.IsRunningOnMono)
+				{
+					foreach (string DestFile in DestFiles)
+					{
+						CommandUtils.FixUnixFilePermissions(DestFile);
+					}
+				}
+
+				TelemetryStopwatch.Finish(string.Format("RetrieveFromTempStorage.{0}.{1}.{2}.Remote.{3}.{4}.{5}", DestFiles.Count, Shared.GetTotalSize(), ZipResult.ZipFilesTotalSize, 0, (long)ZipResult.ZipTime.TotalMilliseconds, BlockName));
+				return DestFiles;
+			}
+		}
+
+		/// <summary>
+		/// Runs a test of the temp storage system. This function is part of the class so it doesn't have to expose internals just to run the test.
+		/// </summary>
+		/// <param name="CmdEnv">The environment to use.</param>
+		static internal void TestTempStorage(CommandEnvironment CmdEnv)
+		{
+			// We are not a real GUBP job, so fake the values.
+			string BlockName = "Test";
+
+			string SharedStorageDir = CommandUtils.CombinePaths(CommandUtils.CmdEnv.LocalRoot, "Engine", "Saved", "TmpStoreShared");
+
+			// Delete any local and shared temp storage that may exist.
+			DeleteLocalTempStorage();
+			DeleteSharedTempStorage(SharedStorageDir, BlockName);
+			if (TempStorageExists(BlockName, SharedStorageDir, false))
+			{
+				throw new AutomationException("storage should not exist");
+			}
+
+			// Create some test files in various locations in the local root with unique content.
+			var TestFiles = new[]
+                { 
+                    CommandUtils.CombinePaths(CmdEnv.LocalRoot, "Engine", "Build", "Batchfiles", "TestFile.Txt"),
+                    CommandUtils.CombinePaths(CmdEnv.LocalRoot, "TestFile2.Txt"),
+                    CommandUtils.CombinePaths(CmdEnv.LocalRoot, "engine", "plugins", "TestFile3.Txt"),
                 }
-            }
-        }
-    }
+				.Select(TestFile => new
+				{
+					FileName = TestFile,
+					FileContents = string.Format("{0} - {1}", TestFile, DateTime.UtcNow)
+				})
+				.ToList();
+			foreach (var TestFile in TestFiles)
+			{
+				CommandUtils.DeleteFile(TestFile.FileName);
+				CommandUtils.Log("Test file {0}", TestFile.FileName);
+				File.WriteAllText(TestFile.FileName, TestFile.FileContents);
+				// we should be able to overwrite readonly files.
+				File.SetAttributes(TestFile.FileName, FileAttributes.ReadOnly);
+			}
+
+			// wrap the operation so we are sure to clean up afterward.
+			try
+			{
+				// Store the test file to temp storage.
+				StoreToTempStorage(BlockName, TestFiles.Select(TestFile => TestFile.FileName).ToList(), SharedStorageDir, true);
+				// The manifest should exist locally.
+				if (!LocalTempStorageManifestExists(BlockName))
+				{
+					throw new AutomationException("local storage should exist");
+				}
+				// The manifest should exist on the shared drive.
+				if (!SharedTempStorageManifestExists(SharedStorageDir, BlockName, false))
+				{
+					throw new AutomationException("shared storage should exist");
+				}
+				// Now delete the local manifest
+				DeleteLocalTempStorage();
+				// It should no longer be there.
+				if (LocalTempStorageManifestExists(BlockName))
+				{
+					throw new AutomationException("local storage should not exist");
+				}
+				// But the shared storage should still exist.
+				if (!TempStorageExists(BlockName, SharedStorageDir, false))
+				{
+					throw new AutomationException("some storage should exist");
+				}
+
+				// Now we should be able to retrieve the test files from shared storage, and it should overwrite our read-only files.
+				RetrieveFromTempStorage(SharedStorageDir, BlockName);
+				// Now delete the local manifest so we can try again with no files present (the usual case for restoring from temp storage).
+				DeleteLocalTempStorage();
+
+				// Ok, delete our test files locally.
+				foreach (var TestFile in TestFiles)
+				{
+					CommandUtils.DeleteFile(TestFile.FileName);
+				}
+
+				// Now we should be able to retrieve the test files from shared storage.
+				RetrieveFromTempStorage(SharedStorageDir, BlockName);
+				// the local manifest should be there, since we just retrieved from shared storage.
+				if (!LocalTempStorageManifestExists(BlockName))
+				{
+					throw new AutomationException("local storage should exist");
+				}
+				// The shared manifest should also still be there.
+				if (!SharedTempStorageManifestExists(SharedStorageDir, BlockName, false))
+				{
+					throw new AutomationException("shared storage should exist");
+				}
+				// Verify the retrieved files have the correct content.
+				foreach (var TestFile in TestFiles)
+				{
+					if (File.ReadAllText(TestFile.FileName) != TestFile.FileContents)
+					{
+						throw new AutomationException("Contents of the test file {0} was changed after restoring from shared temp storage.", TestFile.FileName);
+					}
+				}
+				// Now delete the shared temp storage
+				DeleteSharedTempStorage(SharedStorageDir, BlockName);
+				// Shared temp storage manifest should no longer exist.
+				if (SharedTempStorageManifestExists(SharedStorageDir, BlockName, false))
+				{
+					throw new AutomationException("shared storage should not exist");
+				}
+				// Retrieving temp storage should now just retrieve from local
+				RetrieveFromTempStorage(SharedStorageDir, BlockName);
+				if (!LocalTempStorageManifestExists(BlockName))
+				{
+					throw new AutomationException("local storage should exist");
+				}
+
+				// and now lets test tampering
+				DeleteLocalTempStorage();
+				{
+					bool bFailedProperly = false;
+					var MissingFile = new List<string>(TestFiles.Select(TestFile => TestFile.FileName));
+					// add a file to the manifest that shouldn't be there.
+					MissingFile.Add(CommandUtils.CombinePaths(CmdEnv.LocalRoot, "Engine", "SomeFileThatDoesntExist.txt"));
+					try
+					{
+						StoreToTempStorage(BlockName, MissingFile, SharedStorageDir, true);
+					}
+					catch (AutomationException)
+					{
+						bFailedProperly = true;
+					}
+					if (!bFailedProperly)
+					{
+						throw new AutomationException("Missing file did not fail.");
+					}
+				}
+
+				// clear the shared temp storage again.
+				DeleteSharedTempStorage(SharedStorageDir, BlockName);
+				// store the test files to shared temp storage again.
+				StoreToTempStorage(BlockName, TestFiles.Select(TestFile => TestFile.FileName).ToList(), SharedStorageDir, true);
+				// force a load from shared by deleting the local manifest
+				DeleteLocalTempStorage();
+				// delete our test files locally.
+				foreach (var TestFile in TestFiles)
+				{
+					CommandUtils.DeleteFile(TestFile.FileName);
+				}
+
+				// now test that retrieving from shared temp storage properly balks that a file is missing.
+				{
+					// tamper with the shared files.
+					var RandomSharedZipFile = Directory.EnumerateFiles(SharedTempStorageDirectory(SharedStorageDir, BlockName), "*.zip").First();
+					// delete the shared file.
+					CommandUtils.DeleteFile(RandomSharedZipFile);
+
+					bool bFailedProperly = false;
+					try
+					{
+						RetrieveFromTempStorage(SharedStorageDir, BlockName);
+					}
+					catch (AutomationException)
+					{
+						bFailedProperly = true;
+					}
+					if (!bFailedProperly)
+					{
+						throw new AutomationException("Did not fail to load from missing file.");
+					}
+				}
+				// recreate our temp files.
+				foreach (var TestFile in TestFiles)
+				{
+					File.WriteAllText(TestFile.FileName, TestFile.FileContents);
+				}
+
+				// clear the shared temp storage again.
+				DeleteSharedTempStorage(SharedStorageDir, BlockName);
+				// Copy the files to temp storage.
+				StoreToTempStorage(BlockName, TestFiles.Select(TestFile => TestFile.FileName).ToList(), SharedStorageDir, true);
+				// Delete a local file.
+				CommandUtils.DeleteFile(TestFiles[0].FileName);
+				// retrieving from temp storage should use WasLocal, but should balk because a local file was deleted.
+				{
+					bool bFailedProperly = false;
+					try
+					{
+						RetrieveFromTempStorage(SharedStorageDir, BlockName);
+					}
+					catch (AutomationException)
+					{
+						bFailedProperly = true;
+					}
+					if (!bFailedProperly)
+					{
+						throw new AutomationException("Did not fail to load from missing local file.");
+					}
+				}
+			}
+			finally
+			{
+				// Do a final cleanup.
+				DeleteSharedTempStorage(SharedStorageDir, BlockName);
+				DeleteLocalTempStorage();
+				foreach (var TestFile in TestFiles)
+				{
+					CommandUtils.DeleteFile(TestFile.FileName);
+				}
+			}
+		}
+	}
 
 	[Help("Tests the temp storage operations.")]
 	class TestTempStorage : BuildCommand
@@ -1380,7 +1259,7 @@ namespace AutomationTool
 		{
 			Log("TestTempStorage********");
 
-            TempStorage.TestTempStorage(CmdEnv);
+			TempStorage.TestTempStorage(CmdEnv);
 
 		}
 	}
