Copying //UE4/Orion-Staging to //UE4/Main (Origin: //Orion/Dev-General [REDACTED_HANDLE])

diff --git a/Engine/Source/Programs/AutomationTool/Scripts/Localisation.Automation.cs b/Engine/Source/Programs/AutomationTool/Scripts/Localisation.Automation.cs
--- a/Engine/Source/Programs/AutomationTool/Scripts/Localisation.Automation.cs
+++ b/Engine/Source/Programs/AutomationTool/Scripts/Localisation.Automation.cs
@@ -7,59 +7,16 @@ using System.Text;
 using System.IO;
 using AutomationTool;
 using UnrealBuildTool;
-using OneSky;
-using EpicGames.OneSkyLocalization.Config;
+using EpicGames.Localization;
 
 [Help("Updates the OneSky localization data using the arguments provided.")]
 [Help("UEProjectDirectory", "Sub-path to the project we're gathering for (relative to CmdEnv.LocalRoot).")]
 [Help("UEProjectName", "Optional name of the project we're gathering for (should match its .uproject file, eg QAGame).")]
-[Help("OneSkyConfigName", "Name of the config data to use (see OneSkyConfigHelper).")]
-[Help("OneSkyProjectGroupName", "Name of the project group in OneSky.")]
-[Help("OneSkyProjectNames", "Comma separated list of the projects in OneSky.")]
-[Help("OneSkyBranchSuffix", "Optional suffix to use when uploading the new data to OneSky.")]
+[Help("LocalizationProjectNames", "Comma separated list of the projects to gather text from.")]
+[Help("LocalizationBranch", "Optional suffix to use when uploading the new data to the localization provider.")]
+[Help("LocalizationProvider", "Optional localization provide override (default is OneSky).")]
 class Localise : BuildCommand
 {
-	struct ProjectImportExportInfo
-	{
-		/** The destination path for the files containing the gathered data for this project (extracted from its config file - relative to the root working directory for the commandlet) */
-		public string DestinationPath;
-
-		/** The name to use for this projects manifest file (extracted from its config file) */
-		public string ManifestName;
-
-		/** The name to use for this projects archive file (extracted from its config file) */
-		public string ArchiveName;
-
-		/** The name to use for this projects portable object file (extracted from its config file) */
-		public string PortableObjectName;
-
-		/** The native culture for this project (extracted from its config file) */
-		public string NativeCulture;
-
-		/** The cultures to generate for this project (extracted from its config file) */
-		public List<string> CulturesToGenerate;
-
-		/** True if we should use a per-culture directly when importing/exporting */
-		public bool bUseCultureDirectory;
-	};
-
-	struct ProjectInfo
-	{
-		/** The name of this project */
-		public string ProjectName;
-
-		/** Path to this projects localization config files (relative to the root working directory for the commandlet) - ordered so that iterating them runs in the correct order */
-		public List<string> LocalizationConfigFiles;
-
-		/** Config data used by the PO file import process */
-		public ProjectImportExportInfo ImportInfo;
-
-		/** Config data used by the PO file export process */
-		public ProjectImportExportInfo ExportInfo;
-	};
-
-    private int OneSkyDownloadedPOChangeList = 0;
-
 	public override void ExecuteBuild()
 	{
 		var EditorExe = CombinePaths(CmdEnv.LocalRoot, @"Engine/Binaries/Win64/UE4Editor-Cmd.exe");
@@ -77,35 +34,36 @@ class Localise : BuildCommand
 			UEProjectName = "";
 		}
 
-		var OneSkyConfigName = ParseParamValue("OneSkyConfigName");
-		if (OneSkyConfigName == null)
-		{
-			throw new AutomationException("Missing required command line argument: 'OneSkyConfigName'");
-		}
-
-		var OneSkyProjectGroupName = ParseParamValue("OneSkyProjectGroupName");
-		if (OneSkyProjectGroupName == null)
-		{
-			throw new AutomationException("Missing required command line argument: 'OneSkyProjectGroupName'");
-		}
-
-		var OneSkyProjectNames = new List<string>();
+		var LocalizationProjectNames = new List<string>();
 		{
-			var OneSkyProjectNamesStr = ParseParamValue("OneSkyProjectNames");
-			if (OneSkyProjectNamesStr == null)
+			var LocalizationProjectNamesStr = ParseParamValue("LocalizationProjectNames");
+			if (LocalizationProjectNamesStr == null)
 			{
-				throw new AutomationException("Missing required command line argument: 'OneSkyProjectNames'");
+				throw new AutomationException("Missing required command line argument: 'LocalizationProjectNames'");
 			}
-			foreach (var ProjectName in OneSkyProjectNamesStr.Split(','))
+			foreach (var ProjectName in LocalizationProjectNamesStr.Split(','))
 			{
-				OneSkyProjectNames.Add(ProjectName.Trim());
+				LocalizationProjectNames.Add(ProjectName.Trim());
 			}
 		}
 
-		var OneSkyBranchSuffix = ParseParamValue("OneSkyBranchSuffix");
+		var LocalizationProviderName = ParseParamValue("LocalizationProvider");
+		if (LocalizationProviderName == null)
+		{
+			LocalizationProviderName = "OneSky";
+		}
 
 		var RootWorkingDirectory = CombinePaths(CmdEnv.LocalRoot, UEProjectDirectory);
 
+		// Try and find our localization provider
+		LocalizationProvider LocProvider = null;
+		{
+			LocalizationProvider.LocalizationProviderArgs LocProviderArgs;
+			LocProviderArgs.RootWorkingDirectory = RootWorkingDirectory;
+			LocProviderArgs.CommandUtils = this;
+			LocProvider = LocalizationProvider.GetLocalizationProvider(LocalizationProviderName, LocProviderArgs);
+		}
+
 		// Make sure the Localization configs and content is up-to-date to ensure we don't get errors later on
 		if (P4Enabled)
 		{
@@ -116,32 +74,15 @@ class Localise : BuildCommand
 
 		// Generate the info we need to gather for each project
 		var ProjectInfos = new List<ProjectInfo>();
-		foreach (var ProjectName in OneSkyProjectNames)
+		foreach (var ProjectName in LocalizationProjectNames)
 		{
 			ProjectInfos.Add(GenerateProjectInfo(RootWorkingDirectory, ProjectName));
 		}
 
-		OneSkyConfigData OneSkyConfig = OneSkyConfigHelper.Find(OneSkyConfigName);
-		var OneSkyService = new OneSkyService(OneSkyConfig.ApiKey, OneSkyConfig.ApiSecret);
-		var OneSkyProjectGroup = GetOneSkyProjectGroup(OneSkyService, OneSkyProjectGroupName);
-
-		// Create changelist for backed up POs from OneSky.
-		if (P4Enabled)
-		{
-			OneSkyDownloadedPOChangeList = P4.CreateChange(P4Env.Client, "OneSky downloaded PO backup.");
-		}
-
 		// Export all text from OneSky
 		foreach (var ProjectInfo in ProjectInfos)
 		{
-			ExportOneSkyProjectToDirectory(RootWorkingDirectory, OneSkyService, OneSkyProjectGroup, OneSkyBranchSuffix, ProjectInfo);
-		}
-
-		// Submit changelist for backed up POs from OneSky.
-		if (P4Enabled)
-		{
-			int SubmittedChangeList;
-			P4.Submit(OneSkyDownloadedPOChangeList, out SubmittedChangeList);
+			LocProvider.DownloadProjectFromLocalizationProvider(ProjectInfo.ProjectName, ProjectInfo.ImportInfo);
 		}
 
 		// Setup editor arguments for SCC.
@@ -184,7 +125,7 @@ class Localise : BuildCommand
 		// Upload all text to OneSky
 		foreach (var ProjectInfo in ProjectInfos)
 		{
-			UploadProjectToOneSky(RootWorkingDirectory, OneSkyService, OneSkyProjectGroup, OneSkyBranchSuffix, ProjectInfo);
+			LocProvider.UploadProjectToLocalizationProvider(ProjectInfo.ProjectName, ProjectInfo.ExportInfo);
 		}
 	}
 
@@ -288,208 +229,4 @@ class Localise : BuildCommand
 
 		return ProjectImportExportInfo;
 	}
-
-	private static ProjectGroup GetOneSkyProjectGroup(OneSkyService OneSkyService, string ProjectGroupName)
-    {
-		var OneSkyProjectGroup = OneSkyService.ProjectGroups.FirstOrDefault(g => g.Name == ProjectGroupName);
-
-		if (OneSkyProjectGroup == null)
-        {
-			OneSkyProjectGroup = new ProjectGroup(ProjectGroupName, "en");
-			OneSkyService.ProjectGroups.Add(OneSkyProjectGroup);
-        }
-
-		return OneSkyProjectGroup;
-    }
-
-	private static OneSky.Project GetOneSkyProject(OneSkyService OneSkyService, string GroupName, string ProjectName, string ProjectDescription = "")
-    {
-		var OneSkyProjectGroup = GetOneSkyProjectGroup(OneSkyService, GroupName);
-
-		OneSky.Project OneSkyProject = OneSkyProjectGroup.Projects.FirstOrDefault(p => p.Name == ProjectName);
-
-		if (OneSkyProject == null)
-        {
-			ProjectType projectType = OneSkyService.ProjectTypes.First(pt => pt.Code == "website");
-
-			OneSkyProject = new OneSky.Project(ProjectName, ProjectDescription, projectType);
-			OneSkyProjectGroup.Projects.Add(OneSkyProject);
-        }
-
-		return OneSkyProject;
-    }
-
-	private void ExportOneSkyProjectToDirectory(string RootWorkingDirectory, OneSkyService OneSkyService, ProjectGroup OneSkyProjectGroup, string OneSkyBranchSuffix, ProjectInfo ProjectInfo)
-	{
-		var OneSkyFileName = GetOneSkyFilename(ProjectInfo.ImportInfo.PortableObjectName, OneSkyBranchSuffix);
-		var OneSkyProject = GetOneSkyProject(OneSkyService, OneSkyProjectGroup.Name, ProjectInfo.ProjectName);
-		var OneSkyFile = OneSkyProject.UploadedFiles.FirstOrDefault(f => f.Filename == OneSkyFileName);
-
-		//Export
-		if (OneSkyFile != null)
-		{
-			var CulturesToExport = new List<string>();
-			foreach (var OneSkyCulture in OneSkyProject.EnabledCultures)
-			{
-				// Only export the OneSky cultures that we care about for this project
-				if (ProjectInfo.ImportInfo.CulturesToGenerate.Contains(OneSkyCulture))
-				{
-					CulturesToExport.Add(OneSkyCulture);
-				}
-			}
-
-			ExportOneSkyFileToDirectory(OneSkyFile, new DirectoryInfo(CombinePaths(RootWorkingDirectory, ProjectInfo.ImportInfo.DestinationPath)), ProjectInfo.ImportInfo.PortableObjectName, CulturesToExport, ProjectInfo.ImportInfo.bUseCultureDirectory);
-		}
-	}
-
-	private void ExportOneSkyFileToDirectory(UploadedFile OneSkyFile, DirectoryInfo DestinationDirectory, string DestinationFilename, IEnumerable<string> Cultures, bool bUseCultureDirectory)
-    {
-		foreach (var Culture in Cultures)
-        {
-			var CultureDirectory = (bUseCultureDirectory) ? new DirectoryInfo(Path.Combine(DestinationDirectory.FullName, Culture)) : DestinationDirectory;
-			if (!CultureDirectory.Exists)
-            {
-				CultureDirectory.Create();
-            }
-
-            using (var MemoryStream = new MemoryStream())
-            {
-				var ExportTranslationState = OneSkyFile.ExportTranslation(Culture, MemoryStream).Result;
-				if (ExportTranslationState == UploadedFile.ExportTranslationState.Success)
-				{
-					var ExportFile = new FileInfo(Path.Combine(CultureDirectory.FullName, DestinationFilename));
-
-					// Write out the updated PO file so that the gather commandlet will import the new data from it
-					{
-						var ExportFileWasReadOnly = false;
-						if (ExportFile.Exists)
-						{
-							// We're going to clobber the existing PO file, so make sure it's writable (it may be read-only if in Perforce)
-							ExportFileWasReadOnly = ExportFile.IsReadOnly;
-							ExportFile.IsReadOnly = false;
-						}
-
-						MemoryStream.Position = 0;
-						using (Stream FileStream = ExportFile.OpenWrite())
-						{
-							MemoryStream.CopyTo(FileStream);
-							Console.WriteLine("[SUCCESS] Exporting: '{0}' as '{1}' ({2})", OneSkyFile.Filename, ExportFile.FullName, Culture);
-						}
-
-						if (ExportFileWasReadOnly)
-						{
-							ExportFile.IsReadOnly = true;
-						}
-					}
-
-					// Also update the back-up copy so we can diff against what we got from OneSky, and what the gather commandlet produced
-					{
-						var ExportFileCopy = new FileInfo(Path.Combine(ExportFile.DirectoryName, String.Format("{0}_FromOneSky{1}", Path.GetFileNameWithoutExtension(ExportFile.Name), ExportFile.Extension)));
-
-						var ExportFileCopyWasReadOnly = false;
-						if (ExportFileCopy.Exists)
-						{
-							// We're going to clobber the existing PO file, so make sure it's writable (it may be read-only if in Perforce)
-							ExportFileCopyWasReadOnly = ExportFileCopy.IsReadOnly;
-							ExportFileCopy.IsReadOnly = false;
-						}
-
-						ExportFile.CopyTo(ExportFileCopy.FullName, true);
-
-						if (ExportFileCopyWasReadOnly)
-						{
-							ExportFileCopy.IsReadOnly = true;
-						}
-
-						// Add/check out backed up POs from OneSky.
-						if (P4Enabled)
-						{
-							UE4Build.AddBuildProductsToChangelist(OneSkyDownloadedPOChangeList, new List<string>() { ExportFileCopy.FullName });
-						}
-					}
-				}
-				else if (ExportTranslationState == UploadedFile.ExportTranslationState.NoContent)
-                {
-					Console.WriteLine("[WARNING] Exporting: '{0}' ({1}) has no translations!", OneSkyFile.Filename, Culture);
-                }
-                else
-                {
-					Console.WriteLine("[FAILED] Exporting: '{0}' ({1})", OneSkyFile.Filename, Culture);
-                }
-            }
-        }
-    }
-
-	private void UploadProjectToOneSky(string RootWorkingDirectory, OneSkyService OneSkyService, ProjectGroup OneSkyProjectGroup, string OneSkyBranchSuffix, ProjectInfo ProjectInfo)
-	{
-		var OneSkyProject = GetOneSkyProject(OneSkyService, OneSkyProjectGroup.Name, ProjectInfo.ProjectName);
-
-		Func<string, FileInfo> GetPathForCulture = (string Culture) =>
-		{
-			if (ProjectInfo.ExportInfo.bUseCultureDirectory)
-			{
-				return new FileInfo(Path.Combine(RootWorkingDirectory, ProjectInfo.ExportInfo.DestinationPath, Culture, ProjectInfo.ExportInfo.PortableObjectName));
-			}
-			else
-			{
-				return new FileInfo(Path.Combine(RootWorkingDirectory, ProjectInfo.ExportInfo.DestinationPath, ProjectInfo.ExportInfo.PortableObjectName));
-			}
-		};
-
-		// Upload the .po file for the native culture first
-		UploadFileToOneSky(OneSkyProject, OneSkyBranchSuffix, GetPathForCulture(ProjectInfo.ExportInfo.NativeCulture), ProjectInfo.ExportInfo.NativeCulture);
-
-		// Upload the remaining .po files for the other cultures
-		foreach (var Culture in ProjectInfo.ExportInfo.CulturesToGenerate)
-		{
-			// Skip native culture as we uploaded it above
-			if (Culture != ProjectInfo.ExportInfo.NativeCulture)
-			{
-				UploadFileToOneSky(OneSkyProject, OneSkyBranchSuffix, GetPathForCulture(Culture), Culture);	
-			}
-		}
-	}
-
-	private void UploadFileToOneSky(OneSky.Project OneSkyProject, string OneSkyBranchSuffix, FileInfo FileToUpload, string Culture)
-	{
-		using (var FileStream = FileToUpload.OpenRead())
-		{
-			// Read the BOM
-			var UTF8BOM = new byte[3];
-			FileStream.Read(UTF8BOM, 0, 3);
-
-			// We want to ignore the utf8 BOM
-			if (UTF8BOM[0] != 0xef || UTF8BOM[1] != 0xbb || UTF8BOM[2] != 0xbf)
-			{
-				FileStream.Position = 0;
-			}
-
-			var OneSkyFileName = GetOneSkyFilename(Path.GetFileName(FileToUpload.FullName), OneSkyBranchSuffix);
-
-			Console.WriteLine("Uploading: '{0}' as '{1}' ({2})", FileToUpload.FullName, OneSkyFileName, Culture);
-
-			var UploadedFile = OneSkyProject.Upload(OneSkyFileName, FileStream, Culture).Result;
-
-			if (UploadedFile == null)
-			{
-				Console.WriteLine("[FAILED] Uploading: '{0}' ({1})", FileToUpload.FullName, Culture);
-			}
-			else
-			{
-				Console.WriteLine("[SUCCESS] Uploading: '{0}' ({1})", FileToUpload.FullName, Culture);
-			}
-		}
-	}
-
-	private static string GetOneSkyFilename(string BaseFilename, string OneSkyBranchSuffix)
-	{
-		var OneSkyFileName = BaseFilename;
-		if (!String.IsNullOrEmpty(OneSkyBranchSuffix))
-		{
-			// Apply the branch suffix. OneSky will take care of merging the files from different branches together.
-			var OneSkyFileNameWithSuffix = Path.GetFileNameWithoutExtension(OneSkyFileName) + "_" + OneSkyBranchSuffix + Path.GetExtension(OneSkyFileName);
-			OneSkyFileName = OneSkyFileNameWithSuffix;
-		}
-		return OneSkyFileName;
-	}
 }
diff --git a/Engine/Source/Programs/AutomationTool/Scripts/RebuildLightMapsCommand.Automation.cs b/Engine/Source/Programs/AutomationTool/Scripts/RebuildLightMapsCommand.Automation.cs
--- a/Engine/Source/Programs/AutomationTool/Scripts/RebuildLightMapsCommand.Automation.cs
+++ b/Engine/Source/Programs/AutomationTool/Scripts/RebuildLightMapsCommand.Automation.cs
@@ -42,10 +42,15 @@ namespace AutomationScripts.Automation
 				);
 
 				// Sync and build our targets required for the commandlet to run correctly.
-				P4.Sync(String.Format("{0}/...#head", P4Env.BuildRootP4));
+				// P4.Sync(String.Format("-f {0}/...#head", P4Env.BuildRootP4));
 
-				BuildNecessaryTargets();
-				CheckOutMaps(Params);
+                bool NoBuild = ParseParam("nobuild");
+
+                if (!NoBuild)
+                {
+                    BuildNecessaryTargets();
+                }
+				CreateChangelist(Params);
 				RunRebuildLightmapsCommandlet(Params);
 				SubmitRebuiltMaps();
 			}
@@ -88,24 +93,14 @@ namespace AutomationScripts.Automation
 			}
 		}
 
-		private void CheckOutMaps(ProjectParams Params)
+
+		private void CreateChangelist(ProjectParams Params)
 		{
 			Log("Running Step:- RebuildLightMaps::CheckOutMaps");
 			// Setup a P4 Cl we will use to submit the new lightmaps
-			WorkingCL = P4.CreateChange(P4Env.Client, String.Format("{0} rebuilding lightmaps from changelist {1}", Params.ShortProjectName, P4Env.Changelist));
+			WorkingCL = P4.CreateChange(P4Env.Client, String.Format("{0} rebuilding lightmaps from changelist {1}\n#rb None\n#tests None", Params.ShortProjectName, P4Env.Changelist));
 			Log("Working in {0}", WorkingCL);
 
-			string AllMapsWildcardCmdline = String.Format("{0}\\...\\*.umap", Params.RawProjectPath.Directory);
-			string Output = "";
-			P4.P4Output(out Output, String.Format("edit -c {0} {1}", WorkingCL, AllMapsWildcardCmdline));
-
-			// We need to ensure that any error in the output log is observed.
-			// P4 is still successful if it manages to run the operation.
-			if (FoundCheckOutErrorInP4Output(Output) == true)
-			{
-				LogError("Failed to check out every one of the project maps.");
-				throw new AutomationException("Failed to check out every one of the project maps.");
-			}
 		}
 
 		private void RunRebuildLightmapsCommandlet(ProjectParams Params)
@@ -124,12 +119,17 @@ namespace AutomationScripts.Automation
 			try
 			{
 				var CommandletParams = IsBuildMachine ? "-unattended -buildmachine -fileopenlog" : "-fileopenlog";
+                CommandletParams += " -AutoCheckOutPackages";
+                if (P4Enabled)
+                {
+                    CommandletParams += String.Format(" -SCCProvider={0} -P4Port={1} -P4User={2} -P4Client={3} -P4Changelist={4} -P4Passwd={5}", "Perforce", P4Env.P4Port, P4Env.User, P4Env.Client, WorkingCL.ToString(), P4.GetAuthenticationToken());
+                }
 				RebuildLightMapsCommandlet(Params.RawProjectPath, Params.UE4Exe, Params.MapsToRebuildLightMaps.ToArray(), CommandletParams);
 			}
 			catch (Exception Ex)
 			{
 				// Something went wrong with the commandlet. Abandon this run, don't check in any updated files, etc.
-				LogError("Rebuild Light Maps has failed.");
+				LogError("Rebuild Light Maps has failed. because "+ Ex.ToString());
 				throw new AutomationException(ExitCode.Error_Unknown, Ex, "RebuildLightMaps failed.");
 			}
 		}
@@ -141,9 +141,10 @@ namespace AutomationScripts.Automation
 			// Check everything in!
 			if (WorkingCL != -1)
 			{
+                Log("Running Step:- Submitting CL " + WorkingCL);
 				int SubmittedCL;
 				P4.Submit(WorkingCL, out SubmittedCL, true, true);
-				Log("New lightmaps have been submitted in changelist {0}", SubmittedCL);
+				Log("INFO: Lightmaps successfully submitted in cl "+ SubmittedCL.ToString());
 			}
 		}
 
diff --git a/Engine/Source/Programs/CrashReporter/CrashReportClient/CrashReportClient.Target.cs b/Engine/Source/Programs/CrashReporter/CrashReportClient/CrashReportClient.Target.cs
--- a/Engine/Source/Programs/CrashReporter/CrashReportClient/CrashReportClient.Target.cs
+++ b/Engine/Source/Programs/CrashReporter/CrashReportClient/CrashReportClient.Target.cs
@@ -88,7 +88,7 @@ public class CrashReportClientTarget : TargetRules
 		// the linker.
 		OutLinkEnvironmentConfiguration.bHasExports = false;
 
-		OutCPPEnvironmentConfiguration.Definitions.Add( "USE_CHECKS_IN_SHIPPING=1" );
+		UEBuildConfiguration.bUseChecksInShipping = true;
 	}
     public override bool GUBP_AlwaysBuildWithTools(UnrealTargetPlatform InHostPlatform, out bool bInternalToolOnly, out bool SeparateNode, out bool CrossCompile)
 	{
diff --git a/Engine/Source/Programs/CrashReporter/CrashReportClient/Private/CrashReportClientApp.cpp b/Engine/Source/Programs/CrashReporter/CrashReportClient/Private/CrashReportClientApp.cpp
--- a/Engine/Source/Programs/CrashReporter/CrashReportClient/Private/CrashReportClientApp.cpp
+++ b/Engine/Source/Programs/CrashReporter/CrashReportClient/Private/CrashReportClientApp.cpp
@@ -116,6 +116,7 @@ FPlatformErrorReport LoadErrorReport()
 	}
 	else
 	{
+		UE_LOG(CrashReportClientLog, Warning, TEXT("No error report found"));
 		return FPlatformErrorReport();
 	}
 
@@ -167,8 +168,8 @@ void RunCrashReportClient(const TCHAR* CommandLine)
 		if (bUnattended)
 		{
 			// In the unattended mode we don't send any PII.
-			ErrorReport.SetUserComment(NSLOCTEXT("CrashReportClient", "UnattendedMode", "Sent in the unattended mode"));
 			FCrashReportClientUnattended CrashReportClient(ErrorReport);
+			ErrorReport.SetUserComment(NSLOCTEXT("CrashReportClient", "UnattendedMode", "Sent in the unattended mode"));
 
 			// loop until the app is ready to quit
 			while (!GIsRequestingExit)
@@ -194,6 +195,7 @@ void RunCrashReportClient(const TCHAR* CommandLine)
 			auto Window = FSlateApplication::Get().AddWindow(
 				SNew(SWindow)
 				.Title(NSLOCTEXT("CrashReportClient", "CrashReportClientAppName", "Unreal Engine 4 Crash Reporter"))
+				.HasCloseButton(FCrashReportClientConfig::Get().IsAllowedToCloseWithoutSending())
 				.ClientSize(InitialWindowDimensions)
 				[
 					ClientControl
diff --git a/Engine/Source/Programs/CrashReporter/CrashReportClient/Private/CrashReportClientConfig.cpp b/Engine/Source/Programs/CrashReporter/CrashReportClient/Private/CrashReportClientConfig.cpp
--- a/Engine/Source/Programs/CrashReporter/CrashReportClient/Private/CrashReportClientConfig.cpp
+++ b/Engine/Source/Programs/CrashReporter/CrashReportClient/Private/CrashReportClientConfig.cpp
@@ -51,6 +51,12 @@ FCrashReportClientConfig::FCrashReportClientConfig()
 		bHideLogFilesOption = false;
 	}
 	
+	if (!GConfig->GetBool(TEXT("CrashReportClient"), TEXT("bIsAllowedToCloseWithoutSending"), bIsAllowedToCloseWithoutSending, GEngineIni))
+	{
+		// Default to true (Allow the user to close without sending) when config is missing.
+		bIsAllowedToCloseWithoutSending = true;
+	}
+
 	ReadFullCrashDumpConfigurations();
 }
 
diff --git a/Engine/Source/Programs/CrashReporter/CrashReportClient/Private/CrashReportClientConfig.h b/Engine/Source/Programs/CrashReporter/CrashReportClient/Private/CrashReportClientConfig.h
--- a/Engine/Source/Programs/CrashReporter/CrashReportClient/Private/CrashReportClientConfig.h
+++ b/Engine/Source/Programs/CrashReporter/CrashReportClient/Private/CrashReportClientConfig.h
@@ -78,6 +78,11 @@ struct FCrashReportClientConfig
 		return bHideLogFilesOption;
 	}
 
+	const bool& IsAllowedToCloseWithoutSending() const
+	{
+		return bIsAllowedToCloseWithoutSending;
+	}
+
 	void SetAllowToBeContacted( bool bNewValue );
 	void SetSendLogFile( bool bNewValue );
 
@@ -120,4 +125,7 @@ protected:
 
 	/** Whether the user is shown the option to enable/disable sending the log file. */
 	bool bHideLogFilesOption;
+
+	/** Whether the user is allowed to close the crash reporter without sending a report */
+	bool bIsAllowedToCloseWithoutSending;
 };
diff --git a/Engine/Source/Programs/CrashReporter/CrashReportClient/Private/SCrashReportClient.cpp b/Engine/Source/Programs/CrashReporter/CrashReportClient/Private/SCrashReportClient.cpp
--- a/Engine/Source/Programs/CrashReporter/CrashReportClient/Private/SCrashReportClient.cpp
+++ b/Engine/Source/Programs/CrashReporter/CrashReportClient/Private/SCrashReportClient.cpp
@@ -254,6 +254,7 @@ void SCrashReportClient::Construct(const FArguments& InArgs, TSharedRef<FCrashRe
 					.ContentPadding( FMargin( 8, 2 ) )
 					.Text( LOCTEXT( "CloseWithoutSending", "Close Without Sending" ) )
 					.OnClicked( Client, &FCrashReportClient::CloseWithoutSending )
+					.Visibility(FCrashReportClientConfig::Get().IsAllowedToCloseWithoutSending() ? EVisibility::Visible : EVisibility::Hidden)
 				]
 
 				+SHorizontalBox::Slot()
diff --git a/Engine/Source/Programs/UnrealBuildTool/Configuration/UEBuildConfiguration.cs b/Engine/Source/Programs/UnrealBuildTool/Configuration/UEBuildConfiguration.cs
--- a/Engine/Source/Programs/UnrealBuildTool/Configuration/UEBuildConfiguration.cs
+++ b/Engine/Source/Programs/UnrealBuildTool/Configuration/UEBuildConfiguration.cs
@@ -247,6 +247,12 @@ namespace UnrealBuildTool
 		[XmlConfig]
 		public static bool bUseLoggingInShipping;
 
+		/// <summary>
+		/// Whether to turn on checks (asserts) for test/shipping builds
+		/// </summary>
+		[XmlConfig]
+		public static bool bUseChecksInShipping;
+
 		/// <summary>
 		/// True if we need PhysX vehicle support
 		/// </summary>
@@ -329,6 +335,7 @@ namespace UnrealBuildTool
 			bCompileWithStatsWithoutEngine = false;
 			bCompileWithPluginSupport = false;
 			bUseLoggingInShipping = false;
+			bUseChecksInShipping = false;
 			bCompileSteamOSS = true;
 			bCompileMcpOSS = true;
 			bCompilePhysXVehicle = true;
diff --git a/Engine/Source/Programs/UnrealBuildTool/Configuration/UEBuildTarget.cs b/Engine/Source/Programs/UnrealBuildTool/Configuration/UEBuildTarget.cs
--- a/Engine/Source/Programs/UnrealBuildTool/Configuration/UEBuildTarget.cs
+++ b/Engine/Source/Programs/UnrealBuildTool/Configuration/UEBuildTarget.cs
@@ -3251,7 +3251,7 @@ namespace UnrealBuildTool
 					ExcludeFolders.Add(DirectoryFragment);
 				}
 			}
-			ValidPlugins.RemoveAll(x => ShouldExcludePlugin(x, ExcludeFolders));
+			ValidPlugins.RemoveAll(x => x.Descriptor.bRequiresBuildPlatform && ShouldExcludePlugin(x, ExcludeFolders));
 
 			// Build a list of enabled plugins
 			EnabledPlugins = new List<PluginInfo>();
@@ -3522,6 +3522,15 @@ namespace UnrealBuildTool
 				GlobalCompileEnvironment.Config.Definitions.Add("USE_LOGGING_IN_SHIPPING=0");
 			}
 
+			if (UEBuildConfiguration.bUseChecksInShipping)
+			{
+				GlobalCompileEnvironment.Config.Definitions.Add("USE_CHECKS_IN_SHIPPING=1");
+			}
+			else
+			{
+				GlobalCompileEnvironment.Config.Definitions.Add("USE_CHECKS_IN_SHIPPING=0");
+			}
+
 			// Propagate whether we want a lean and mean build to the C++ code.
 			if (UEBuildConfiguration.bCompileLeanAndMeanUE)
 			{
diff --git a/Engine/Source/Programs/UnrealBuildTool/System/PluginDescriptor.cs b/Engine/Source/Programs/UnrealBuildTool/System/PluginDescriptor.cs
--- a/Engine/Source/Programs/UnrealBuildTool/System/PluginDescriptor.cs
+++ b/Engine/Source/Programs/UnrealBuildTool/System/PluginDescriptor.cs
@@ -78,12 +78,16 @@ namespace UnrealBuildTool
 		// Set for plugins which are installed
 		public bool bInstalled;
 
+		// For plugins that are under a platform folder (eg. /PS4/), determines whether compiling the plugin requires the build platform and/or SDK to be available
+		public bool bRequiresBuildPlatform;
+
 		/// <summary>
 		/// Private constructor. This object should not be created directly; read it from disk using FromFile() instead.
 		/// </summary>
 		private PluginDescriptor()
 		{
 			FileVersion = (int)PluginDescriptorVersion.Latest;
+			bRequiresBuildPlatform = true;
 		}
 
 		/// <summary>
@@ -149,6 +153,7 @@ namespace UnrealBuildTool
 				RawObject.TryGetBoolField("IsBetaVersion", out Descriptor.bIsBetaVersion);
 				RawObject.TryGetBoolField("Installed", out Descriptor.bInstalled);
 				RawObject.TryGetBoolField("CanBeUsedWithUnrealHeaderTool", out Descriptor.bCanBeUsedWithUnrealHeaderTool);
+				RawObject.TryGetBoolField("RequiresBuildPlatform", out Descriptor.bRequiresBuildPlatform);
 
 				return Descriptor;
 			}
@@ -183,6 +188,7 @@ namespace UnrealBuildTool
 				Writer.WriteValue("CanContainContent", bCanContainContent);
 				Writer.WriteValue("IsBetaVersion", bIsBetaVersion);
 				Writer.WriteValue("Installed", bInstalled);
+				Writer.WriteValue("RequiresBuildPlatform", bRequiresBuildPlatform);
 
 				ModuleDescriptor.WriteArray(Writer, "Modules", Modules);
 
@@ -200,6 +206,9 @@ namespace UnrealBuildTool
 		// Whether it should be enabled by default
 		public bool bEnabled;
 
+		// Whether this plugin is optional, and the game should silently ignore it not being present
+		public bool bOptional;
+
 		// Description of the plugin for users that do not have it installed.
 		public string Description;
 
@@ -232,6 +241,7 @@ namespace UnrealBuildTool
 		public static PluginReferenceDescriptor FromJsonObject(JsonObject RawObject)
 		{
 			PluginReferenceDescriptor Descriptor = new PluginReferenceDescriptor(RawObject.GetStringField("Name"), null, RawObject.GetBoolField("Enabled"));
+			RawObject.TryGetBoolField("Optional", out Descriptor.bOptional);
 			RawObject.TryGetStringField("Description", out Descriptor.Description);
 			RawObject.TryGetStringField("MarketplaceURL", out Descriptor.MarketplaceURL);
 			RawObject.TryGetEnumArrayField<UnrealTargetPlatform>("WhitelistPlatforms", out Descriptor.WhitelistPlatforms);
diff --git a/Engine/Source/Programs/UnrealLightmass/Private/Lighting/Embree.cpp b/Engine/Source/Programs/UnrealLightmass/Private/Lighting/Embree.cpp
--- a/Engine/Source/Programs/UnrealLightmass/Private/Lighting/Embree.cpp
+++ b/Engine/Source/Programs/UnrealLightmass/Private/Lighting/Embree.cpp
@@ -133,7 +133,7 @@ struct FEmbreeFilterProcessor
 	/**
 	 * Test if the ray comes from an HLOD (massive LOD) ignore intersection from other LODs of this HLOD.
 	 *
-	 * @return		true if the ray is allowed to intersect with this geometry.
+	 * @return		true if the ray is rejected.
      */
 	EMBREE_INLINE bool HitRejectedByHLODTest() const
 	{
@@ -142,8 +142,8 @@ struct FEmbreeFilterProcessor
 			uint32 GeoHLODIndex = (Geo.Mesh->GetLODIndices() & 0xFFFF0000) >> 16;
 			uint32 RayHLODIndex = (Ray.MappingMesh->GetLODIndices() & 0xFFFF0000) >> 16;
 
-			// If they are from the same HLOD (0xFFFF and 0 being invalid HLOD).
-			if (GeoHLODIndex != 0 && GeoHLODIndex != 0xFFFF && GeoHLODIndex == RayHLODIndex) 
+			// If either Geo or Ray is a HLOD (0xFFFF being invalid HLOD).
+			if (GeoHLODIndex != 0xFFFF || RayHLODIndex != 0xFFFF)
 			{
 				uint32 GeoHLODRange = Geo.Mesh->GetHLODRange();
 				uint32 GeoHLODRangeStart = GeoHLODRange & 0xFFFF;
@@ -153,10 +153,18 @@ struct FEmbreeFilterProcessor
 				uint32 RayHLODRangeStart = RayHLODRange & 0xFFFF;
 				uint32 RayHLODRangeEnd = (RayHLODRange & 0xFFFF0000) >> 16;
 
-				// If both are not leaf nodes (Start == End)
-				if (GeoHLODRangeStart != GeoHLODRangeEnd || RayHLODRangeStart != RayHLODRangeEnd)
+				// Different rules apply if we're dealing with the same cluster.
+				if (GeoHLODIndex != RayHLODIndex)
 				{
-					// If they are not the same HLOD mesh (identical ranges), ignore the children of this HLOD node (range would intersect)
+					// Allow HLOD leaf nodes (Start == End) to interact with other meshes, else reject.
+					if (GeoHLODRangeStart != GeoHLODRangeEnd)
+					{
+						return true;
+					}
+				}
+				else
+				{
+					// Allow HLOD nodes to self-shadow (identical ranges), else reject where range intersects.
 					if (GeoHLODRange != RayHLODRange && 
 						GeoHLODRangeStart <= RayHLODRangeEnd && RayHLODRangeStart <= GeoHLODRangeEnd)
 					{
diff --git a/Engine/Source/Programs/UnrealLightmass/Private/Lighting/TextureMapping.cpp b/Engine/Source/Programs/UnrealLightmass/Private/Lighting/TextureMapping.cpp
--- a/Engine/Source/Programs/UnrealLightmass/Private/Lighting/TextureMapping.cpp
+++ b/Engine/Source/Programs/UnrealLightmass/Private/Lighting/TextureMapping.cpp
@@ -3569,11 +3569,11 @@ void FStaticLightingSystem::ColorInvalidLightmapUVs(
 		FGuid ObjectGuid = TextureMapping->Mesh->GetObjectGuid();
 		if (bHasWrappingUVs)
 		{
-			GSwarm->SendAlertMessage(NSwarm::ALERT_LEVEL_ERROR, ObjectGuid, TypeId, TEXT("LightmassError_ObjectWrappedUVs"));
+			GSwarm->SendAlertMessage(NSwarm::ALERT_LEVEL_WARNING, ObjectGuid, TypeId, TEXT("LightmassError_ObjectWrappedUVs"));
 		}
 		if (bHasOverlappedUVs && NormalizedOverlap > OverlapThreshold)
 		{
-			GSwarm->SendAlertMessage(NSwarm::ALERT_LEVEL_ERROR, ObjectGuid, TypeId, TEXT("LightmassError_ObjectOverlappedUVs"));
+			GSwarm->SendAlertMessage(NSwarm::ALERT_LEVEL_WARNING, ObjectGuid, TypeId, TEXT("LightmassError_ObjectOverlappedUVs"));
 			FString Info = FString::Printf(TEXT("Lightmap UV are overlapping by %0.1f%%. Please adjust content - Enable Error Coloring to visualize."), NormalizedOverlap * 100.0f);
 			GSwarm->SendAlertMessage(NSwarm::ALERT_LEVEL_INFO, ObjectGuid, TypeId, Info.GetCharArray().GetData());
 		}
diff --git a/Engine/Source/Runtime/Analytics/QoSReporter/Private/QoSReporter.cpp b/Engine/Source/Runtime/Analytics/QoSReporter/Private/QoSReporter.cpp
--- a/Engine/Source/Runtime/Analytics/QoSReporter/Private/QoSReporter.cpp
+++ b/Engine/Source/Runtime/Analytics/QoSReporter/Private/QoSReporter.cpp
@@ -17,7 +17,11 @@
 
 bool FQoSReporter::bIsInitialized;
 TSharedPtr<IAnalyticsProvider> FQoSReporter::Analytics;
-FGuid FQoSReporter::InstanceId;
+
+namespace
+{
+	FString		StoredDeploymentName;
+};
 
 /**
 * External code should bind this delegate if QoS reporting is desired,
@@ -33,14 +37,14 @@ QOSREPORTER_API FAnalytics::FProviderConfigurationDelegate& GetQoSOverrideConfig
 /**
 * Get analytics pointer
 */
-IAnalyticsProvider& FQoSReporter::GetProvider()
+QOSREPORTER_API IAnalyticsProvider& FQoSReporter::GetProvider()
 {
 	checkf(bIsInitialized && IsAvailable(), TEXT("FQoSReporter::GetProvider called outside of Initialize/Shutdown."));
 
 	return *Analytics.Get();
 }
 
-void FQoSReporter::Initialize()
+QOSREPORTER_API void FQoSReporter::Initialize()
 {
 	checkf(!bIsInitialized, TEXT("FQoSReporter::Initialize called more than once."));
 
@@ -82,13 +86,14 @@ void FQoSReporter::Initialize()
 	Analytics = FAnalytics::Get().CreateAnalyticsProvider(
 		FName(*DefaultEngineAnalyticsConfig.Execute(TEXT("ProviderModuleName"), true)),
 		DefaultEngineAnalyticsConfig);
+	if (!Analytics.IsValid())
+	{
+		return;
+	}
 
-	// add a unique id
-	InstanceId = FGuid::NewGuid();
-	FPlatformMisc::CreateGuid(InstanceId);
+	SetBackendDeploymentName(StoredDeploymentName);
 
 	// check if Configs override the heartbeat interval
-	
 	float ConfigHeartbeatInterval = 0.0;
 	if (GConfig->GetFloat(TEXT("QoSReporter"), TEXT("HeartbeatInterval"), ConfigHeartbeatInterval, GEngineIni))
 	{
@@ -110,9 +115,14 @@ void FQoSReporter::Shutdown()
 	bIsInitialized = false;
 }
 
-FGuid FQoSReporter::GetQoSReporterInstanceId()
+FString FQoSReporter::GetQoSReporterInstanceId()
 {
-	return InstanceId;
+	if (Analytics.IsValid())
+	{
+		return Analytics->GetUserID();
+	}
+
+	return FString();
 }
 
 double FQoSReporter::ModuleInitializationTime = FPlatformTime::Seconds();
@@ -129,12 +139,36 @@ void FQoSReporter::ReportStartupCompleteEvent()
 	double Difference = CurrentTime - ModuleInitializationTime;
 
 	TArray<FAnalyticsEventAttribute> ParamArray;
+
 	ParamArray.Add(FAnalyticsEventAttribute(EQoSEvents::ToString(EQoSEventParam::StartupTime), Difference));
 	Analytics->RecordEvent(EQoSEvents::ToString(EQoSEventParam::StartupTime), ParamArray);
 
 	bStartupEventReported = true;
 }
 
+void FQoSReporter::SetBackendDeploymentName(const FString & InDeploymentName)
+{
+	if (Analytics.IsValid())
+	{
+		// abuse somewhat outdated IAnalyticsProvider API for this
+		Analytics->SetLocation(InDeploymentName);
+
+		if (InDeploymentName.Len() > 0)
+		{
+			UE_LOG(LogQoSReporter, Log, TEXT("QoSReporter has been configured for '%s' deployment."), *InDeploymentName);
+		}
+		else
+		{
+			UE_LOG(LogQoSReporter, Log, TEXT("QoSReporter has been configured without a valid deployment name."));
+		}
+	}
+	else
+	{
+		StoredDeploymentName = InDeploymentName;
+		UE_LOG(LogQoSReporter, Log, TEXT("QoSReporter will be configured for '%s' deployment."), *StoredDeploymentName);
+	}
+}
+
 double FQoSReporter::HeartbeatInterval = 300;
 double FQoSReporter::LastHeartbeatTimestamp = 0;
 extern ENGINE_API float GAverageFPS;
@@ -180,9 +214,6 @@ void FQoSReporter::SendHeartbeat()
 	checkf(Analytics.IsValid(), TEXT("SendHeartbeat() should not be called if Analytics provider was not configured"));
 
 	TArray<FAnalyticsEventAttribute> ParamArray;
-	ParamArray.Add(FAnalyticsEventAttribute(TEXT("Role"), GetApplicationRole()));
-	ParamArray.Add(FAnalyticsEventAttribute(TEXT("SystemId"), FPlatformMisc::GetOperatingSystemId()));
-	ParamArray.Add(FAnalyticsEventAttribute(TEXT("InstanceId"), InstanceId.ToString()));
 
 	if (IsRunningDedicatedServer())
 	{
@@ -227,30 +258,9 @@ void FQoSReporter::AddServerHeartbeatAttributes(TArray<FAnalyticsEventAttribute>
 			}
 		}
 
-		// for compatibility with wash, avoid resetting if -statsport is used (temporary measure to avoid interference)
-		static bool bCheckedResettingStats = false;
-		static bool bResetStats = true;
-
-		if (!bCheckedResettingStats)
-		{
-			int32 StatsPort = -1;
-			if (FParse::Value(FCommandLine::Get(), TEXT("statsPort="), StatsPort) && StatsPort > 0)
-			{
-				bResetStats = false;
-			}
-
-			bCheckedResettingStats = true;
-		}
-		
-		if (bResetStats)
-		{
-			UE_LOG(LogQoSReporter, Verbose, TEXT("Resetting PerfCounters - new stat period begins."));
-			PerfCounters->ResetStatsForNextPeriod();
-		}
-		else
-		{
-			UE_LOG(LogQoSReporter, Verbose, TEXT("Not resetting PerfCounters, relying on wash to drive stats periods."));
-		}
+		// reset perfcounters stats here
+		UE_LOG(LogQoSReporter, Verbose, TEXT("Resetting PerfCounters - new stat period begins."));
+		PerfCounters->ResetStatsForNextPeriod();
 	}
 	else if (UE_SERVER)
 	{
@@ -268,29 +278,3 @@ void FQoSReporter::AddClientHeartbeatAttributes(TArray<FAnalyticsEventAttribute>
 {
 	OutArray.Add(FAnalyticsEventAttribute(TEXT("AverageFPS"), GAverageFPS));
 }
-
-/**
- * Returns application role (server, client)
- */
-FString FQoSReporter::GetApplicationRole()
-{
-	if (IsRunningDedicatedServer())
-	{
-		static FString DedicatedServer(TEXT("DedicatedServer"));
-		return DedicatedServer;
-	}
-	else if (IsRunningClientOnly())
-	{
-		static FString ClientOnly(TEXT("ClientOnly"));
-		return ClientOnly;
-	}
-	else if (IsRunningGame())
-	{
-		static FString StandaloneGame(TEXT("StandaloneGame"));
-		return StandaloneGame;
-	}
-
-	static FString Editor(TEXT("Editor"));
-	return Editor;
-}
-
diff --git a/Engine/Source/Runtime/Analytics/QoSReporter/Private/QoSReporterModule.cpp b/Engine/Source/Runtime/Analytics/QoSReporter/Private/QoSReporterModule.cpp
--- a/Engine/Source/Runtime/Analytics/QoSReporter/Private/QoSReporterModule.cpp
+++ b/Engine/Source/Runtime/Analytics/QoSReporter/Private/QoSReporterModule.cpp
@@ -13,6 +13,9 @@ DEFINE_LOG_CATEGORY(LogQoSReporter);
 
 IMPLEMENT_MODULE(FQoSReporterModule, QoSReporter);
 
+// helps to version QoS events (date*10 to allow for 10 revisions per day)
+#define QOS_EVENTS_REVISION					201512170
+
 FString FQoSReporterModule::Config::GetDefaultAppVersion()
 { 
 	return FString::Printf(TEXT("UE4-CL-%d"), FEngineVersion::Current().GetChangelist());
@@ -37,14 +40,17 @@ public:
 
 	/** This provider is not using user IDs */
 	virtual void SetUserID(const FString& InUserID) override {};
-	/** This provider is not using user IDs */
-	virtual FString GetUserID() const override { checkf(false, TEXT("FAnalyticsProviderQoSReporter does not use user ids"));  return TEXT("UnknownUserId"); };
+	/** This provider is not using user IDs, but we're (ab)using this API to return InstanceId */
+	virtual FString GetUserID() const override { return InstanceId.ToString(); };
 
 	/** This provider does not have a concept of sessions */
 	virtual FString GetSessionID() const override { checkf(false, TEXT("FAnalyticsProviderQoSReporter is not session based"));  return TEXT("UnknownSessionId"); };
 	/** This provider does not have a concept of sessions */
 	virtual bool SetSessionID(const FString& InSessionID) override { return false; };
 
+	/** We're (ab)using this API to set DeploymentName */
+	virtual void SetLocation(const FString& InLocation) override { DeploymentName = InLocation; };
+
 	virtual void RecordEvent(const FString& EventName, const TArray<FAnalyticsEventAttribute>& Attributes) override;
 	virtual ~FAnalyticsProviderQoSReporter();
 
@@ -63,10 +69,20 @@ private:
 	/** The upload type to use. */
 	FString UploadType;
 
+	/** Unique identifier for this QoS reporter instance (only changed on module initialization) */
+	FGuid InstanceId;
+	/** Deployment name (if empty, it won't be sent). */
+	FString DeploymentName;
+
 	/**
 	 * Delegate called when an event Http request completes
 	 */
 	void EventRequestComplete(FHttpRequestPtr HttpRequest, FHttpResponsePtr HttpResponse, bool bSucceeded);
+
+	/**
+	 * Returns application role (server, client)
+	 */
+	static FString GetApplicationRole();
 };
 
 void FQoSReporterModule::StartupModule()
@@ -140,6 +156,11 @@ FAnalyticsProviderQoSReporter::FAnalyticsProviderQoSReporter(const FQoSReporterM
 		UploadType = ConfigValues.GetDefaultUploadType();
 	}
 
+	// add a unique id
+	InstanceId = FGuid::NewGuid();
+	FPlatformMisc::CreateGuid(InstanceId);
+
+	UE_LOG(LogQoSReporter, Log, TEXT("QoSReporter initialized (Guid = '%s')"), *InstanceId.ToString());
 	UE_LOG(LogQoSReporter, Log, TEXT("APIKey = '%s'. APIServer = '%s'. AppVersion = '%s'. AppEnvironment = '%s'"), *APIKey, *APIServer, *AppVersion, *AppEnvironment);
 }
 
@@ -149,13 +170,31 @@ FAnalyticsProviderQoSReporter::~FAnalyticsProviderQoSReporter()
 	EndSession();
 }
 
-void FAnalyticsProviderQoSReporter::RecordEvent(const FString& EventName, const TArray<FAnalyticsEventAttribute>& Attributes)
+void FAnalyticsProviderQoSReporter::RecordEvent(const FString& EventName, const TArray<FAnalyticsEventAttribute>& InAttributes)
 {
 	if (APIKey.IsEmpty() || APIServer.IsEmpty())
 	{
 		return;
 	}
 
+	// add attributes common to each QoS event first
+	TArray<FAnalyticsEventAttribute> Attributes;
+	Attributes.Add(FAnalyticsEventAttribute(TEXT("QoSRevision"), QOS_EVENTS_REVISION));
+	Attributes.Add(FAnalyticsEventAttribute(TEXT("Role"), GetApplicationRole()));
+	Attributes.Add(FAnalyticsEventAttribute(TEXT("SystemId"), FPlatformMisc::GetOperatingSystemId()));
+	Attributes.Add(FAnalyticsEventAttribute(TEXT("InstanceId"), InstanceId.ToString()));
+	if (LIKELY(DeploymentName.Len() > 0))
+	{
+		Attributes.Add(FAnalyticsEventAttribute(TEXT("Deployment"), DeploymentName));
+	}
+	else
+	{
+		UE_LOG(LogQoSReporter, Warning, TEXT("QoSReporter was not configured for any deployment; metrics will be likely discarded."));
+	}
+
+	// append the rest
+	Attributes += InAttributes;
+
 	// encode params as Json
 	if (ensure(FModuleManager::Get().IsModuleLoaded("HTTP")))
 	{
@@ -240,3 +279,28 @@ void FAnalyticsProviderQoSReporter::EventRequestComplete(FHttpRequestPtr HttpReq
 	}
 }
 
+/**
+ * Returns application role (server, client)
+ */
+FString FAnalyticsProviderQoSReporter::GetApplicationRole()
+{
+	if (IsRunningDedicatedServer())
+	{
+		static FString DedicatedServer(TEXT("DedicatedServer"));
+		return DedicatedServer;
+	}
+	else if (IsRunningClientOnly())
+	{
+		static FString ClientOnly(TEXT("ClientOnly"));
+		return ClientOnly;
+	}
+	else if (IsRunningGame())
+	{
+		static FString StandaloneGame(TEXT("StandaloneGame"));
+		return StandaloneGame;
+	}
+
+	static FString Editor(TEXT("Editor"));
+	return Editor;
+}
+
diff --git a/Engine/Source/Runtime/Analytics/QoSReporter/Public/QoSReporter.h b/Engine/Source/Runtime/Analytics/QoSReporter/Public/QoSReporter.h
--- a/Engine/Source/Runtime/Analytics/QoSReporter/Public/QoSReporter.h
+++ b/Engine/Source/Runtime/Analytics/QoSReporter/Public/QoSReporter.h
@@ -132,9 +132,9 @@ public:
 	 * Return the provider instance. Not valid outside of Initialize/Shutdown calls.
 	 * Note: must check IsAvailable() first else this code will assert if the provider is not valid.
 	 */
-	static IAnalyticsProvider& GetProvider();
+	QOSREPORTER_API static IAnalyticsProvider& GetProvider();
 	/** Helper function to determine if the provider is valid. */
-	static bool IsAvailable() { return Analytics.IsValid(); }
+	QOSREPORTER_API static bool IsAvailable() { return Analytics.IsValid(); }
 	/** Called to initialize the singleton. */
 	QOSREPORTER_API static void Initialize();
 	/** Called to shut down the singleton */
@@ -158,9 +158,16 @@ public:
 	static void QOSREPORTER_API Tick();
 
 	/**
-	 * This function is expected to be called periodically to update ongoing tasks.
+	 * Returns instance id that QoS reporter is using in its events
+	 */
+	static FString QOSREPORTER_API GetQoSReporterInstanceId();
+
+	/**
+	 * Backend services are split into deployments; this can be used to distinguish between them.
+	 *
+	 * @param InDeploymentName deployment name (has a meaning for the game and QoS alerting). Won't be sent if empty.
 	 */
-	static FGuid QOSREPORTER_API GetQoSReporterInstanceId();
+	static void QOSREPORTER_API SetBackendDeploymentName(const FString & InDeploymentName);
 
 private:
 
@@ -179,17 +186,9 @@ private:
 	 */
 	static void AddClientHeartbeatAttributes(TArray<FAnalyticsEventAttribute> & OutArray);
 
-	/**
-	 * Returns application role (server, client)
-	 */
-	static FString GetApplicationRole();
-
 	/** Whether the module has been initialized. */
 	static bool bIsInitialized;
 
-	/** Unique identifier for this QoS reporter instance (only changed on module initialization) */
-	static FGuid InstanceId;
-
 	/** Chosen analytics provider. */
 	static TSharedPtr<IAnalyticsProvider> Analytics;
 
@@ -207,5 +206,4 @@ private:
 
 	/** Whether startup event was reported. */
 	static bool bStartupEventReported;
-
 };
diff --git a/Engine/Source/Runtime/AnimGraphRuntime/Private/BoneControllers/AnimNode_LookAt.cpp b/Engine/Source/Runtime/AnimGraphRuntime/Private/BoneControllers/AnimNode_LookAt.cpp
--- a/Engine/Source/Runtime/AnimGraphRuntime/Private/BoneControllers/AnimNode_LookAt.cpp
+++ b/Engine/Source/Runtime/AnimGraphRuntime/Private/BoneControllers/AnimNode_LookAt.cpp
@@ -3,6 +3,7 @@
 #include "AnimGraphRuntimePrivatePCH.h"
 #include "BoneControllers/AnimNode_LookAt.h"
 #include "AnimationRuntime.h"
+#include "Animation/AnimInstanceProxy.h"
 
 FVector FAnimNode_LookAt::GetAlignVector(const FTransform& Transform, EAxisOption::Type AxisOption)
 {
@@ -32,6 +33,8 @@ FAnimNode_LookAt::FAnimNode_LookAt()
 	, InterpolationTriggerThreashold(0.f)
 	, CurrentLookAtLocation(FVector::ZeroVector)
 	, AccumulatedInterpoolationTime(0.f)
+	, CachedLookAtSocketMeshBoneIndex(INDEX_NONE)
+	, CachedLookAtSocketBoneIndex(INDEX_NONE)
 {
 }
 
@@ -78,6 +81,18 @@ void FAnimNode_LookAt::EvaluateBoneTransforms(USkeletalMeshComponent* SkelComp,
 		const FTransform& LookAtTransform  = MeshBases.GetComponentSpaceTransform(LookAtBone.GetCompactPoseIndex(BoneContainer));
 		TargetLocationInComponentSpace = LookAtTransform.GetLocation();
 	}
+	// if valid data is available
+	else if (CachedLookAtSocketMeshBoneIndex != INDEX_NONE)
+	{
+		// current LOD has valid index (FCompactPoseBoneIndex is valid if current LOD supports)
+		if (CachedLookAtSocketBoneIndex != INDEX_NONE)
+		{
+			FTransform BoneTransform = MeshBases.GetComponentSpaceTransform(CachedLookAtSocketBoneIndex);
+			FTransform SocketTransformInCS = CachedSocketLocalTransform * BoneTransform;
+
+			TargetLocationInComponentSpace = SocketTransformInCS.GetLocation();
+		}
+	}
 	else
 	{
 		TargetLocationInComponentSpace = SkelComp->ComponentToWorld.InverseTransformPosition(LookAtLocation);
@@ -195,6 +210,13 @@ void FAnimNode_LookAt::InitializeBoneReferences(const FBoneContainer& RequiredBo
 {
 	BoneToModify.Initialize(RequiredBones);
 	LookAtBone.Initialize(RequiredBones);
+
+	// this should be called whenver LOD changes and so on
+	if (CachedLookAtSocketMeshBoneIndex != INDEX_NONE)
+	{
+		const int32 SocketBoneSkeletonIndex = RequiredBones.GetPoseToSkeletonBoneIndexArray()[CachedLookAtSocketMeshBoneIndex];
+		CachedLookAtSocketBoneIndex = RequiredBones.GetCompactPoseIndexFromSkeletonIndex(SocketBoneSkeletonIndex);
+	}
 }
 
 void FAnimNode_LookAt::UpdateInternal(const FAnimationUpdateContext& Context)
@@ -203,3 +225,32 @@ void FAnimNode_LookAt::UpdateInternal(const FAnimationUpdateContext& Context)
 
 	AccumulatedInterpoolationTime = FMath::Clamp(AccumulatedInterpoolationTime+Context.GetDeltaTime(), 0.f, InterpolationTime);;
 }
+
+void FAnimNode_LookAt::Initialize(const FAnimationInitializeContext& Context)
+{
+	FAnimNode_SkeletalControlBase::Initialize(Context);
+
+	CachedLookAtSocketMeshBoneIndex = INDEX_NONE;
+
+	if (LookAtSocket != NAME_None)
+	{
+		const USkeletalMeshComponent* OwnerMeshComponent = Context.AnimInstanceProxy->GetSkelMeshComponent();
+		if (OwnerMeshComponent && OwnerMeshComponent->DoesSocketExist(LookAtSocket))
+		{
+			USkeletalMeshSocket const* const Socket = OwnerMeshComponent->GetSocketByName(LookAtSocket);
+			if (Socket)
+			{
+				CachedSocketLocalTransform = Socket->GetSocketLocalTransform();
+				// cache mesh bone index, so that we know this is valid information to follow
+				CachedLookAtSocketMeshBoneIndex = OwnerMeshComponent->GetBoneIndex(Socket->BoneName);
+
+				ensureMsgf(CachedLookAtSocketMeshBoneIndex != INDEX_NONE, TEXT("%s : socket has invalid bone."), *LookAtSocket.ToString());
+			}
+		}
+		else
+		{
+			// @todo : move to graph node warning
+			UE_LOG(LogAnimation, Warning, TEXT("%s: socket doesn't exist"), *LookAtSocket.ToString());
+		}
+	}
+}
\ No newline at end of file
diff --git a/Engine/Source/Runtime/AnimGraphRuntime/Public/BoneControllers/AnimNode_LookAt.h b/Engine/Source/Runtime/AnimGraphRuntime/Public/BoneControllers/AnimNode_LookAt.h
--- a/Engine/Source/Runtime/AnimGraphRuntime/Public/BoneControllers/AnimNode_LookAt.h
+++ b/Engine/Source/Runtime/AnimGraphRuntime/Public/BoneControllers/AnimNode_LookAt.h
@@ -50,11 +50,14 @@ struct ANIMGRAPHRUNTIME_API FAnimNode_LookAt : public FAnimNode_SkeletalControlB
 	FBoneReference BoneToModify;
 
 	/** Target Bone to look at - you can't use LookAtLocation as alternative as you'll get a delay on bone location if you query directly **/
-	UPROPERTY(EditAnywhere, BlueprintReadWrite, Category=SkeletalControl)
+	UPROPERTY(EditAnywhere, BlueprintReadWrite, Category=Target)
 	FBoneReference LookAtBone;
 
+	UPROPERTY(EditAnywhere, BlueprintReadWrite, Category = Target)
+	FName LookAtSocket;
+
 	/** Target Location in world space if LookAtBone is empty */
-	UPROPERTY(EditAnywhere, BlueprintReadWrite, Category=SkeletalControl, meta=(PinHiddenByDefault))
+	UPROPERTY(EditAnywhere, BlueprintReadWrite, Category = Target, meta = (PinHiddenByDefault))
 	FVector LookAtLocation;
 
 	/** Look at axis, which axis to align to look at point */
@@ -85,22 +88,13 @@ struct ANIMGRAPHRUNTIME_API FAnimNode_LookAt : public FAnimNode_SkeletalControlB
 	UPROPERTY(EditAnywhere, Category=SkeletalControl)
 	bool	bEnableDebug;
 
-	/** Debug transient data */
-	FVector CurrentLookAtLocation;
-
-	/** Current Target Location */
-	FVector CurrentTargetLocation;
-	FVector PreviousTargetLocation;
-
-	/** Current Alpha */
-	float AccumulatedInterpoolationTime;
-
 	// in the future, it would be nice to have more options, -i.e. lag, interpolation speed
 	FAnimNode_LookAt();
 
 	// FAnimNode_Base interface
 	virtual void GatherDebugData(FNodeDebugData& DebugData) override;
 	virtual void UpdateInternal(const FAnimationUpdateContext& Context) override;
+	virtual void Initialize(const FAnimationInitializeContext& Context) override;
 	// End of FAnimNode_Base interface
 
 	// FAnimNode_SkeletalControlBase interface
@@ -135,4 +129,19 @@ private:
 
 		return ABT_Linear;
 	}
+
+	/** Debug transient data */
+	FVector CurrentLookAtLocation;
+
+	/** Current Target Location */
+	FVector CurrentTargetLocation;
+	FVector PreviousTargetLocation;
+
+	/** Current Alpha */
+	float AccumulatedInterpoolationTime;
+
+	/** Look at socket bone cache data */
+	int32 CachedLookAtSocketMeshBoneIndex;
+	FCompactPoseBoneIndex CachedLookAtSocketBoneIndex;
+	FTransform CachedSocketLocalTransform;
 };
diff --git a/Engine/Source/Runtime/Apple/MetalRHI/Private/MetalTexture.cpp b/Engine/Source/Runtime/Apple/MetalRHI/Private/MetalTexture.cpp
--- a/Engine/Source/Runtime/Apple/MetalRHI/Private/MetalTexture.cpp
+++ b/Engine/Source/Runtime/Apple/MetalRHI/Private/MetalTexture.cpp
@@ -926,6 +926,7 @@ void FMetalDynamicRHI::RHIGetTextureMemoryStats(FTextureMemoryStats& OutStats)
 	OutStats.TotalGraphicsMemory = 0;
 
 	OutStats.AllocatedMemorySize = int64(GCurrentTextureMemorySize) * 1024;
+	OutStats.LargestContiguousAllocation = OutStats.AllocatedMemorySize;
 	OutStats.TexturePoolSize = GTexturePoolSize;
 	OutStats.PendingMemoryAdjustment = 0;
 }
diff --git a/Engine/Source/Runtime/AssetRegistry/Private/AssetRegistry.cpp b/Engine/Source/Runtime/AssetRegistry/Private/AssetRegistry.cpp
--- a/Engine/Source/Runtime/AssetRegistry/Private/AssetRegistry.cpp
+++ b/Engine/Source/Runtime/AssetRegistry/Private/AssetRegistry.cpp
@@ -1710,7 +1710,15 @@ void FAssetRegistry::SaveRegistryData(FArchive& Ar, TMap<FName, FAssetData*>& Da
 		{
 			bool bIsMap = InMaps && InMaps->Contains(InDependency->GetPackageName());
 
-			if (!bIsMap)
+			check(!bIsMap || (InDependencyType == EAssetRegistryDependencyType::Soft));
+
+			// Force map dependencies to be soft references as we don't want the package loading system to try and load them
+			// automatically - the level loading stuff will do that!
+			if (bIsMap)
+			{
+				InDependencyType = EAssetRegistryDependencyType::Soft;
+			}
+
 			{
 				auto RedirectedDependency = ResolveRedirector(InDependency, Data, RedirectCache);
 
diff --git a/Engine/Source/Runtime/Core/Private/GenericPlatform/GenericPlatformMemory.cpp b/Engine/Source/Runtime/Core/Private/GenericPlatform/GenericPlatformMemory.cpp
--- a/Engine/Source/Runtime/Core/Private/GenericPlatform/GenericPlatformMemory.cpp
+++ b/Engine/Source/Runtime/Core/Private/GenericPlatform/GenericPlatformMemory.cpp
@@ -125,17 +125,21 @@ void FGenericPlatformMemory::OnOutOfMemory(uint64 Size, uint32 Alignment)
 	}
 	UE_LOG(LogMemory, Warning, TEXT("MemoryStats:")\
 		TEXT("\n\tAvailablePhysical %llu")\
-		TEXT("\n\tAvailableVirtual %llu")\
-		TEXT("\n\tUsedPhysical %llu")\
-		TEXT("\n\tPeakUsedPhysical %llu")\
-		TEXT("\n\tUsedVirtual %llu")\
-		TEXT("\n\tPeakUsedVirtual %llu"),
+		TEXT("\n\t AvailableVirtual %llu")\
+		TEXT("\n\t     UsedPhysical %llu")\
+		TEXT("\n\t PeakUsedPhysical %llu")\
+		TEXT("\n\t      UsedVirtual %llu")\
+		TEXT("\n\t  PeakUsedVirtual %llu"),
 		(uint64)PlatformMemoryStats.AvailablePhysical,
 		(uint64)PlatformMemoryStats.AvailableVirtual,
 		(uint64)PlatformMemoryStats.UsedPhysical,
 		(uint64)PlatformMemoryStats.PeakUsedPhysical,
 		(uint64)PlatformMemoryStats.UsedVirtual,
 		(uint64)PlatformMemoryStats.PeakUsedVirtual);
+	if (GWarn)
+	{
+		GMalloc->DumpAllocatorStats(*GWarn);
+	}
 	UE_LOG(LogMemory, Fatal, TEXT("Ran out of memory allocating %llu bytes with alignment %u"), Size, Alignment);
 }
 
diff --git a/Engine/Source/Runtime/Core/Private/GenericPlatform/GenericPlatformMisc.cpp b/Engine/Source/Runtime/Core/Private/GenericPlatform/GenericPlatformMisc.cpp
--- a/Engine/Source/Runtime/Core/Private/GenericPlatform/GenericPlatformMisc.cpp
+++ b/Engine/Source/Runtime/Core/Private/GenericPlatform/GenericPlatformMisc.cpp
@@ -247,6 +247,14 @@ FString FGenericPlatformMisc::GetPrimaryGPUBrand()
 	return FString( TEXT( "GenericGPUBrand" ) );
 }
 
+void FGenericPlatformMisc::GetGPUDriverInfo(const FString DeviceDescription, FString& InternalDriverVersion, FString& UserDriverVersion, FString& DriverDate)
+{
+	// to prevent bugs
+	InternalDriverVersion.Empty();
+	UserDriverVersion.Empty();
+	DriverDate.Empty();
+}
+
 void FGenericPlatformMisc::GetOSVersions( FString& out_OSVersionLabel, FString& out_OSSubVersionLabel )
 {
 	// Not implemented cross-platform. Each platform may or may not choose to implement this.
diff --git a/Engine/Source/Runtime/Core/Private/Internationalization/ICUText.cpp b/Engine/Source/Runtime/Core/Private/Internationalization/ICUText.cpp
--- a/Engine/Source/Runtime/Core/Private/Internationalization/ICUText.cpp
+++ b/Engine/Source/Runtime/Core/Private/Internationalization/ICUText.cpp
@@ -74,8 +74,6 @@ FText FText::AsTimespan(const FTimespan& Timespan, const FCulturePtr& TargetCult
 	checkf(I18N.IsInitialized() == true, TEXT("FInternationalization is not initialized. An FText formatting method was likely used in static object initialization - this is not supported."));
 	FCulturePtr Culture = TargetCulture.IsValid() ? TargetCulture : I18N.GetCurrentCulture();
 
-	FText TimespanFormatPattern = NSLOCTEXT("Timespan", "FormatPattern", "{Hours}:{Minutes}:{Seconds}");
-
 	double TotalHours = Timespan.GetTotalHours();
 	int32 Hours = static_cast<int32>(TotalHours);
 	int32 Minutes = Timespan.GetMinutes();
@@ -85,11 +83,23 @@ FText FText::AsTimespan(const FTimespan& Timespan, const FCulturePtr& TargetCult
 	NumberFormattingOptions.MinimumIntegralDigits = 2;
 	NumberFormattingOptions.MaximumIntegralDigits = 2;
 
-	FFormatNamedArguments TimeArguments;
-	TimeArguments.Add(TEXT("Hours"), Hours);
-	TimeArguments.Add(TEXT("Minutes"), FText::AsNumber(Minutes, &(NumberFormattingOptions), Culture));
-	TimeArguments.Add(TEXT("Seconds"), FText::AsNumber(Seconds, &(NumberFormattingOptions), Culture));
-	return FText::Format(TimespanFormatPattern, TimeArguments);
+	if (Hours > 0)
+	{
+		FText TimespanFormatPattern = NSLOCTEXT("Timespan", "FormatPattern", "{Hours}:{Minutes}:{Seconds}");
+		FFormatNamedArguments TimeArguments;
+		TimeArguments.Add(TEXT("Hours"), Hours);
+		TimeArguments.Add(TEXT("Minutes"), FText::AsNumber(Minutes, &(NumberFormattingOptions), Culture));
+		TimeArguments.Add(TEXT("Seconds"), FText::AsNumber(Seconds, &(NumberFormattingOptions), Culture));
+		return FText::Format(TimespanFormatPattern, TimeArguments);
+	}
+	else
+	{
+		FText TimespanFormatPattern = NSLOCTEXT("Timespan", "FormatPattern2", "{Minutes}:{Seconds}");
+		FFormatNamedArguments TimeArguments;
+		TimeArguments.Add(TEXT("Minutes"), Minutes);
+		TimeArguments.Add(TEXT("Seconds"), FText::AsNumber(Seconds, &(NumberFormattingOptions), Culture));
+		return FText::Format(TimespanFormatPattern, TimeArguments);
+	}
 }
 
 FText FText::AsDateTime(const FDateTime& DateTime, const EDateTimeStyle::Type DateStyle, const EDateTimeStyle::Type TimeStyle, const FString& TimeZone, const FCulturePtr& TargetCulture)
diff --git a/Engine/Source/Runtime/Core/Private/Internationalization/Text.cpp b/Engine/Source/Runtime/Core/Private/Internationalization/Text.cpp
--- a/Engine/Source/Runtime/Core/Private/Internationalization/Text.cpp
+++ b/Engine/Source/Runtime/Core/Private/Internationalization/Text.cpp
@@ -428,6 +428,11 @@ public:
 
 		TOptional<int32> AsIndex() const
 		{
+			if (NameLen <= 0)
+			{
+				return TOptional<int32>();
+			}
+
 			int32 Index = 0;
 			for (int32 NameOffset = 0; NameOffset < NameLen; ++NameOffset)
 			{
@@ -841,7 +846,7 @@ FText FText::FormatInternal(FText Pattern, FFormatOrderedArguments Arguments, bo
 		EstimatedArgumentValuesLength += FTextFormatHelper::EstimateArgumentValueLength(Arg);
 	}
 
-	auto GetArgumentValue = [&Arguments](const FTextFormatHelper::FArgumentName& ArgumentName, int32 ArgumentNumber) -> const FFormatArgumentValue*
+	auto GetArgumentValue = [&Arguments, &Pattern](const FTextFormatHelper::FArgumentName& ArgumentName, int32 ArgumentNumber) -> const FFormatArgumentValue*
 	{
 		TOptional<int32> ArgumentIndex = ArgumentName.AsIndex();
 		if (!ArgumentIndex.IsSet())
@@ -850,6 +855,7 @@ FText FText::FormatInternal(FText Pattern, FFormatOrderedArguments Arguments, bo
 			// We have existing code that is incorrectly using names in the format string when providing ordered arguments
 			// ICU used to fallback to treating the index of the argument within the string as if it were the index specified 
 			// by the argument name, so we need to emulate that behavior to avoid breaking some format operations
+			UE_LOG(LogText, Warning, TEXT("Failed to parse argument \"%s\" as a number (using \"%d\" as a fallback). Please check your format string for errors: \"%s\"."), *FString(ArgumentName.NameLen, ArgumentName.NamePtr), ArgumentNumber, *Pattern.ToString());
 			ArgumentIndex = ArgumentNumber;
 		}
 		return Arguments.IsValidIndex(ArgumentIndex.GetValue()) ? &(Arguments[ArgumentIndex.GetValue()]) : nullptr;
diff --git a/Engine/Source/Runtime/Core/Private/Linux/LinuxPlatformCrashContext.cpp b/Engine/Source/Runtime/Core/Private/Linux/LinuxPlatformCrashContext.cpp
--- a/Engine/Source/Runtime/Core/Private/Linux/LinuxPlatformCrashContext.cpp
+++ b/Engine/Source/Runtime/Core/Private/Linux/LinuxPlatformCrashContext.cpp
@@ -333,7 +333,13 @@ void DLLEXPORT GenerateCrashInfoAndLaunchReporter(const FLinuxCrashContext & Con
 #if !IS_PROGRAM
 
 	// create a crash-specific directory
-	FString CrashInfoFolder = FString::Printf(TEXT("crashinfo-%s-pid-%d-%s"), FApp::GetGameName(), getpid(), *FGuid::NewGuid().ToString());
+	FString CrashGuid;
+	if (!FParse::Value(FCommandLine::Get(), TEXT("CrashGUID="), CrashGuid) || CrashGuid.Len() <= 0)
+	{
+		CrashGuid = FGuid::NewGuid().ToString();
+	}
+
+	FString CrashInfoFolder = FString::Printf(TEXT("crashinfo-%s-pid-%d-%s"), FApp::GetGameName(), getpid(), *CrashGuid);
 	FString CrashInfoAbsolute = FPaths::ConvertRelativePathToFull(CrashInfoFolder);
 	if (IFileManager::Get().MakeDirectory(*CrashInfoFolder))
 	{
diff --git a/Engine/Source/Runtime/Core/Private/Linux/LinuxPlatformProcess.cpp b/Engine/Source/Runtime/Core/Private/Linux/LinuxPlatformProcess.cpp
--- a/Engine/Source/Runtime/Core/Private/Linux/LinuxPlatformProcess.cpp
+++ b/Engine/Source/Runtime/Core/Private/Linux/LinuxPlatformProcess.cpp
@@ -596,6 +596,60 @@ TArray<FChildWaiterThread *> FChildWaiterThread::ChildWaiterThreadsArray;
 /** See FChildWaiterThread */
 FCriticalSection FChildWaiterThread::ChildWaiterThreadsArrayGuard;
 
+namespace LinuxPlatformProcess
+{
+	/**
+	 * This function tries to set exec permissions on the file (if it is missing them).
+	 * It exists because files copied manually from foreign filesystems (e.g. CrashReportClient) or unzipped from
+	 * certain arhcive types may lack +x, yet we still want to execute them.
+	 *
+	 * @param AbsoluteFilename absolute filename to the file in question
+	 *
+	 * @return true if we should attempt to execute the file, false if it is not worth even trying
+	 */	
+	bool AttemptToMakeExecIfNotAlready(const FString & AbsoluteFilename)
+	{
+		bool bWorthTryingToExecute = true;	// be conservative and let the OS decide in most cases
+
+		FTCHARToUTF8 AbsoluteFilenameUTF8Buffer(*AbsoluteFilename);
+		const char* AbsoluteFilenameUTF8 = AbsoluteFilenameUTF8Buffer.Get();
+
+		struct stat FilePerms;
+		if (UNLIKELY(stat(AbsoluteFilenameUTF8, &FilePerms) == -1))
+		{
+			int ErrNo = errno;
+			UE_LOG(LogHAL, Warning, TEXT("LinuxPlatformProcess::AttemptToMakeExecIfNotAlready: could not stat '%s', errno=%d (%s)"),
+				*AbsoluteFilename,
+				ErrNo,
+				ANSI_TO_TCHAR(strerror(ErrNo))
+				);
+		}
+		else
+		{
+			// Try to make a guess if we can execute the file. We are not trying to do the exact check,
+			// so if any of executable bits are set, assume it's executable
+			if ((FilePerms.st_mode & (S_IXUSR | S_IXGRP | S_IXOTH)) == 0)
+			{
+				// if no executable bits at all, try setting permissions
+				if (chmod(AbsoluteFilenameUTF8, FilePerms.st_mode | S_IXUSR) == -1)
+				{
+					int ErrNo = errno;
+					UE_LOG(LogHAL, Warning, TEXT("LinuxPlatformProcess::AttemptToMakeExecIfNotAlready: could not chmod +x '%s', errno=%d (%s)"),
+						*AbsoluteFilename,
+						ErrNo,
+						ANSI_TO_TCHAR(strerror(ErrNo))
+						);
+
+					// at this point, assume that execution will fail
+					bWorthTryingToExecute = false;
+				}
+			}
+		}
+
+		return bWorthTryingToExecute;
+	}
+}
+
 FProcHandle FLinuxPlatformProcess::CreateProc(const TCHAR* URL, const TCHAR* Parms, bool bLaunchDetached, bool bLaunchHidden, bool bLaunchReallyHidden, uint32* OutProcessID, int32 PriorityModifier, const TCHAR* OptionalWorkingDirectory, void* PipeWriteChild, void * PipeReadChild)
 {
 	// @TODO bLaunchHidden bLaunchReallyHidden are not handled
@@ -611,6 +665,12 @@ FProcHandle FLinuxPlatformProcess::CreateProc(const TCHAR* URL, const TCHAR* Par
 		return FProcHandle();
 	}
 
+	// check if it's worth attemptting to execute the file
+	if (!LinuxPlatformProcess::AttemptToMakeExecIfNotAlready(ProcessPath))
+	{
+		return FProcHandle();
+	}
+
 	FString Commandline = ProcessPath;
 	Commandline += TEXT(" ");
 	Commandline += Parms;
diff --git a/Engine/Source/Runtime/Core/Private/Misc/ConfigCacheIni.cpp b/Engine/Source/Runtime/Core/Private/Misc/ConfigCacheIni.cpp
--- a/Engine/Source/Runtime/Core/Private/Misc/ConfigCacheIni.cpp
+++ b/Engine/Source/Runtime/Core/Private/Misc/ConfigCacheIni.cpp
@@ -105,8 +105,25 @@ class FTextFriendHelper
 		}
 		else
 		{
-			return FString::Printf( TEXT( "NSLOCTEXT(\"%s\",\"%s\",\"%s\")" ),
-				*FTextInspector::GetNamespace(Text).Get(TEXT("")), *FTextInspector::GetKey(Text).Get(TEXT("")), *Str );
+			FTextDisplayStringRef DisplayString = FTextInspector::GetSharedDisplayString(Text);
+
+			FString Namespace;
+			FString Key;
+			const bool FoundNamespaceAndKey = FTextLocalizationManager::Get().FindNamespaceAndKeyFromDisplayString(DisplayString, Namespace, Key);
+
+			// If this has no namespace or key, attempt to give it a GUID for a key and register it.
+			if (!FoundNamespaceAndKey && GIsEditor)
+			{
+				Key = FGuid::NewGuid().ToString();
+				if (!FTextLocalizationManager::Get().AddDisplayString(DisplayString, Namespace, Key))
+				{
+					// Could not add display string, reset namespace and key.
+					Namespace.Empty();
+					Key.Empty();
+				}
+			}
+
+			return FString::Printf( TEXT( "NSLOCTEXT(\"%s\",\"%s\",\"%s\")" ), *Namespace, *Key, *Str );
 		}
 #undef LOC_DEFINE_REGION
 	}
@@ -593,7 +610,7 @@ bool FConfigFile::ShouldExportQuotedString(const FString& PropertyValue) const
 }
 
 
-#if !UE_BUILD_SHIPPING
+#if ALLOW_INI_OVERRIDE_FROM_COMMANDLINE
 
 /** A collection of identifiers which will help us parse the commandline opions. */
 namespace CommandlineOverrideSpecifiers
@@ -859,7 +876,7 @@ bool PropertySetFromCommandlineOption(const FConfigFile* InConfigFile, const FSt
 {
 	bool bFromCommandline = false;
 
-#if !UE_BUILD_SHIPPING
+#if ALLOW_INI_OVERRIDE_FROM_COMMANDLINE
 	for (const FConfigCommandlineOverride& CommandlineOverride : InConfigFile->CommandlineOptions)
 	{
 		if (CommandlineOverride.PropertyKey.Equals(InPropertyName.ToString(), ESearchCase::IgnoreCase) &&
@@ -870,7 +887,7 @@ bool PropertySetFromCommandlineOption(const FConfigFile* InConfigFile, const FSt
 			bFromCommandline = true;
 		}
 	}
-#endif // !UE_BUILD_SHIPPING
+#endif // ALLOW_INI_OVERRIDE_FROM_COMMANDLINE
 
 	return bFromCommandline;
 }
@@ -1143,7 +1160,11 @@ bool FConfigFile::GetText( const TCHAR* Section, const TCHAR* Key, FText& Value
 	{
 		return false;
 	}
-	return FParse::Text( **PairString, Value, Section );
+	if( !FParse::Text( **PairString, Value, Section ) )
+	{
+		Value = FText::FromString( **PairString );
+	}
+	return true;
 }
 
 bool FConfigFile::GetInt64( const TCHAR* Section, const TCHAR* Key, int64& Value ) const
@@ -1707,7 +1728,11 @@ bool FConfigCacheIni::GetText( const TCHAR* Section, const TCHAR* Key, FText& Va
 	{
 		return false;
 	}
-	return FParse::Text( **PairString, Value, Section );
+	if( !FParse::Text( **PairString, Value, Section ) )
+	{
+		Value = FText::FromString( **PairString );
+	}
+	return true;
 }
 
 bool FConfigCacheIni::GetSection( const TCHAR* Section, TArray<FString>& Result, const FString& Filename )
@@ -2609,7 +2634,7 @@ static bool GenerateDestIniFile(FConfigFile& DestConfigFile, const FString& Dest
 	}
 	LoadAnIniFile(DestIniFilename, DestConfigFile);
 
-#if !UE_BUILD_SHIPPING
+#if ALLOW_INI_OVERRIDE_FROM_COMMANDLINE
 	// process any commandline overrides
 	OverrideFromCommandline(&DestConfigFile, DestIniFilename);
 #endif
diff --git a/Engine/Source/Runtime/Core/Private/Misc/CoreMisc.cpp b/Engine/Source/Runtime/Core/Private/Misc/CoreMisc.cpp
--- a/Engine/Source/Runtime/Core/Private/Misc/CoreMisc.cpp
+++ b/Engine/Source/Runtime/Core/Private/Misc/CoreMisc.cpp
@@ -559,6 +559,7 @@ bool FFileHelper::LoadANSITextFileToStrings(const TCHAR* InFilename, IFileManage
 bool FCommandLine::bIsInitialized = false;
 TCHAR FCommandLine::CmdLine[FCommandLine::MaxCommandLineSize] = TEXT("");
 TCHAR FCommandLine::OriginalCmdLine[FCommandLine::MaxCommandLineSize] = TEXT("");
+TCHAR FCommandLine::LoggingCmdLine[FCommandLine::MaxCommandLineSize] = TEXT("");
 FString FCommandLine::SubprocessCommandLine(TEXT(" -Multiprocess"));
 
 bool FCommandLine::IsInitialized()
@@ -572,6 +573,12 @@ const TCHAR* FCommandLine::Get()
 	return CmdLine;
 }
 
+const TCHAR* FCommandLine::GetForLogging()
+{
+	UE_CLOG(!bIsInitialized, LogInit, Fatal, TEXT("Attempting to get the command line but it hasn't been initialized yet."));
+	return LoggingCmdLine;
+}
+
 const TCHAR* FCommandLine::GetOriginal()
 {
 	UE_CLOG(!bIsInitialized, LogInit, Fatal, TEXT("Attempting to get the command line but it hasn't been initialized yet."));
@@ -586,6 +593,7 @@ bool FCommandLine::Set(const TCHAR* NewCommandLine)
 	}
 
 	FCString::Strncpy( CmdLine, NewCommandLine, ARRAY_COUNT(CmdLine) );
+	FCString::Strncpy(LoggingCmdLine, NewCommandLine, ARRAY_COUNT(LoggingCmdLine));
 	// If configured as part of the build, strip out any unapproved args
 	WhitelistCommandLines();
 
@@ -615,6 +623,7 @@ void FCommandLine::Append(const TCHAR* AppendString)
 
 #if WANTS_COMMANDLINE_WHITELIST
 TArray<FString> FCommandLine::ApprovedArgs;
+TArray<FString> FCommandLine::FilterArgsForLogging;
 
 #ifdef OVERRIDE_COMMANDLINE_WHITELIST
 /**
@@ -630,6 +639,19 @@ const TCHAR* OverrideList = TEXT(OVERRIDE_COMMANDLINE_WHITELIST);
 const TCHAR* OverrideList = TEXT("-fullscreen /windowed");
 #endif
 
+#ifdef FILTER_COMMANDLINE_LOGGING
+/**
+ * When overriding this setting make sure that your define looks like the following in your .cs file:
+ *
+ *		OutCPPEnvironmentConfiguration.Definitions.Add("FILTER_COMMANDLINE_LOGGING=\"-arg1 -arg2 -arg3 -arg4\"");
+ *
+ * The important part is the \" as they quotes get stripped off by the compiler without them
+ */
+const TCHAR* FilterForLoggingList = TEXT(FILTER_COMMANDLINE_LOGGING);
+#else
+const TCHAR* FilterForLoggingList = TEXT("");
+#endif
+
 void FCommandLine::WhitelistCommandLines()
 {
 	if (ApprovedArgs.Num() == 0)
@@ -637,12 +659,20 @@ void FCommandLine::WhitelistCommandLines()
 		TArray<FString> Ignored;
 		FCommandLine::Parse(OverrideList, ApprovedArgs, Ignored);
 	}
+	if (FilterArgsForLogging.Num() == 0)
+	{
+		TArray<FString> Ignored;
+		FCommandLine::Parse(FilterForLoggingList, FilterArgsForLogging, Ignored);
+	}
 	// Process the original command line
 	TArray<FString> OriginalList = FilterCommandLine(OriginalCmdLine);
 	BuildWhitelistCommandLine(OriginalCmdLine, ARRAY_COUNT(OriginalCmdLine), OriginalList);
 	// Process the current command line
 	TArray<FString> CmdList = FilterCommandLine(CmdLine);
 	BuildWhitelistCommandLine(CmdLine, ARRAY_COUNT(CmdLine), CmdList);
+	// Process the command line for logging purposes
+	TArray<FString> LoggingCmdList = FilterCommandLineForLogging(LoggingCmdLine);
+	BuildWhitelistCommandLine(LoggingCmdLine, ARRAY_COUNT(LoggingCmdLine), LoggingCmdList);
 }
 
 TArray<FString> FCommandLine::FilterCommandLine(TCHAR* CommandLine)
@@ -654,9 +684,41 @@ TArray<FString> FCommandLine::FilterCommandLine(TCHAR* CommandLine)
 	// Remove any that are not in our approved list
 	for (int32 Index = 0; Index < ParsedList.Num(); Index++)
 	{
-		if (!ApprovedArgs.Contains(ParsedList[Index]))
+		bool bFound = false;
+		for (auto ApprovedArg : ApprovedArgs)
+		{
+			if (ParsedList[Index].StartsWith(ApprovedArg))
+			{
+				bFound = true;
+				break;
+			}
+		}
+		if (!bFound)
 		{
 			ParsedList.RemoveAt(Index);
+			Index--;
+		}
+	}
+	return ParsedList;
+}
+
+TArray<FString> FCommandLine::FilterCommandLineForLogging(TCHAR* CommandLine)
+{
+	TArray<FString> Ignored;
+	TArray<FString> ParsedList;
+	// Parse the command line list
+	FCommandLine::Parse(CommandLine, ParsedList, Ignored);
+	// Remove any that are not in our approved list
+	for (int32 Index = 0; Index < ParsedList.Num(); Index++)
+	{
+		for (auto Filter : FilterArgsForLogging)
+		{
+			if (ParsedList[Index].StartsWith(Filter))
+			{
+				ParsedList.RemoveAt(Index);
+				Index--;
+				break;
+			}
 		}
 	}
 	return ParsedList;
@@ -666,7 +728,7 @@ void FCommandLine::BuildWhitelistCommandLine(TCHAR* CommandLine, uint32 ArrayCou
 {
 	check(ArrayCount > 0);
 	// Zero the whole string
-	FMemory::Memzero(OriginalCmdLine, sizeof(TCHAR) * ArrayCount);
+	FMemory::Memzero(CommandLine, sizeof(TCHAR) * ArrayCount);
 
 	uint32 StartIndex = 0;
 	for (auto Arg : FilteredArgs)
@@ -1004,7 +1066,7 @@ void GenerateConvenientWindowedResolutions(const FDisplayMetrics& InDisplayMetri
 	bool bInPortraitMode = InDisplayMetrics.PrimaryDisplayWidth < InDisplayMetrics.PrimaryDisplayHeight;
 
 	// Generate windowed resolutions as scaled versions of primary monitor size
-	static const float Scales[] = { 3.0f / 6.0f, 4.0 / 6.0f, 5.0f / 6.0f };
+	static const float Scales[] = { 3.0f / 6.0f, 4.0f / 6.0f, 4.5f / 6.0f, 5.0f / 6.0f };
 	static const float Ratios[] = { 9.0f, 10.0f, 12.0f };
 	static const float MinWidth = 1280.0f;
 	static const float MinHeight = 720.0f; // UI layout doesn't work well below this, as the accept/cancel buttons go off the bottom of the screen
diff --git a/Engine/Source/Runtime/Core/Private/Misc/MemStack.cpp b/Engine/Source/Runtime/Core/Private/Misc/MemStack.cpp
--- a/Engine/Source/Runtime/Core/Private/Misc/MemStack.cpp
+++ b/Engine/Source/Runtime/Core/Private/Misc/MemStack.cpp
@@ -9,6 +9,249 @@ DECLARE_MEMORY_STAT(TEXT("PageAllocator Used"), STAT_PageAllocatorUsed, STATGROU
 TLockFreeFixedSizeAllocator<FPageAllocator::PageSize, FThreadSafeCounter> FPageAllocator::TheAllocator;
 TLockFreeFixedSizeAllocator<FPageAllocator::SmallPageSize, FThreadSafeCounter> FPageAllocator::TheSmallAllocator;
 
+
+#define USE_MEMSTACK_PUGATORY (0)
+
+#if PLATFORM_WINDOWS && USE_MEMSTACK_PUGATORY
+
+#define MEMSTACK_PUGATORY_COMPILED_IN (1)
+
+#include "AllowWindowsPlatformTypes.h"
+#include <Psapi.h>
+
+TLockFreePointerListUnordered<void> TheAllocatorReady;
+TLockFreePointerListUnordered<void> TheSmallAllocatorReady;
+
+#define FRAMES_TO_WAIT (4)
+TLockFreePointerListUnordered<void> TheAllocatorPurgatory[FRAMES_TO_WAIT];
+TLockFreePointerListUnordered<void> TheSmallAllocatorPurgatory[FRAMES_TO_WAIT];
+
+FThreadSafeCounter BigBlocksOutstanding;
+FThreadSafeCounter SmallBlocksOutstanding;
+
+FThreadSafeCounter BigBlocksInUse;
+FThreadSafeCounter SmallBlocksInUse;
+
+static int32 GMemStackProtectionTestWrite = 0;
+static FAutoConsoleVariableRef CVarTestProtectionWrite(
+	TEXT("memstack.TestProtection.Write"),
+	GMemStackProtectionTestWrite,
+	TEXT("For debugging. If > 0, then a newly freed memstack page will be written to. This will (correctly) cause a GPF the next time a memstack frees a page."));
+
+static int32 GMemStackProtectionTestRead = 0;
+static FAutoConsoleVariableRef CVarTestProtectionRead(
+	TEXT("memstack.TestProtection.Read"),
+	GMemStackProtectionTestRead,
+	TEXT("For debugging. If > 0, then a newly freed memstack page will be read from. This will (correctly) cause a GPF the next time a memstack frees a page."));
+
+static int32 GMemStackProtection = 0;
+static int32 GMemStackProtectionLatched = 0;
+
+// used to prevent the optimizer from skipping a read.
+int32 SomethingWithExternalLinkage = 0;
+
+int32 FrameIndex()
+{
+	static uint32 LastFrame = MAX_uint32;
+	uint32 LocalLastFrame = LastFrame;
+	uint32 CurrentFrame = GFrameNumber;
+	if (CurrentFrame != LocalLastFrame && LocalLastFrame > FRAMES_TO_WAIT)
+	{
+		uint32 ToEmpty = (CurrentFrame - (FRAMES_TO_WAIT - 2)) % FRAMES_TO_WAIT; // here we don't empty the one we are about to fill; we do the one before that so we are sure we aren't emptying stuff that is being filled
+		while (1)
+		{
+			void *Block = TheAllocatorPurgatory[ToEmpty].Pop();
+			if (!Block)
+			{
+				break;
+			}
+			DWORD Old;
+			VirtualProtect(Block, FPageAllocator::PageSize, PAGE_READWRITE, &Old);
+			check(Old == PAGE_NOACCESS);
+			TheAllocatorReady.Push(Block);
+		}
+		while (1)
+		{
+			void *Block = TheSmallAllocatorPurgatory[ToEmpty].Pop();
+			if (!Block)
+			{
+				break;
+			}
+			DWORD Old;
+			VirtualProtect(Block, FPageAllocator::SmallPageSize, PAGE_READWRITE, &Old);
+			check(Old == PAGE_NOACCESS);
+			TheSmallAllocatorReady.Push(Block);
+		}
+	}
+	return int32(CurrentFrame % FRAMES_TO_WAIT);
+}
+
+void *FPageAllocator::Alloc()
+{
+	check(GMemStackProtectionLatched); // attempt to allocate a page before we know if we are doing protected mode or not
+	void *Result;
+	if (GMemStackProtection)
+	{
+		Result = TheAllocatorReady.Pop();
+		if (!Result)
+		{
+			Result = VirtualAlloc(NULL, PageSize, MEM_COMMIT, PAGE_READWRITE);
+			BigBlocksOutstanding.Increment();
+		}
+		BigBlocksInUse.Increment();
+	}
+	else
+	{
+		Result = TheAllocator.Allocate();
+	}
+	STAT(UpdateStats());
+	return Result;
+}
+
+void FPageAllocator::Free(void *Mem)
+{
+	if (GMemStackProtection)
+	{
+		BigBlocksInUse.Decrement();
+		DWORD Old;
+		verify(VirtualProtect(Mem, PageSize, PAGE_NOACCESS, &Old));
+		check(Old == PAGE_READWRITE);
+		TheAllocatorPurgatory[FrameIndex()].Push(Mem);
+
+		if (GMemStackProtectionTestWrite)
+		{
+			*(uint8*)Mem = 0;
+		}
+		if (GMemStackProtectionTestRead)
+		{
+			SomethingWithExternalLinkage += *(int32*)Mem;
+		}
+	}
+	else
+	{
+		TheAllocator.Free(Mem);
+	}
+	STAT(UpdateStats());
+}
+
+void *FPageAllocator::AllocSmall()
+{
+	check(GMemStackProtectionLatched); // attempt to allocate a page before we know if we are doing protected mode or not
+	void *Result;
+	if (GMemStackProtection)
+	{
+		Result = TheSmallAllocatorReady.Pop();
+		if (!Result)
+		{
+			Result = VirtualAlloc(NULL, SmallPageSize, MEM_COMMIT, PAGE_READWRITE);
+			SmallBlocksOutstanding.Increment();
+		}
+		SmallBlocksInUse.Increment();
+	}
+	else
+	{
+		Result = TheSmallAllocator.Allocate();
+	}
+	STAT(UpdateStats());
+	return Result;
+}
+
+void FPageAllocator::FreeSmall(void *Mem)
+{
+	if (GMemStackProtection)
+	{
+		SmallBlocksInUse.Decrement();
+		DWORD Old;
+		verify(VirtualProtect(Mem, SmallPageSize, PAGE_NOACCESS, &Old));
+		check(Old == PAGE_READWRITE);
+		TheSmallAllocatorPurgatory[FrameIndex()].Push(Mem);
+
+		if (GMemStackProtectionTestWrite)
+		{
+			*(uint8*)Mem = 0;
+		}
+		if (GMemStackProtectionTestRead)
+		{
+			SomethingWithExternalLinkage += *(int32*)Mem;
+		}
+	}
+	else
+	{
+		TheSmallAllocator.Free(Mem);
+	}
+	STAT(UpdateStats());
+}
+
+uint64 FPageAllocator::BytesUsed()
+{
+	return uint64(BigBlocksInUse.GetValue()) * PageSize + uint64(SmallBlocksInUse.GetValue()) * SmallPageSize +
+		uint64(TheAllocator.GetNumUsed().GetValue()) * PageSize + uint64(TheSmallAllocator.GetNumUsed().GetValue()) * SmallPageSize;
+}
+
+uint64 FPageAllocator::BytesFree()
+{
+	return uint64(BigBlocksOutstanding.GetValue() - BigBlocksInUse.GetValue()) * PageSize + uint64(SmallBlocksOutstanding.GetValue() - SmallBlocksInUse.GetValue()) * SmallPageSize +
+		uint64(TheAllocator.GetNumFree().GetValue()) * PageSize + uint64(TheSmallAllocator.GetNumFree().GetValue()) * SmallPageSize;
+}
+
+#include "HideWindowsPlatformTypes.h"
+
+#else
+
+#define MEMSTACK_PUGATORY_COMPILED_IN (0)
+
+void *FPageAllocator::Alloc()
+{
+	void *Result = TheAllocator.Allocate();
+	STAT(UpdateStats());
+	return Result;
+}
+
+void FPageAllocator::Free(void *Mem)
+{
+	TheAllocator.Free(Mem);
+	STAT(UpdateStats());
+}
+
+void *FPageAllocator::AllocSmall()
+{
+	void *Result = TheSmallAllocator.Allocate();
+	STAT(UpdateStats());
+	return Result;
+}
+
+void FPageAllocator::FreeSmall(void *Mem)
+{
+	TheSmallAllocator.Free(Mem);
+	STAT(UpdateStats());
+}
+
+uint64 FPageAllocator::BytesUsed()
+{
+	return uint64(TheAllocator.GetNumUsed().GetValue()) * PageSize + uint64(TheSmallAllocator.GetNumUsed().GetValue()) * SmallPageSize;
+}
+
+uint64 FPageAllocator::BytesFree()
+{
+	return uint64(TheAllocator.GetNumFree().GetValue()) * PageSize + uint64(TheSmallAllocator.GetNumFree().GetValue()) * SmallPageSize;
+}
+
+#endif
+
+void FPageAllocator::LatchProtectedMode()
+{
+#if MEMSTACK_PUGATORY_COMPILED_IN
+	const bool bNoProtection = FParse::Param(FCommandLine::Get(), TEXT("NoProtectMemStack"));
+	GMemStackProtection = 1;
+	if (bNoProtection)
+	{
+		UE_LOG(LogTemp, Display, TEXT("Enabling purgatory and virtual memory protection to catch stale pointers to mem stacks items (and other users of the page allocator)."));
+		GMemStackProtection = 0;
+	}
+	GMemStackProtectionLatched = 1;
+#endif
+}
+
 #if STATS
 void FPageAllocator::UpdateStats()
 {
@@ -17,6 +260,7 @@ void FPageAllocator::UpdateStats()
 }
 #endif
 
+
 /*-----------------------------------------------------------------------------
 	FMemStack implementation.
 -----------------------------------------------------------------------------*/
diff --git a/Engine/Source/Runtime/Core/Private/Misc/Parse.cpp b/Engine/Source/Runtime/Core/Private/Misc/Parse.cpp
--- a/Engine/Source/Runtime/Core/Private/Misc/Parse.cpp
+++ b/Engine/Source/Runtime/Core/Private/Misc/Parse.cpp
@@ -355,7 +355,7 @@ bool FParse::Text( const TCHAR* Str, FText& Value, const TCHAR* Namespace )
 			FString SourceString;
 			const int32 ExpectedStringCount = bFoundNSLocText ? 3 : 2;
 
-			while( *Str && *Str != ')' && !bError )
+			while( *Str && (*Str != ')' || bInQuotes) && !bError )
 			{
 				const TCHAR c = *Str;
 				if( bInQuotes )
@@ -457,11 +457,14 @@ bool FParse::Text( const TCHAR* Str, FText& Value, const TCHAR* Namespace )
 
 			if( *Str == ')' && !bError && StringCount == ExpectedStringCount )
 			{
-				if ( !FText::FindText( bFoundNSLocText ? NamespaceString : Namespace, KeyString, /*OUT*/Value ) )
-				{ 
-					Value = FText::FromString( ParsedString );
+				if (KeyString.IsEmpty())
+				{
+					Value = FText::AsCultureInvariant( ParsedString );
+				}
+				else
+				{
+					Value = FInternationalization::Get().ForUseOnlyByLocMacroAndGraphNodeTextLiterals_CreateText( *ParsedString, *(bFoundNSLocText ? NamespaceString : Namespace), *KeyString );
 				}
-
 				return true;
 			}
 		}
diff --git a/Engine/Source/Runtime/Core/Private/Misc/Paths.cpp b/Engine/Source/Runtime/Core/Private/Misc/Paths.cpp
--- a/Engine/Source/Runtime/Core/Private/Misc/Paths.cpp
+++ b/Engine/Source/Runtime/Core/Private/Misc/Paths.cpp
@@ -490,6 +490,26 @@ FString FPaths::GetPath(const FString& InPath)
 	return Result;
 }
 
+FString FPaths::ChangeExtension(const FString& InPath, const FString& InNewExtension)
+{
+	int32 Pos = INDEX_NONE;
+	if (InPath.FindLastChar('.', Pos))
+	{
+		FString Result = InPath.Left(Pos);
+
+		if (InNewExtension.Len() && InNewExtension[0] != '.')
+		{
+			Result += '.';
+		}
+
+		Result += InNewExtension;
+
+		return Result;
+	}
+
+	return InPath;
+}
+
 bool FPaths::FileExists(const FString& InPath)
 {
 	return IFileManager::Get().FileExists(*InPath);
diff --git a/Engine/Source/Runtime/Core/Private/ProfilingDebugging/ExternalProfiler.cpp b/Engine/Source/Runtime/Core/Private/ProfilingDebugging/ExternalProfiler.cpp
--- a/Engine/Source/Runtime/Core/Private/ProfilingDebugging/ExternalProfiler.cpp
+++ b/Engine/Source/Runtime/Core/Private/ProfilingDebugging/ExternalProfiler.cpp
@@ -74,7 +74,7 @@ FExternalProfiler* FActiveExternalProfilerBase::GetActiveProfiler()
 		}
 		else
 		{
-			UE_LOG(LogExternalProfiler, Warning, TEXT("No registered external profiler was matched with a command-line switch (or the DLL could not be loaded).  External profiling features will not be available."));
+			UE_LOG(LogExternalProfiler, Log, TEXT("No external profilers were discovered.  External profiling features will not be available."));
 		}
 
 		// Don't try to initialize again this session
diff --git a/Engine/Source/Runtime/Core/Private/ProfilingDebugging/ProfilingHelpers.cpp b/Engine/Source/Runtime/Core/Private/ProfilingDebugging/ProfilingHelpers.cpp
--- a/Engine/Source/Runtime/Core/Private/ProfilingDebugging/ProfilingHelpers.cpp
+++ b/Engine/Source/Runtime/Core/Private/ProfilingDebugging/ProfilingHelpers.cpp
@@ -63,6 +63,10 @@ void SendDataToPCViaUnrealConsole( const FString& NotifyType, const FString& Ful
 }
 
 
+FString CreateProfileFilename(const FString& InFileExtension, bool bIncludeDateForDirectoryName)
+{
+	return CreateProfileFilename(TEXT(""), InFileExtension, bIncludeDateForDirectoryName);
+}
 
 /** 
  * This will generate the profiling file name that will work with limited filename sizes on consoles.
@@ -71,7 +75,7 @@ void SendDataToPCViaUnrealConsole( const FString& NotifyType, const FString& Ful
  * @param ProfilingType this is the type of profiling file this is
  * 
  **/
-FString CreateProfileFilename( const FString& InFileExtension, bool bIncludeDateForDirectoryName )
+FString CreateProfileFilename( const FString& InFilename, const FString& InFileExtension, bool bIncludeDateForDirectoryName )
 {
 	FString Retval;
 
@@ -117,7 +121,15 @@ FString CreateProfileFilename( const FString& InFileExtension, bool bIncludeDate
 	FolderNameOfProfileNoDate = FolderNameOfProfileNoDate.Right(MaxFilenameLen);
 
 
-	FString NameOfProfile = FString::Printf(TEXT("%s-%s-%s"), *MapNameStr, *PlatformStr, *FDateTime::Now().ToString(TEXT("%d-%H.%M.%S")));
+	FString NameOfProfile;
+	if (InFilename.IsEmpty())
+	{
+		NameOfProfile = FString::Printf(TEXT("%s-%s-%s"), *MapNameStr, *PlatformStr, *FDateTime::Now().ToString(TEXT("%d-%H.%M.%S")));
+	}
+	else
+	{
+		NameOfProfile = InFilename;
+	}
 	NameOfProfile = NameOfProfile.Right(MaxFilenameLen);
 
 	FString FileNameWithExtension = FString::Printf( TEXT("%s%s"), *NameOfProfile, *InFileExtension );
@@ -148,7 +160,7 @@ FString CreateProfileDirectoryAndFilename( const FString& InSubDirectoryName, co
 	check(GGetMapNameDelegate.IsBound());
 	MapNameStr = GGetMapNameDelegate.Execute();
 #endif		// WITH_ENGINE
-	const FString PlatformStr = FString(TEXT("PC"));
+	const FString PlatformStr(FPlatformProperties::PlatformName());
 
 
 	// create Profiling dir and sub dir
diff --git a/Engine/Source/Runtime/Core/Private/Stats/StatsCommand.cpp b/Engine/Source/Runtime/Core/Private/Stats/StatsCommand.cpp
--- a/Engine/Source/Runtime/Core/Private/Stats/StatsCommand.cpp
+++ b/Engine/Source/Runtime/Core/Private/Stats/StatsCommand.cpp
@@ -1050,9 +1050,13 @@ struct FHUDGroupManager
 		else
 		{
 			TSet<FName> HierEnabledItems;
-			for( auto It = EnabledGroups.CreateConstIterator(); It; ++It )
+			for( auto It = EnabledGroups.CreateIterator(); It; ++It )
 			{
-				HierEnabledItems.Append( It.Value().EnabledItems );
+				{
+					QUICK_SCOPE_CYCLE_COUNTER(STAT_GetStatsForGroup_EveryFrame);
+					GetStatsForGroup(It.Value().EnabledItems, It.Key());
+				}
+				HierEnabledItems.Append(It.Value().EnabledItems);
 			}
 		
 			
@@ -1421,7 +1425,7 @@ struct FHUDGroupManager
 
 	void GetStatsForGroup( TSet<FName>& out_EnabledItems, const FName GroupName )
 	{
-		out_EnabledItems.Empty();
+		out_EnabledItems.Reset();
 	
 		TArray<FName> GroupItems;
 		Stats.Groups.MultiFind( GroupName, GroupItems );
@@ -1462,6 +1466,7 @@ bool FGroupFilter::IsRoot(const FName& MessageName) const
 
 static int32 MaxDepth = MAX_int32;
 static FString NameFilter;
+static FString LeafFilter;
 static FDelegateHandle DumpFrameDelegateHandle;
 static FDelegateHandle DumpCPUDelegateHandle;
 
@@ -1535,7 +1540,18 @@ struct FDumpMultiple
 			{
 				Stack->CullByCycles( int64( DumpCull / FPlatformTime::ToMilliseconds( 1 ) ) );
 			}
-			Stack->DebugPrint(*NameFilter, MaxDepth);
+			if (!NameFilter.IsEmpty() && !LeafFilter.IsEmpty())
+			{
+				UE_LOG(LogStats, Log, TEXT("You can't have both a root and a leaf filter (though this wouldn't be hard to add)."));
+			}
+			else if (!LeafFilter.IsEmpty())
+			{
+				Stack->DebugPrintLeafFilter(*LeafFilter);
+			}
+			else
+			{
+				Stack->DebugPrint(*NameFilter, MaxDepth);
+			}
 			delete Stack;
 			Stack = NULL;
 		}
@@ -1693,8 +1709,10 @@ static void StatCmd(FString InCmd, bool bStatCommand)
 		DumpCull = 1.0f;
 		MaxDepth = MAX_int32;
 		NameFilter.Empty();
+		LeafFilter.Empty();
 
 		FParse::Value(Cmd, TEXT("ROOT="), NameFilter);
+		FParse::Value(Cmd, TEXT("LEAF="), LeafFilter);
 		FParse::Value(Cmd, TEXT("MS="), DumpCull);
 		FParse::Value(Cmd, TEXT("DEPTH="), MaxDepth);
 		if (FParse::Command(&Cmd, TEXT("DUMPFRAME")))
diff --git a/Engine/Source/Runtime/Core/Private/Stats/StatsData.cpp b/Engine/Source/Runtime/Core/Private/Stats/StatsData.cpp
--- a/Engine/Source/Runtime/Core/Private/Stats/StatsData.cpp
+++ b/Engine/Source/Runtime/Core/Private/Stats/StatsData.cpp
@@ -333,7 +333,7 @@ void FRawStatStackNode::DebugPrint(TCHAR const* Filter, int32 InMaxDepth, int32
 			UE_LOG(LogStats, Log, TEXT("%s%s"), FCString::Spc(Depth*2), *TmpDebugStr);
 		}
 
-		static int64 MinPrint = int64( .004f / FPlatformTime::ToMilliseconds( 1 ) + 0.5f );
+		static int64 MinPrint = -1;
 		if (Children.Num())
 		{
 			TArray<FRawStatStackNode*> ChildArray;
@@ -365,6 +365,60 @@ void FRawStatStackNode::DebugPrint(TCHAR const* Filter, int32 InMaxDepth, int32
 	}
 }
 
+void FRawStatStackNode::DebugPrintLeafFilter(TCHAR const* Filter) const
+{
+	TArray<FString> Stack;
+	DebugPrintLeafFilterInner(Filter, 0, Stack);
+}
+
+void FRawStatStackNode::DebugPrintLeafFilterInner(TCHAR const* Filter, int32 Depth, TArray<FString>& Stack) const
+{
+	{
+		FString TmpDebugStr = FStatsUtils::DebugPrint(Meta);
+		Stack.Add(TmpDebugStr);
+	}
+	if (!Filter || !*Filter)
+	{
+		int32 Offset = 1 + Depth - Stack.Num();
+		check(Offset >= 0);
+		for (int32 Index = 0; Index < Stack.Num(); Index++)
+		{
+			UE_LOG(LogStats, Log, TEXT("%s%s"), FCString::Spc((Index + Offset) * 2), *Stack[Index]);
+		}
+		Stack.Reset();
+	}
+	else
+	{
+		static int64 MinPrint = -1;
+		if (Children.Num())
+		{
+			TArray<FRawStatStackNode*> ChildArray;
+			Children.GenerateValueArray(ChildArray);
+			ChildArray.Sort(FStatDurationComparer<FRawStatStackNode>());
+			for (int32 Index = 0; Index < ChildArray.Num(); Index++)
+			{
+				if (ChildArray[Index]->Meta.GetValue_Duration() < MinPrint)
+				{
+					break;
+				}
+				if (ChildArray[Index]->Meta.NameAndInfo.GetRawName().ToString().Contains(Filter))
+				{
+					ChildArray[Index]->DebugPrintLeafFilterInner(nullptr, Depth + 1, Stack);
+				}
+				else
+				{
+					ChildArray[Index]->DebugPrintLeafFilterInner(Filter, Depth + 1, Stack);
+				}
+			}
+		}
+		if (Stack.Num())
+		{
+			Stack.Pop();
+		}
+	}
+}
+
+
 void FRawStatStackNode::Encode(TArray<FStatMessage>& OutStats) const
 {
 	FStatMessage* NewStat = new (OutStats) FStatMessage(Meta);
diff --git a/Engine/Source/Runtime/Core/Private/Windows/WindowsPlatformMemory.cpp b/Engine/Source/Runtime/Core/Private/Windows/WindowsPlatformMemory.cpp
--- a/Engine/Source/Runtime/Core/Private/Windows/WindowsPlatformMemory.cpp
+++ b/Engine/Source/Runtime/Core/Private/Windows/WindowsPlatformMemory.cpp
@@ -15,7 +15,7 @@
 #endif // ENABLE_WIN_ALLOC_TRACKING
 
 #if !FORCE_ANSI_ALLOCATOR
-#include "MallocBinned2.h"
+#include "MallocBinned.h"
 #endif
 
 #include "AllowWindowsPlatformTypes.h"
@@ -81,7 +81,7 @@ FMalloc* FWindowsPlatformMemory::BaseAllocator()
 #elif (WITH_EDITORONLY_DATA || IS_PROGRAM) && TBB_ALLOCATOR_ALLOWED
 	return new FMallocTBB();
 #else
-	return new FMallocBinned2((uint32)(GetConstants().PageSize&MAX_uint32), (uint64)MAX_uint32+1);
+	return new FMallocBinned((uint32)(GetConstants().PageSize&MAX_uint32), (uint64)MAX_uint32+1);
 #endif
 }
 
diff --git a/Engine/Source/Runtime/Core/Private/Windows/WindowsPlatformMisc.cpp b/Engine/Source/Runtime/Core/Private/Windows/WindowsPlatformMisc.cpp
--- a/Engine/Source/Runtime/Core/Private/Windows/WindowsPlatformMisc.cpp
+++ b/Engine/Source/Runtime/Core/Private/Windows/WindowsPlatformMisc.cpp
@@ -2388,10 +2388,58 @@ FString FWindowsPlatformMisc::GetPrimaryGPUBrand()
 			DisplayDevice.cb = sizeof( DisplayDevice );
 			DeviceIndex++;
 		}
-
 	}
+
 	return PrimaryGPUBrand;
 }
+
+void FWindowsPlatformMisc::GetGPUDriverInfo(const FString DeviceDescription, FString& InternalDriverVersion, FString& UserDriverVersion, FString& DriverDate)
+{
+	// to prevent bugs
+	InternalDriverVersion.Empty();
+	UserDriverVersion.Empty();
+	DriverDate.Empty();
+
+	for(uint32 i = 0;; ++i)
+	{
+		// Iterate all installed display adapters
+		FString Key = FString::Printf(TEXT("SYSTEM\\CurrentControlSet\\Control\\Class\\{4D36E968-E325-11CE-BFC1-08002BE10318}\\%04d"), i);
+
+		FString LocalDeviceDescription; // e.g. "NVIDIA GeForce GTX 680" or "AMD Radeon R9 200 / HD 7900 Series"
+		bool bDevice = FWindowsPlatformMisc::QueryRegKey(HKEY_LOCAL_MACHINE, *Key, TEXT("Device Description"), LocalDeviceDescription);	// AMD and NVIDIA
+
+		if(!bDevice)
+		{
+			break;
+		}
+
+		// AMD (not the Catalyst one) and NVIDIA
+		// e.g. "15.200.1062.1004"(AMD) "9.18.13.4788" (NVIDIA)
+		FString LocalInternalDriverVersion;
+		FWindowsPlatformMisc::QueryRegKey(HKEY_LOCAL_MACHINE, *Key, TEXT("DriverVersion"), LocalInternalDriverVersion);
+
+		// if there is no Catalyst_Version we use the internal version
+		FString LocalUserDriverVersion = LocalInternalDriverVersion;
+
+		// AMD only
+		// e.g. 15.7.1
+		FWindowsPlatformMisc::QueryRegKey(HKEY_LOCAL_MACHINE, *Key, TEXT("Catalyst_Version"), LocalUserDriverVersion);
+
+		// AMD and NVIDIA
+		// e.g. 3-13-2015
+		FString LocalDriverDate;
+		FWindowsPlatformMisc::QueryRegKey(HKEY_LOCAL_MACHINE, *Key, TEXT("DriverDate"), LocalDriverDate);
+
+		if(LocalDeviceDescription == DeviceDescription)
+		{
+			// found the one we are searching for (if there are multiple we only get the first one)
+			InternalDriverVersion = LocalInternalDriverVersion;
+			UserDriverVersion = LocalUserDriverVersion;
+			DriverDate = LocalDriverDate;
+			break;
+		}
+	}
+}
 #include "HideWindowsPlatformTypes.h"
 
 void FWindowsPlatformMisc::GetOSVersions( FString& out_OSVersionLabel, FString& out_OSSubVersionLabel )
diff --git a/Engine/Source/Runtime/Core/Public/GenericPlatform/GenericPlatformMemory.h b/Engine/Source/Runtime/Core/Public/GenericPlatform/GenericPlatformMemory.h
--- a/Engine/Source/Runtime/Core/Public/GenericPlatform/GenericPlatformMemory.h
+++ b/Engine/Source/Runtime/Core/Public/GenericPlatform/GenericPlatformMemory.h
@@ -129,6 +129,7 @@ struct CORE_API FGenericPlatformMemory
 		MCR_GPU, // memory directly a GPU (graphics card, etc)
 		MCR_GPUSystem, // system memory directly accessible by a GPU
 		MCR_TexturePool, // presized texture pools
+		MCR_GPUDefragPool, // presized pool of memory that can be defragmented.
 		MCR_MAX
 	};
 
diff --git a/Engine/Source/Runtime/Core/Public/GenericPlatform/GenericPlatformMisc.h b/Engine/Source/Runtime/Core/Public/GenericPlatform/GenericPlatformMisc.h
--- a/Engine/Source/Runtime/Core/Public/GenericPlatform/GenericPlatformMisc.h
+++ b/Engine/Source/Runtime/Core/Public/GenericPlatform/GenericPlatformMisc.h
@@ -319,6 +319,11 @@ struct CORE_API FGenericPlatformMisc
 	 */
 	static FString GetPrimaryGPUBrand();
 
+	// @param InternalDriverVersion e.g. "15.200.1062.1004"(AMD) "9.18.13.4788" (NVIDIA) for comparison
+	// @param UserDriverVersion e.g. "15.7.1"(Catalyst AMD) "9.18.13.4788" (NVIDIA) for user feedback
+	// @param DriverDate e.g. 3-13-2015
+	static void GetGPUDriverInfo(const FString DeviceDescription, FString& InternalDriverVersion, FString& UserDriverVersion, FString& DriverDate);
+
 	/**
 	 * Gets the OS Version and OS Subversion.
 	 */
diff --git a/Engine/Source/Runtime/Core/Public/Misc/Build.h b/Engine/Source/Runtime/Core/Public/Misc/Build.h
--- a/Engine/Source/Runtime/Core/Public/Misc/Build.h
+++ b/Engine/Source/Runtime/Core/Public/Misc/Build.h
@@ -247,3 +247,8 @@
 
 #define USE_CIRCULAR_DEPENDENCY_LOAD_DEFERRING 1
 #define USE_DEFERRED_DEPENDENCY_CHECK_VERIFICATION_TESTS (USE_CIRCULAR_DEPENDENCY_LOAD_DEFERRING && 0)
+
+// 0 (default), set this to 1 to get draw events with "TOGGLEDRAWEVENTS" "SHOWMATERIALDRAWEVENTS" and the "ProfileGPU" command working in test
+#define ALLOW_PROFILEGPU_IN_TEST 0
+// draw events with "TOGGLEDRAWEVENTS" "SHOWMATERIALDRAWEVENTS" (for ProfileGPU, Pix, Razor, RenderDoc, ...) and the "ProfileGPU" command are normally compiled out for TEST and SHIPPING
+#define WITH_PROFILEGPU (!(UE_BUILD_SHIPPING || UE_BUILD_TEST) || (UE_BUILD_TEST && ALLOW_PROFILEGPU_IN_TEST))
diff --git a/Engine/Source/Runtime/Core/Public/Misc/CallbackDevice.h b/Engine/Source/Runtime/Core/Public/Misc/CallbackDevice.h
--- a/Engine/Source/Runtime/Core/Public/Misc/CallbackDevice.h
+++ b/Engine/Source/Runtime/Core/Public/Misc/CallbackDevice.h
@@ -49,7 +49,7 @@ public:
 	DECLARE_MULTICAST_DELEGATE_OneParam(FOnActorLabelChanged, AActor*);
 
 	// delegate type for prompting the pak system to mount a new pak
-	DECLARE_DELEGATE_RetVal_TwoParams(bool, FOnMountPak, const FString&, uint32);
+	DECLARE_DELEGATE_RetVal_ThreeParams(bool, FOnMountPak, const FString&, uint32, IPlatformFile::FDirectoryVisitor*);
 
 	// delegate type for prompting the pak system to mount a new pak
 	DECLARE_DELEGATE_RetVal_OneParam(bool, FOnUnmountPak, const FString&);
diff --git a/Engine/Source/Runtime/Core/Public/Misc/ConfigCacheIni.h b/Engine/Source/Runtime/Core/Public/Misc/ConfigCacheIni.h
--- a/Engine/Source/Runtime/Core/Public/Misc/ConfigCacheIni.h
+++ b/Engine/Source/Runtime/Core/Public/Misc/ConfigCacheIni.h
@@ -8,6 +8,8 @@
 
 CORE_API DECLARE_LOG_CATEGORY_EXTERN(LogConfig, Warning, All);
 
+// Server builds should be tweakable even in Shipping
+#define ALLOW_INI_OVERRIDE_FROM_COMMANDLINE			(UE_SERVER || !(UE_BUILD_SHIPPING))
 
 typedef TMultiMap<FName,FString> FConfigSectionMap;
 
@@ -42,13 +44,13 @@ struct FIniFilename
 };
 
 
-#if !UE_BUILD_SHIPPING
+#if ALLOW_INI_OVERRIDE_FROM_COMMANDLINE
 // Options which stemmed from the commandline
 struct FConfigCommandlineOverride
 {
 	FString BaseFileName, Section, PropertyKey, PropertyValue;
 };
-#endif // !UE_BUILD_SHIPPING
+#endif // ALLOW_INI_OVERRIDE_FROM_COMMANDLINE
 
 // Possible entries in a config hierarchy
 enum class EConfigFileHierarchy : uint8
@@ -117,10 +119,10 @@ public:
 	/** Key to the cache to speed up ini parsing */
 	FString CacheKey;
 
-#if !UE_BUILD_SHIPPING
+#if ALLOW_INI_OVERRIDE_FROM_COMMANDLINE
 	/** The collection of overrides which stemmed from the commandline */
 	TArray<FConfigCommandlineOverride> CommandlineOptions;
-#endif // !UE_BUILD_SHIPPING
+#endif // ALLOW_INI_OVERRIDE_FROM_COMMANDLINE
 	
 	CORE_API FConfigFile();
 	FConfigFile( int32 ) {}	// @todo UE4 DLL: Workaround for instantiated TMap template during DLLExport (TMap::FindRef)
