Fix being able to trigger Rocket builds from main branch.

Optimization to avoid copying unused build products to the network interferes with determining whether an upstream node has completed or not. Check for the 'succeeded' flag, as well as checking for the existence of build products.

diff --git a/Engine/Source/Programs/AutomationTool/GUBP/GUBP.cs b/Engine/Source/Programs/AutomationTool/GUBP/GUBP.cs
--- a/Engine/Source/Programs/AutomationTool/GUBP/GUBP.cs
+++ b/Engine/Source/Programs/AutomationTool/GUBP/GUBP.cs
@@ -377,8 +377,8 @@ public partial class GUBP : BuildCommand
 	{
 		foreach(BuildNode NodeToDo in NodesToDo)
 		{
-            // Construct the full temp storage node info
-			NodeToDo.IsComplete = TempStorage.TempStorageExists(NodeToDo.Name, SharedStorageDir, bQuiet: true);
+            // Construct the full temp storage node info. Check both the block name and the succeeded record; we may have skipped copying the build products themselves because they weren't input dependencies.
+			NodeToDo.IsComplete = TempStorage.TempStorageExists(NodeToDo.Name, SharedStorageDir, bQuiet: true) || TempStorage.TempStorageExists(NodeToDo.Name + SucceededTempStorageSuffix, SharedStorageDir, bQuiet: true);
 
             LogVerbose("** {0}", NodeToDo.Name);
 			if (!NodeToDo.IsComplete)
