Tweaks for setting up new build process in main. Add aggregate nodes for common CIS targets ("AllEditors", "AllTools", "AllMonolithics"), and settings file for the branch which reference the correct agents.

diff --git a/Engine/Source/Programs/AutomationTool/GUBP/LegacyBranchSetup.cs b/Engine/Source/Programs/AutomationTool/GUBP/LegacyBranchSetup.cs
--- a/Engine/Source/Programs/AutomationTool/GUBP/LegacyBranchSetup.cs
+++ b/Engine/Source/Programs/AutomationTool/GUBP/LegacyBranchSetup.cs
@@ -1158,6 +1158,22 @@ partial class GUBP
                 BranchConfig.AddNode(new CleanSharedTempStorageNode(this, BranchConfig));
             }
         }
+
+		// Add an aggregate for all the editors in this branch
+		List<GUBPNode> AllEditorNodes = new List<GUBPNode>(BranchConfig.GUBPNodes.Values.Where(x => (x is RootEditorNode) || (x is EditorGameNode)));
+		BranchConfig.AddNode(new GenericAggregateNode("AllEditors", AllEditorNodes.Select(x => x.GetFullName())));
+
+		// Add an aggregate for all the tools in the branch
+		List<GUBPNode> AllToolNodes = new List<GUBPNode>(BranchConfig.GUBPNodes.Values.Where(x => (x is ToolsNode) || (x is SingleToolsNode) || (x is InternalToolsNode) || (x is SingleInternalToolsNode)));
+		BranchConfig.AddNode(new GenericAggregateNode("AllTools", AllToolNodes.Select(x => x.GetFullName())));
+
+		// Add an aggregate for all the tools in the branch
+		List<GUBPNode> AllMonolithicNodes = new List<GUBPNode>(BranchConfig.GUBPNodes.Values.Where(x => (x is GamePlatformMonolithicsNode)));
+		BranchConfig.AddNode(new GenericAggregateNode("AllMonolithics", AllMonolithicNodes.Select(x => x.GetFullName())));
+
+		// Add an aggregate for all the nodes in the branch
+		BranchConfig.AddNode(new GenericAggregateNode("AllNodes", BranchConfig.GUBPNodes.Keys));
+
 #if false
         // this doesn't work for lots of reasons...we can't figure out what the dependencies are until far later
         if (bPreflightBuild)
diff --git a/Engine/Source/Programs/AutomationTool/GUBP/LegacyNodes.cs b/Engine/Source/Programs/AutomationTool/GUBP/LegacyNodes.cs
--- a/Engine/Source/Programs/AutomationTool/GUBP/LegacyNodes.cs
+++ b/Engine/Source/Programs/AutomationTool/GUBP/LegacyNodes.cs
@@ -1500,6 +1500,26 @@ partial class GUBP
 		}
 	}
 
+	public class GenericAggregateNode : GUBPAggregateNode
+	{
+		string Name;
+
+		public GenericAggregateNode(string InName, IEnumerable<string> NodeNames)
+		{
+			Name = InName;
+
+			foreach(string NodeName in NodeNames)
+			{
+				AddDependency(NodeName);
+			}
+		}
+
+		public override string GetFullName()
+		{
+			return Name;
+		}
+	}
+
     public class WaitForUserInput : GUBPNode
     {
         public WaitForUserInput()
