Disable warnings for documentation build nodes; there are too many of them, and it's not run frequently enough to be actionable by recipients.

diff --git a/Engine/Source/Programs/AutomationTool/Scripts/BuildDocumentation.Automation.cs b/Engine/Source/Programs/AutomationTool/Scripts/BuildDocumentation.Automation.cs
--- a/Engine/Source/Programs/AutomationTool/Scripts/BuildDocumentation.Automation.cs
+++ b/Engine/Source/Programs/AutomationTool/Scripts/BuildDocumentation.Automation.cs
@@ -65,6 +65,11 @@ class DocumentationNode : GUBP.HostPlatformNode
 		CommandUtils.RunAndLog(CommandUtils.CmdEnv, ApiDocToolPath, ApiToolCommandLine, LogName);
 	}
 
+	public override bool NotifyOnWarnings()
+	{
+		return false;
+	}
+
 	public static void SubmitOutputs(string Description, params string[] FileSpecs)
 	{
 		if (CommandUtils.P4Enabled)
