Remove game/platform owners from notification emails for builds in dev branches, and switch target platforms to be compiled in serial for UE4 Main & Dev branches. We don't particularly need fast builds out of main any more, but we want to avoid saturating the build farm so that dev branches can build.

diff --git a/Engine/Source/Programs/AutomationTool/GUBP/LegacyBranchSetup.cs b/Engine/Source/Programs/AutomationTool/GUBP/LegacyBranchSetup.cs
--- a/Engine/Source/Programs/AutomationTool/GUBP/LegacyBranchSetup.cs
+++ b/Engine/Source/Programs/AutomationTool/GUBP/LegacyBranchSetup.cs
@@ -188,6 +188,10 @@ partial class GUBP
 			/// </summary>
 			public bool bBuildEngineLocalization = false;
 			public string EngineLocalizationBranchSuffix = "";
+			/// <summary>
+			/// Whether to build all target platforms in parallel, or serialize them to avoid saturating the farm.
+			/// </summary>
+			public bool bTargetPlatformsInParallel = true;
         }
         public virtual void ModifyOptions(GUBP bp, ref BranchOptions Options, string Branch)
         {
diff --git a/Engine/Source/Programs/AutomationTool/GUBP/LegacyNodes.cs b/Engine/Source/Programs/AutomationTool/GUBP/LegacyNodes.cs
--- a/Engine/Source/Programs/AutomationTool/GUBP/LegacyNodes.cs
+++ b/Engine/Source/Programs/AutomationTool/GUBP/LegacyNodes.cs
@@ -1096,9 +1096,11 @@ partial class GUBP
 				AddDependency(ToolsNode.StaticGetFullName(InHostPlatform));
 			}
 
+			bool bBehindTrigger = false;
 			if(IsSample(BranchConfig, InGameProj))
 			{
 				AddPseudodependency(WaitToPackageSamplesNode.StaticGetFullName());
+				bBehindTrigger = true;
 			}
 
             if (InGameProj.GameName != BranchConfig.Branch.BaseEngineProject.GameName && GameProj.Properties.Targets.ContainsKey(TargetRules.TargetType.Editor))
@@ -1122,8 +1124,13 @@ partial class GUBP
             if (InGameProj.Options(InHostPlatform).bTestWithShared)  /// compiling templates is only for testing purposes, and we will group them to avoid saturating the farm
             {
                 AddPseudodependency(WaitForTestShared.StaticGetFullName());
+				bBehindTrigger = true;
                 //AgentSharingGroup = "TemplateMonolithics" + StaticGetHostPlatformSuffix(InHostPlatform);
             }
+			if(!InBranchConfig.BranchOptions.bTargetPlatformsInParallel && !bBehindTrigger)
+			{
+				AgentSharingGroup = "TargetPlatforms" + StaticGetHostPlatformSuffix(InHostPlatform);
+			}
         }
 
 		public static bool IsSample(GUBPBranchConfig BranchConfig, BranchInfo.BranchUProject GameProj)
