Disable warnings for build nodes that we don't expect to fix soon.

diff --git a/Engine/Source/Programs/AutomationTool/GUBP/LegacyNodes.cs b/Engine/Source/Programs/AutomationTool/GUBP/LegacyNodes.cs
--- a/Engine/Source/Programs/AutomationTool/GUBP/LegacyNodes.cs
+++ b/Engine/Source/Programs/AutomationTool/GUBP/LegacyNodes.cs
@@ -2703,6 +2703,11 @@ partial class GUBP
             AddDependency(ToolsNode.StaticGetFullName(HostPlatform));
 		}
 
+		public override bool NotifyOnWarnings()
+		{
+			return false;
+		}
+
 		public override void DoBuild(GUBP bp)
 		{
 			var UEProjectDirectory = GetUEProjectDirectory();
