Force metadata for trigger nodes to be copied to shared storage. Even though nothing depends on their outputs directly, we need a record that they've run for downstream nodes. Fixes UEB-604.

diff --git a/Engine/Source/Programs/AutomationTool/GUBP/GUBP.cs b/Engine/Source/Programs/AutomationTool/GUBP/GUBP.cs
--- a/Engine/Source/Programs/AutomationTool/GUBP/GUBP.cs
+++ b/Engine/Source/Programs/AutomationTool/GUBP/GUBP.cs
@@ -195,6 +195,10 @@ public partial class GUBP : BuildCommand
 				}
 			}
 		}
+		foreach(TriggerNode TriggerNode in NodesToDo.OfType<TriggerNode>())
+		{
+			TriggerNode.CopyToSharedStorage = true;
+		}
 
 		// From each of those nodes, remove any order dependencies that we're not going to run. There are no circumstances in which they should be 
 		// considered dependencies from this point on.
