Don't output a warning message when dumping Mac crash logs - if it caused a problem, there will already be an error propagated. Ignore any crashes in snmpd, which seems to fail whenever there's heavy load on a machine (eg. during a cook), and fix postp parsing of standard UAT/UBT warnings and errors.

diff --git a/Engine/Build/BuildFarm/PostpExtensions.pl b/Engine/Build/BuildFarm/PostpExtensions.pl
--- a/Engine/Build/BuildFarm/PostpExtensions.pl
+++ b/Engine/Build/BuildFarm/PostpExtensions.pl
@@ -104,6 +104,21 @@ unshift @::gMatchers, (
 		id =>               "automationException",
 		pattern =>          q{AutomationTool\\.AutomationException: },
 		action =>           q{incValue("errors"); diagnostic("Exception", "error", 0, forwardWhile("^  at "));}
+	},
+	{
+		id =>				"ubtFatal",
+		pattern =>			q{^FATAL:},
+		action =>			q{incValue("errors"); diagnostic("AutomationTool", "error")}
+	},
+	{
+		id =>				"ubtError",
+		pattern =>			q{^ERROR:},
+		action =>			q{incValue("errors"); diagnostic("AutomationTool", "error")}
+	},
+	{
+		id =>				"ubtWarning",
+		pattern =>			q{^WARNING:},
+		action =>			q{incValue("warnings"); diagnostic("AutomationTool", "warning")}
 	}
 );
 
diff --git a/Engine/Source/Programs/AutomationTool/AutomationUtils/CommandletUtils.cs b/Engine/Source/Programs/AutomationTool/AutomationUtils/CommandletUtils.cs
--- a/Engine/Source/Programs/AutomationTool/AutomationUtils/CommandletUtils.cs
+++ b/Engine/Source/Programs/AutomationTool/AutomationUtils/CommandletUtils.cs
@@ -301,18 +301,22 @@ namespace AutomationTool
 				// Dump them all to the log
 				foreach(FileInfo CrashFileInfo in CrashFileInfos)
 				{
-					CommandUtils.LogWarning("Found crash log - {0}:", CrashFileInfo.FullName);
-					try
+					// snmpd seems to often crash (suspect due to it being starved of CPU cycles during cooks)
+					if(!CrashFileInfo.Name.StartsWith("snmpd_"))
 					{
-						string[] Lines = File.ReadAllLines(CrashFileInfo.FullName);
-						foreach(string Line in Lines)
+						CommandUtils.Log("Found crash log - {0}", CrashFileInfo.FullName);
+						try
 						{
-							CommandUtils.Log("CrashDump: {0}", Line);
+							string[] Lines = File.ReadAllLines(CrashFileInfo.FullName);
+							foreach(string Line in Lines)
+							{
+								CommandUtils.Log("Crash: {0}", Line);
+							}
+						}
+						catch(Exception Ex)
+						{
+							CommandUtils.LogWarning("Failed to read file ({0})", Ex.Message);
 						}
-					}
-					catch(Exception Ex)
-					{
-						CommandUtils.LogWarning("Failed to read file ({0})", Ex.Message);
 					}
 				}
 			}
