Merging //UE4/Release-4.11 to //UE4/Main (Up to CL#2897738)

==========================
MAJOR FEATURES + CHANGES
==========================

Change 2875445 on 2016/02/22 by [REDACTED_PERSON]

	Added UE4.natvis to Visual Studio Projects

Change 2875456 on 2016/02/22 by [REDACTED_PERSON]

	Fix custom stencil shaders on Xbox One

Change 2875524 on 2016/02/22 by [REDACTED_PERSON]

	More log info when saving shader temp files fails. Increased the number of attemps when moving a file fails.

Change 2875698 on 2016/02/22 by [REDACTED_PERSON]

	UE4.11 - Add new bool for RHIs (unused currently)

Change 2875897 on 2016/02/22 by [REDACTED_PERSON]

Change 2875922 on 2016/02/22 by [REDACTED_PERSON]

	Fixed BP-implemented EQS generators crashing when trying to add generated value of wrong type #UE4

Change 2875960 on 2016/02/22 by [REDACTED_PERSON]

	Added a way to disable right click emulation on Mac and used it in TabNavigator to fix issues with its widgets not reacting to clicking

Change 2875984 on 2016/02/22 by [REDACTED_PERSON]

	Split output struct pins will no longer give a warning about override pins being removed.

Change 2876169 on 2016/02/22 by [REDACTED_PERSON]

	Changes to support building UHT plugins with the binary release.

	* Add source code and target files for UHT to binary distribution
	* Fix UBT deleting build products if we're only compiling a single module.
	* Fix UBT exception setting up compile environment when a module doesn't have any source files set to build.
	* Include DLL import libraries for UHT in the build
	* Add support for compiling UHT modules in BuildPluginCommand. Stages an empty host project to allow UHT to load any enabled plugins.

Change 2876219 on 2016/02/22 by [REDACTED_PERSON]

	UE4.11 - Integration from 2874609
	PC: Update D3D12 RHI
	- Implement _RenderThread versions of Create, Lock and Unlock Index/Vertex Buffer. Only synchronize threads on Readback
	- Limit GPU starvation on CPU bound scenarios by flushing work when the GPU is idle
	- Change texture streaming system to correctly account for placed textures. Also fix texture sizes so they accurately represent the real size of the allocation the GPU.
	- Disable API shader blobs
	- Add the ability to easily change allocation stategy for a given pool, also add a simple linear allocator and a 'Multi-Buddy Allocator' for efficiency in different scenarios
	- Pool Fences to prevent creation and destruction every frame when using Async Compute
	- Implement _RenderThread versions of CreateShaderResourceView and CreateUnorderedAccessView

Change 2876232 on 2016/02/22 by [REDACTED_PERSON]

	UE4.11 - Integration from 2876173
	PC: Update D3D12 RHI

	- Fix ResizeBuffers() failing due to dangling references to the backbuffer if deferred deletion is used.

	- Reorder when pending FRHIResources are deleted. This still needs to flush all pending deletes and ignore the deferred deletion queue otherwise some items may still be left in the engine's queue.

	- Fix UT build error due to missing FPlatformMisc::GetGPUDriverInfo()

Change 2876366 on 2016/02/22 by [REDACTED_PERSON]

	Adding Skeletal Meshes for Import Test Case

Change 2876401 on 2016/02/22 by [REDACTED_PERSON]

	fix for WindowsClient build from UFE and Project Launcher

Change 2876456 on 2016/02/22 by [REDACTED_PERSON]

	Use a more hierarchical directory structure for packaged builds, rather than just dumping everything in the root. Now defaults to <Share>\\PackagedBuilds\\<Branch>\\<CL>\\<ProjectName>_<Platform>_<Configuration>.

Change 2876507 on 2016/02/22 by [REDACTED_PERSON]

	use HOME (osx) and USERPROFILE (windows) on appropriate target platform

Change 2876537 on 2016/02/22 by [REDACTED_PERSON]

	Removed dubious fix for an old bug, no longer needed but I havn't figured out what has changed. This fixes a crash on Replace References, but does not reintroduce UE-9497

Change 2876545 on 2016/02/22 by [REDACTED_PERSON]

	SteamVR camera late-update fix

Change 2876825 on 2016/02/22 by [REDACTED_PERSON]

	Unfortunate edge case in lifetime of UEdGraph's schema, schema is assigned after construction so its modification is in the undo buffer, and we clear it after undoing.

Change 2876878 on 2016/02/22 by [REDACTED_PERSON]

	PSVR HMD Server support

diff --git a/Engine/Build/InstalledEngineFilters.ini b/Engine/Build/InstalledEngineFilters.ini
--- a/Engine/Build/InstalledEngineFilters.ini
+++ b/Engine/Build/InstalledEngineFilters.ini
@@ -20,10 +20,8 @@
 ; Config files
 /Engine/Config/...
 -/Engine/Config/....vdf
-/Engine/Programs/UnrealBuildTool/BuildConfiguration.xml
-/Engine/Programs/UnrealFrontend/Config/DefaultEngine.ini
-/Engine/Programs/UnrealLightmass/Config/DefaultEngine.ini
-/Engine/Programs/CrashReportClient/Config/...
+/Engine/Programs/...
+-/Engine/Programs/UnrealGameSync/...
 
 ; Plugins
 /Engine/Plugins/....uplugin
@@ -44,6 +42,7 @@
 /Engine/Source/Editor/...
 /Engine/Source/ThirdParty/.../*.Build.cs
 /Engine/Source/ThirdParty/Licenses/...
+/Engine/Source/Programs/UnrealHeaderTool/...
 -/Engine/Source/Runtime/SQLiteSupport/...
 
 ; Shaders
diff --git a/Engine/Build/TVOS/XcodeSupportFiles/prepackage.sh b/Engine/Build/TVOS/XcodeSupportFiles/prepackage.sh
new file mode 100755
--- /dev/null
+++ b/Engine/Build/TVOS/XcodeSupportFiles/prepackage.sh
@@ -0,0 +1,68 @@
+#!/bin/sh
+
+echo "Checking project"
+if [ project.pbxproj.datecheck -nt ../UE4_FromPC.xcodeproj/project.pbxproj ]
+then
+  echo "Copying project..."
+  mkdir -p ../UE4_FromPC.xcodeproj > /dev/null
+  cp project.pbxproj.datecheck ../UE4_FromPC.xcodeproj/project.pbxproj
+else
+  echo "Project file is up to date."
+fi
+
+echo "Checking for inline breakpoint strategy"
+lldbinitfile=~/.lldbinit
+if grep -q "strategy" $lldbinitfile
+then
+  echo "Strategy found... inline breakpoints already enabled"
+else
+  echo "settings set target.inline-breakpoint-strategy always" >> ~/.lldbinit
+  echo "Strategy not found -> enabling all inline breakpoints"
+fi
+
+if [ -e ../../../Binaries/TVOS/$1-$2-$3$4.app.dSYM.zip.datecheck ]
+then
+  pushd ../../../Binaries/TVOS > /dev/null
+  if [ $1-$2-$3$4.app.dSYM.zip.datecheck -nt $1-$2-$3.app.dSYM.zip ]
+  then
+    echo "Unzipping new .dSYM..."
+    mv $1-$2-$3$4.app.dSYM.zip.datecheck $1-$2-$3$4.app.dSYM.zip
+    unzip -o $1-$2-$3$4.app.dSYM.zip
+  else
+    echo ".dSYM is up to date."
+  fi
+  popd > /dev/null
+else
+  echo $1"-"$2"-"$3""$4".app.dSYM not found..."
+fi
+
+if [ -e ../../../Binaries/TVOS/$1-$2-$3$4.app.dSYM ]
+then
+  pushd ../../../Binaries/TVOS > /dev/null
+  cp -R $1-$2-$3$4.app.dSYM Payload/$1.app.dSYM
+  popd > /dev/null
+fi
+
+if [ -e ../../../Binaries/TVOS/$1$4.app.dSYM ]
+then
+  pushd ../../../Binaries/TVOS > /dev/null
+  cp -R $1$4.app.dSYM Payload/$1$4.app.dSYM
+  popd > /dev/null
+fi
+
+if [ -e $1-Info.plist ]
+then
+  if [ -d ../../../Intermediate/TVOS ]
+  then
+    cp -R $1-Info.plist ../../../Intermediate/TVOS/$1-Info.plist
+  else
+	if [ "$1" == "UE4Game" ]
+	then
+      mkdir -p ../../../../Engine/Intermediate/TVOS/
+      cp -R $1-Info.plist ../../../../Engine/Intermediate/TVOS/$1-Info.plist
+	else 
+      mkdir -p ../../../Intermediate/TVOS/
+      cp -R $1-Info.plist ../../../Intermediate/TVOS/$1-Info.plist
+	fi
+  fi
+fi
\ No newline at end of file
diff --git a/Engine/Config/BaseDeviceProfiles.ini b/Engine/Config/BaseDeviceProfiles.ini
--- a/Engine/Config/BaseDeviceProfiles.ini
+++ b/Engine/Config/BaseDeviceProfiles.ini
@@ -40,6 +40,7 @@
 +DeviceProfileNameAndTypes=Android_Adreno2xx,Android
 +DeviceProfileNameAndTypes=Android_Adreno330,Android
 +DeviceProfileNameAndTypes=Android_Adreno330_Ver53,Android
++DeviceProfileNameAndTypes=Android_Adreno4xx_Nexus6,Android
 +DeviceProfileNameAndTypes=Android_Adreno4xx,Android
 +DeviceProfileNameAndTypes=Android_PowerVR54x,Android
 +DeviceProfileNameAndTypes=Android_TegraK1_GL4,Android
@@ -231,6 +232,7 @@ MatchProfile=(Profile="Android_PowerVR54x",Match=((SourceType=SRC_GpuFamily,Comp
 +MatchProfile=(Profile="Android_Adreno320",Match=((SourceType=SRC_GpuFamily,CompareType=CMP_Equal,MatchString="Adreno (TM) 320")))
 +MatchProfile=(Profile="Android_Adreno330_Ver53",Match=((SourceType=SRC_GpuFamily,CompareType=CMP_Equal,MatchString="Adreno (TM) 330"),(SourceType=SRC_GlVersion,CompareType=CMP_Regex,MatchString="ES 3\\.0 V@([0-9]+)"),(SourceType=SRC_PreviousRegexMatch,CompareType=CMP_GreaterEqual,MatchString="53")))
 +MatchProfile=(Profile="Android_Adreno330",Match=((SourceType=SRC_GpuFamily,CompareType=CMP_Equal,MatchString="Adreno (TM) 330")))
++MatchProfile=(Profile="Android_Adreno4xx_Nexus6",Match=((SourceType=SRC_DeviceModel,CompareType=CMP_Equal,MatchString="Nexus 6")))
 +MatchProfile=(Profile="Android_Adreno4xx",Match=((SourceType=SRC_GpuFamily,CompareType=CMP_Regex,MatchString="Adreno \\(TM\\) 4[0-9][0-9]")))
 +MatchProfile=(Profile="Android_Adreno5xx",Match=((SourceType=SRC_GpuFamily,CompareType=CMP_Regex,MatchString="Adreno \\(TM\\) 5[0-9][0-9]")))
 +MatchProfile=(Profile="Android_TegraK1_GL4",Match=((SourceType=SRC_GpuFamily,CompareType=CMP_Equal,MatchString="NVIDIA Tegra"),(SourceType=SRC_GlVersion,CompareType=CMP_Regex,MatchString="^4\\.")))
@@ -316,6 +318,11 @@ BaseProfileName=Android_High
 DeviceType=Android
 BaseProfileName=Android_High
 
+[Android_Adreno4xx_Nexus6 DeviceProfile]
+DeviceType=Android
+BaseProfileName=Android_Adreno4xx
++CVars=android.DisableTextureFormatBGRA8888=1
+
 [Android_PowerVR54x DeviceProfile]
 DeviceType=Android
 BaseProfileName=Android_Mid
diff --git a/Engine/Plugins/2D/Paper2D/Source/Paper2DEditor/Private/ExtractSprites/PaperExtractSpritesSettings.h b/Engine/Plugins/2D/Paper2D/Source/Paper2DEditor/Private/ExtractSprites/PaperExtractSpritesSettings.h
--- a/Engine/Plugins/2D/Paper2D/Source/Paper2DEditor/Private/ExtractSprites/PaperExtractSpritesSettings.h
+++ b/Engine/Plugins/2D/Paper2D/Source/Paper2DEditor/Private/ExtractSprites/PaperExtractSpritesSettings.h
@@ -57,6 +57,7 @@ class UPaperExtractSpriteGridSettings : public UObject
 {
 	GENERATED_BODY()
 
+public:
 	//////////////////////////////////////////////////////////////////////////
 	// Grid mode
 
diff --git a/Engine/Plugins/Experimental/StereoPanorama/Source/StereoPanorama/Private/SceneCapturer.cpp b/Engine/Plugins/Experimental/StereoPanorama/Source/StereoPanorama/Private/SceneCapturer.cpp
--- a/Engine/Plugins/Experimental/StereoPanorama/Source/StereoPanorama/Private/SceneCapturer.cpp
+++ b/Engine/Plugins/Experimental/StereoPanorama/Source/StereoPanorama/Private/SceneCapturer.cpp
@@ -62,7 +62,8 @@ void USceneCapturer::InitCaptureComponent( USceneCaptureComponent2D* CaptureComp
     CaptureComponent->bCaptureEveryFrame = false;
     CaptureComponent->CaptureSource = ESceneCaptureSource::SCS_FinalColorLDR;
 
-    CaptureComponent->TextureTarget = NewObject<UTextureRenderTarget2D>(this, TEXT("SceneCaptureTextureTarget"));
+	const FName TargetName = MakeUniqueObjectName(this, UTextureRenderTarget2D::StaticClass(), TEXT("SceneCaptureTextureTarget"));
+    CaptureComponent->TextureTarget = NewObject<UTextureRenderTarget2D>(this, TargetName);
     //TODO: ikrimae: Not sure why the render target needs to be float to avoid banding. Seems like captures to this RT and then applies PP
     //               on top of it which causes degredation.
     CaptureComponent->TextureTarget->InitCustomFormat(CaptureWidth, CaptureHeight, PF_A16B16G16R16, false);
@@ -718,7 +719,7 @@ void USceneCapturer::CaptureComponent( int32 CurrentHorizontalStep, int32 Curren
     }
 }
 
-//TODO: ikrimae: This is a cluster fuck. Come back and actually work out the timings. Trickery b/c SceneCaptureCubes Tick at the end of the frame so we're effectively queuing up the next
+//TODO: ikrimae: Come back and actually work out the timings. Trickery b/c SceneCaptureCubes Tick at the end of the frame so we're effectively queuing up the next
 //               step (pause, unpause, setposition) for the next frame. FlushRenderingCommands() added haphazardly to test but didn't want to remove them so close to delivery. 
 //               Think through when we actually need to flush and document.
 void USceneCapturer::Tick( float DeltaTime )
diff --git a/Engine/Plugins/Media/AndroidMedia/Source/AndroidMedia/Private/AndroidMediaPlayer.cpp b/Engine/Plugins/Media/AndroidMedia/Source/AndroidMedia/Private/AndroidMediaPlayer.cpp
--- a/Engine/Plugins/Media/AndroidMedia/Source/AndroidMedia/Private/AndroidMediaPlayer.cpp
+++ b/Engine/Plugins/Media/AndroidMedia/Source/AndroidMedia/Private/AndroidMediaPlayer.cpp
@@ -7,12 +7,14 @@
 #include "AndroidFile.h"
 #include "Paths.h"
 #include "RenderingThread.h"
-
+#include "RenderResource.h"
 
 DEFINE_LOG_CATEGORY_STATIC(LogAndroidMediaPlayer, Log, All);
 
 #define LOCTEXT_NAMESPACE "FAndroidMediaModule"
 
+// Enables directly copying into bound textures instead of using slower glReadPixels
+#define ENABLE_BINDTEXTURE	1
 
 class FAndroidMediaPlayer::MediaTrack
 	: public IMediaStream
@@ -183,15 +185,42 @@ public:
 		return *this;
 	}
 
+#if ENABLE_BINDTEXTURE
+	virtual void BindTexture(FTextureRHIParamRef BoundResource) override
+	{
+		BoundTextures.Add(BoundResource);
+	}
+
+	virtual void UnbindTexture(FTextureRHIParamRef BoundResource) override
+	{
+		BoundTextures.Remove(BoundResource);
+	}
+
+	virtual void Tick(float DeltaTime)
+	{
+		if (MediaPlayer.MediaState != EMediaState::Error)
+		{
+			int32 CurrentFramePosition
+				= MediaPlayer.JavaMediaPlayer->GetCurrentPosition();
+			if (LastFramePosition != CurrentFramePosition)
+			{
+				if (UpdateBoundTextures())
+				{
+					LastFramePosition = CurrentFramePosition;
+				}
+			}
+		}
+	}
+
+#else
+
 #if WITH_ENGINE
 	virtual void BindTexture(class FRHITexture* Texture) override
 	{
-		// @todo android: cbabcock: implement texture binding
 	}
 
 	virtual void UnbindTexture(class FRHITexture* Texture) override
 	{
-		// @todo android: cbabcock: implement texture binding
 	}
 #endif
 
@@ -218,9 +247,37 @@ public:
 		}
 	}
 
+#endif  // ENABLE_BINDTEXTURE
+
+private:
+
+#if ENABLE_BINDTEXTURE
+	bool UpdateBoundTextures()
+	{
+		bool FrameAvailable = false;
+		if (IsInRenderingThread())
+		{
+			for (auto Iter = BoundTextures.CreateIterator(); Iter; ++Iter)
+			{
+				FTextureRHIParamRef BoundTextureTarget = *Iter;
+				int32 DestTexture = *reinterpret_cast<int32*>(BoundTextureTarget->GetNativeResource());
+				if (MediaPlayer.JavaMediaPlayer->GetVideoLastFrame(DestTexture))
+				{
+					FrameAvailable = true;
+				}
+			}
+		}
+
+		return FrameAvailable;
+	}
+#endif  // ENABLE_BINDTEXTURE
+
 private:
 
 	int32 LastFramePosition;
+#if ENABLE_BINDTEXTURE
+	TSet<FTextureRHIParamRef> BoundTextures;
+#endif  // ENABLE_BINDTEXTURE
 };
 
 
@@ -281,7 +338,11 @@ FAndroidMediaPlayer::FAndroidMediaPlayer()
 	: JavaMediaPlayer(nullptr)
 	, MediaState(EMediaState::Error)
 {
+#if ENABLE_BINDTEXTURE
+	JavaMediaPlayer = MakeShareable(new FJavaAndroidMediaPlayer(false));
+#else
 	JavaMediaPlayer = MakeShareable(new FJavaAndroidMediaPlayer());
+#endif
 	if (JavaMediaPlayer.IsValid())
 	{
 		MediaState = EMediaState::Idle;
@@ -486,6 +547,21 @@ bool FAndroidMediaPlayer::Open(const FString& Url)
 		FString MoviePath = Url;
 		FPaths::NormalizeFilename(MoviePath);
 
+		// Deal with hardcoded path from editor
+		if (MoviePath.Contains(":") || MoviePath.StartsWith("/"))
+		{
+			int32 Index = MoviePath.Find(TEXT("/Content/Movies"));
+			if (Index > 0)
+			{
+				Index--;
+				while (Index > 0 && MoviePath[Index] != '/')
+				{
+					Index--;
+				}
+				MoviePath = "../../.." + MoviePath.Mid(Index);
+			}
+		}
+
 		// Don't bother trying to play it if we can't find it.
 		if (!IAndroidPlatformFile::GetPlatformPhysical().FileExists(*MoviePath))
 		{
@@ -590,14 +666,24 @@ bool FAndroidMediaPlayer::SetRate(float Rate)
 	switch (MediaState)
 	{
 	case EMediaState::Prepared:
-		if (1.0f == Rate)
+		if (Rate > 0.0f)
 		{
 			JavaMediaPlayer->Start();
 			MediaState = EMediaState::Started;
 			MediaEvent.Broadcast(EMediaEvent::PlaybackResumed);
 			return true;
 		}
-		else if (0.0f == Rate)
+		else
+		{
+			JavaMediaPlayer->Pause();
+			MediaState = EMediaState::Paused;
+			MediaEvent.Broadcast(EMediaEvent::PlaybackSuspended);
+			return true;
+		}
+		break;
+
+	case EMediaState::Started:
+		if (FMath::IsNearlyZero(Rate))
 		{
 			JavaMediaPlayer->Pause();
 			MediaState = EMediaState::Paused;
@@ -607,8 +693,17 @@ bool FAndroidMediaPlayer::SetRate(float Rate)
 		break;
 
 	case EMediaState::Paused:
+		if (Rate > 0.0f)
+		{
+			JavaMediaPlayer->Start();
+			MediaState = EMediaState::Started;
+			MediaEvent.Broadcast(EMediaEvent::PlaybackResumed);
+			return true;
+		}
+		break;
+
 	case EMediaState::PlaybackCompleted:
-		if (1.0f == Rate)
+		if (Rate > 0.0f)
 		{
 			JavaMediaPlayer->Start();
 			MediaState = EMediaState::Started;
diff --git a/Engine/Plugins/Messaging/UdpMessaging/Source/UdpMessaging/Private/Transport/UdpMessageTransport.cpp b/Engine/Plugins/Messaging/UdpMessaging/Source/UdpMessaging/Private/Transport/UdpMessageTransport.cpp
--- a/Engine/Plugins/Messaging/UdpMessaging/Source/UdpMessaging/Private/Transport/UdpMessageTransport.cpp
+++ b/Engine/Plugins/Messaging/UdpMessaging/Source/UdpMessaging/Private/Transport/UdpMessageTransport.cpp
@@ -95,12 +95,14 @@ bool FUdpMessageTransport::StartTransport()
 	{
 		MulticastReceiver = new FUdpSocketReceiver(MulticastSocket, ThreadWaitTime, TEXT("UdpMessageMulticastReceiver"));
 		MulticastReceiver->OnDataReceived().BindRaw(this, &FUdpMessageTransport::HandleSocketDataReceived);
+		MulticastReceiver->Start();
 	}
 
 #if PLATFORM_DESKTOP
 	UnicastReceiver = new FUdpSocketReceiver(UnicastSocket, ThreadWaitTime, TEXT("UdpMessageUnicastReceiver"));
 	{
 		UnicastReceiver->OnDataReceived().BindRaw(this, &FUdpMessageTransport::HandleSocketDataReceived);
+		UnicastReceiver->Start();
 	}
 #endif
 
diff --git a/Engine/Plugins/Runtime/AppleMoviePlayer/Source/AppleMoviePlayer/Private/AppleMovieStreamer.cpp b/Engine/Plugins/Runtime/AppleMoviePlayer/Source/AppleMoviePlayer/Private/AppleMovieStreamer.cpp
--- a/Engine/Plugins/Runtime/AppleMoviePlayer/Source/AppleMoviePlayer/Private/AppleMovieStreamer.cpp
+++ b/Engine/Plugins/Runtime/AppleMoviePlayer/Source/AppleMoviePlayer/Private/AppleMovieStreamer.cpp
@@ -435,30 +435,30 @@ bool FAVPlayerMovieStreamer::CheckForNextFrameAndCopy()
 			TempData.AddZeroed(SrcWidth * SrcHeight * 4);
 			TextureData->SetRawData(SrcWidth, SrcHeight, SrcWidth * 4, TempData);
 			check( TextureData->GetRawBytesPtr() != NULL );
+        }
 
-			// Now that we have video information, check on texture allocation. If we don't have a texture yet, create one.
-            if(!Texture.IsValid() || (Texture->GetWidth() != SrcWidth || Texture->GetHeight() != SrcHeight))
-            {
-				MovieViewport->SetTexture(NULL);
+        // Now that we have video information, check on texture allocation. If we don't have a texture yet, create one.
+        if(!Texture.IsValid() || (Texture->GetWidth() != SrcWidth || Texture->GetHeight() != SrcHeight))
+        {
+            MovieViewport->SetTexture(NULL);
 
-				// Release any resources associated with the previous texture
-                if (Texture.IsValid())
-                {
-                    // *** Aren't we already in the rendering thread? Why can't we just release the resource right here, right now?
-                    TexturesPendingDeletion.Add(Texture);
-                    BeginReleaseResource(Texture.Get());
-                    Texture.Reset();
-                }
+            // Release any resources associated with the previous texture
+            if (Texture.IsValid())
+            {
+                // *** Aren't we already in the rendering thread? Why can't we just release the resource right here, right now?
+                TexturesPendingDeletion.Add(Texture);
+                BeginReleaseResource(Texture.Get());
+                Texture.Reset();
+            }
 
-                // Create and initialize a new texture
-                Texture = MakeShareable(new FSlateTexture2DRHIRef(SrcWidth, SrcHeight, PF_B8G8R8A8, nullptr, TexCreate_Dynamic | TexCreate_NoTiling, true));
-                Texture->InitResource();
+            // Create and initialize a new texture
+            Texture = MakeShareable(new FSlateTexture2DRHIRef(SrcWidth, SrcHeight, PF_B8G8R8A8, nullptr, TexCreate_Dynamic | TexCreate_NoTiling, true));
+            Texture->InitResource();
 
-                // Make sure the texture is at least updated once.
-                Texture->UpdateRHI();
-                MovieViewport->SetTexture(Texture);
-			}
-		}
+            // Make sure the texture is at least updated once.
+            Texture->UpdateRHI();
+            MovieViewport->SetTexture(Texture);
+        }
 
         check( TextureData->GetBytesPerPixel() > 0 );
         
diff --git a/Engine/Plugins/Runtime/GearVR/Source/GearVR/Private/GearVR.cpp b/Engine/Plugins/Runtime/GearVR/Source/GearVR/Private/GearVR.cpp
--- a/Engine/Plugins/Runtime/GearVR/Source/GearVR/Private/GearVR.cpp
+++ b/Engine/Plugins/Runtime/GearVR/Source/GearVR/Private/GearVR.cpp
@@ -12,25 +12,13 @@
 
 #include <android_native_app_glue.h>
 
-#if PLATFORM_ANDROID
-// call out to JNI to see if the application was packaged for GearVR
-bool AndroidThunkCpp_IsGearVRApplication()
-{
-	bool bIsGearVRApplication = false;
-	if (JNIEnv* Env = FAndroidApplication::GetJavaEnv())
-	{
-		static jmethodID Method = FJavaWrapper::FindMethod(Env, FJavaWrapper::GameActivityClassID, "AndroidThunkJava_IsGearVRApplication", "()Z", false);
-		bIsGearVRApplication = FJavaWrapper::CallBooleanMethod(Env, FJavaWrapper::GameActivityThis, Method);
-	}
-	return bIsGearVRApplication;
-}
-#endif
-
 //---------------------------------------------------
 // GearVR Plugin Implementation
 //---------------------------------------------------
 
 #if GEARVR_SUPPORTED_PLATFORMS
+extern bool AndroidThunkCpp_IsGearVRApplication();
+
 static TAutoConsoleVariable<int32> CVarGearVREnableMSAA(TEXT("gearvr.EnableMSAA"), 1, TEXT("Enable MSAA when rendering on GearVR"));
 
 static TAutoConsoleVariable<int32> CVarGearVREnableQueueAhead(TEXT("gearvr.EnableQueueAhead"), 1, TEXT("Enable full-frame queue ahead for rendering on GearVR"));
@@ -257,8 +245,8 @@ bool FGearVR::GetHMDMonitorInfo(MonitorInfo& MonitorDesc)
 
 bool FGearVR::IsHMDConnected()
 {
-	// consider HMD connected all the time. 
-	return true;
+	// consider HMD connected all the time if GearVR enabled
+	return AndroidThunkCpp_IsGearVRApplication();
 }
 
 bool FGearVR::IsInLowPersistenceMode() const
@@ -1109,7 +1097,11 @@ void FGearVR::DrawDebug(UCanvas* Canvas)
 		{
 			DrawDebugTrackingCameraFrustum(GWorld, Canvas->SceneView->ViewRotation, Canvas->SceneView->ViewLocation);
 		}
-		DrawSeaOfCubes(GWorld, Canvas->SceneView->ViewLocation);
+
+		if(Canvas && Canvas->SceneView)
+		{
+			DrawSeaOfCubes(GWorld, Canvas->SceneView->ViewLocation);
+		}
 	}
 
 #endif // #if !UE_BUILD_SHIPPING
diff --git a/Engine/Plugins/Runtime/Steam/SteamVR/Source/SteamVR/Private/SteamVRHMD.cpp b/Engine/Plugins/Runtime/Steam/SteamVR/Source/SteamVR/Private/SteamVRHMD.cpp
--- a/Engine/Plugins/Runtime/Steam/SteamVR/Source/SteamVR/Private/SteamVRHMD.cpp
+++ b/Engine/Plugins/Runtime/Steam/SteamVR/Source/SteamVR/Private/SteamVRHMD.cpp
@@ -717,9 +717,22 @@ bool FSteamVRHMD::EnableStereo(bool bStereo)
 	if (VRSystem && SceneVP)
 	{
 		int32 PosX, PosY;
-		uint32 Width, Height;
-		GetWindowBounds(&PosX, &PosY, &Width, &Height);
-		SceneVP->SetViewportSize(Width, Height);
+		if( bStereo )
+		{
+			uint32 Width, Height;
+			GetWindowBounds( &PosX, &PosY, &Width, &Height );
+			SceneVP->SetViewportSize( Width, Height );
+		}
+		else
+		{
+			TSharedPtr<SWindow> Window = SceneVP->FindWindow();
+			if( Window.IsValid() )
+			{
+				FVector2D size = SceneVP->FindWindow()->GetSizeInScreen();
+				SceneVP->SetViewportSize( size.X, size.Y );
+				Window->SetViewportSizeDrivenByWindow( true );
+			}
+		}
 	}
 
 	// Uncap fps to enable FPS higher than 62
@@ -841,13 +854,13 @@ void FSteamVRHMD::SetupView(FSceneViewFamily& InViewFamily, FSceneView& InView)
 
 void FSteamVRHMD::PreRenderView_RenderThread(FRHICommandListImmediate& RHICmdList, FSceneView& View)
 {
-	//check(IsInRenderingThread());
+	check(IsInRenderingThread());
 
-	//// The last view location used to set the view will be in BaseHmdOrientation.  We need to calculate the delta from that, so that
-	//// cameras that rely on game objects (e.g. other components) for their positions don't need to be updated on the render thread.
-	//const FQuat DeltaOrient = View.BaseHmdOrientation.Inverse() * TrackingFrame.DeviceOrientation[vr::k_unTrackedDeviceIndex_Hmd];
-	//View.ViewRotation = FRotator(View.ViewRotation.Quaternion() * DeltaOrient);
- //	View.UpdateViewMatrix();
+	// The last view location used to set the view will be in BaseHmdOrientation.  We need to calculate the delta from that, so that
+	// cameras that rely on game objects (e.g. other components) for their positions don't need to be updated on the render thread.
+	const FQuat DeltaOrient = View.BaseHmdOrientation.Inverse() * TrackingFrame.DeviceOrientation[vr::k_unTrackedDeviceIndex_Hmd];
+	View.ViewRotation = FRotator(View.ViewRotation.Quaternion() * DeltaOrient);
+ 	View.UpdateViewMatrix();
 }
 
 void FSteamVRHMD::PreRenderViewFamily_RenderThread(FRHICommandListImmediate& RHICmdList, FSceneViewFamily& ViewFamily)
diff --git a/Engine/Plugins/Runtime/Steam/SteamVR/Source/SteamVR/Private/SteamVRRender.cpp b/Engine/Plugins/Runtime/Steam/SteamVR/Source/SteamVR/Private/SteamVRRender.cpp
--- a/Engine/Plugins/Runtime/Steam/SteamVR/Source/SteamVR/Private/SteamVRRender.cpp
+++ b/Engine/Plugins/Runtime/Steam/SteamVR/Source/SteamVR/Private/SteamVRRender.cpp
@@ -139,7 +139,6 @@ void FSteamVRHMD::D3D11Bridge::FinishRendering()
 	Texture.eType = vr::API_DirectX;
 	Texture.eColorSpace = vr::ColorSpace_Auto;
 	vr::EVRCompositorError Error = Plugin->VRCompositor->Submit(vr::Eye_Left, &Texture, &LeftBounds);
-	check(Error == vr::VRCompositorError_None);
 
 	vr::VRTextureBounds_t RightBounds;
 	RightBounds.uMin = 0.5f;
@@ -149,7 +148,10 @@ void FSteamVRHMD::D3D11Bridge::FinishRendering()
 
 	Texture.handle = RenderTargetTexture;
 	Error = Plugin->VRCompositor->Submit(vr::Eye_Right, &Texture, &RightBounds);
-	check(Error == vr::VRCompositorError_None);
+	if (Error != vr::VRCompositorError_None)
+	{
+		UE_LOG(LogHMD, Log, TEXT("Warning:  SteamVR Compositor had an error on present (%d)"), (int32)Error);
+	}
 }
 
 void FSteamVRHMD::D3D11Bridge::Reset()
diff --git a/Engine/Shaders/BasePassForForwardShadingVertexShader.usf b/Engine/Shaders/BasePassForForwardShadingVertexShader.usf
--- a/Engine/Shaders/BasePassForForwardShadingVertexShader.usf
+++ b/Engine/Shaders/BasePassForForwardShadingVertexShader.usf
@@ -7,7 +7,9 @@
 #define IS_FORWARD_BASEPASS_VERTEX_SHADER 1
 
 // @todo urban: branch was: #define FORCE_FLOATS 1
+#ifndef FORCE_FLOATS
 #define FORCE_FLOATS (COMPILER_METAL)
+#endif
 
 #include "Common.usf"
 #include "BasePassForForwardShadingCommon.usf"
diff --git a/Engine/Shaders/Common.usf b/Engine/Shaders/Common.usf
--- a/Engine/Shaders/Common.usf
+++ b/Engine/Shaders/Common.usf
@@ -12,6 +12,10 @@
 	#include "PS4/PS4Common.usf"
 #endif
 
+#if XBOXONE_PROFILE
+	#define STENCIL_COMPONENT_SWIZZLE .x
+#endif
+
 #if COMPILER_METAL
 	// Helps with iteration when changing Metal shader code generation backend.
 	#include "MetalCommon.usf"
diff --git a/Engine/Shaders/FlatTessellation.usf b/Engine/Shaders/FlatTessellation.usf
--- a/Engine/Shaders/FlatTessellation.usf
+++ b/Engine/Shaders/FlatTessellation.usf
@@ -44,6 +44,12 @@ struct FFlatTessellationHSToDS
 	FHullShaderConstantOutput HS_FlatTrianglesConstant( const OutputPatch<FFlatTessellationHSToDS, 3> I )
 	{
 		FHullShaderConstantOutput O = (FHullShaderConstantOutput)0;
+
+		#if INSTANCED_STEREO
+			ResolvedView = ResolveView(GetEyeIndex(I[0].PassSpecificData.FactoryInterpolants.InterpolantsVSToPS.PackedEyeIndex));
+		#else
+			ResolvedView = ResolveView();
+		#endif
 	    
 #if DISABLE_TESSELLATION_OVERRIDE
 		O.TessFactor[0] = 1.0f; // 1->2 edge
diff --git a/Engine/Shaders/MaterialTemplate.usf b/Engine/Shaders/MaterialTemplate.usf
--- a/Engine/Shaders/MaterialTemplate.usf
+++ b/Engine/Shaders/MaterialTemplate.usf
@@ -1525,7 +1525,8 @@ void CalcPixelMaterialInputs(in out FMaterialPixelParameters Parameters, in out
 #endif
 
 	// normalizing after the tangent space to world space conversion improves quality with sheared bases (UV layout to WS causes shrearing)
-	Parameters.WorldNormal = normalize(TransformTangentVectorToWorld(Parameters.TangentToWorld, Parameters.TangentNormal));
+	// use full precision normalize to avoid overflows
+	Parameters.WorldNormal = normalize(float3(TransformTangentVectorToWorld(Parameters.TangentToWorld, Parameters.TangentNormal)));
 #else
 	// Here we don't supoport two sided materials
 	Parameters.TangentNormal = Parameters.WorldNormal = normalize(Parameters.TangentNormal);
diff --git a/Engine/Shaders/MetalCommon.usf b/Engine/Shaders/MetalCommon.usf
--- a/Engine/Shaders/MetalCommon.usf
+++ b/Engine/Shaders/MetalCommon.usf
@@ -5,6 +5,6 @@
 =============================================================================*/
 
 // Update this GUID to improve shader recompilation for Metal only shaders
-// GUID = {06C79EC9-644E-4D22-ADA9-47DD4BB343E9f}
+// GUID = {06C79EC9-644E-4D22-ADA9-47DD4BB343D0f}
   
 #pragma once
diff --git a/Engine/Shaders/PNTriangles.usf b/Engine/Shaders/PNTriangles.usf
--- a/Engine/Shaders/PNTriangles.usf
+++ b/Engine/Shaders/PNTriangles.usf
@@ -52,6 +52,12 @@ float4 ComputeControlPoint(float4 PositionA, float4 PositionB, float4 NormalA)
 	FHullShaderConstantOutput PNSharedConstantHull( const OutputPatch<FPNTessellationHSToDS, OUT_CONTROL_POINTS> I )
 	{
 		FHullShaderConstantOutput O = (FHullShaderConstantOutput)0;
+
+		#if INSTANCED_STEREO
+			ResolvedView = ResolveView(GetEyeIndex(I[0].PassSpecificData.FactoryInterpolants.InterpolantsVSToPS.PackedEyeIndex));
+		#else
+			ResolvedView = ResolveView();
+		#endif
 		
 #if DISABLE_TESSELLATION_OVERRIDE
 		O.TessFactor[0] = 1.0f; // 1->2 edge
diff --git a/Engine/Shaders/ParticleSimulationShader.usf b/Engine/Shaders/ParticleSimulationShader.usf
--- a/Engine/Shaders/ParticleSimulationShader.usf
+++ b/Engine/Shaders/ParticleSimulationShader.usf
@@ -70,7 +70,7 @@ float2 TileCoord = Input.TexCoord.xy * float2(TILE_SIZE_X,TILE_SIZE_Y) + TileOff
 /** Define to 1 to force no collision in the shader. */
 #define FORCE_NO_COLLISION 0
 
-#if FORCE_NO_COLLISION || FEATURE_LEVEL < FEATURE_LEVEL_SM5
+#if FORCE_NO_COLLISION || FEATURE_LEVEL < FEATURE_LEVEL_SM4
 	#undef DEPTH_BUFFER_COLLISION
 	#define DEPTH_BUFFER_COLLISION 0
 	#undef DISTANCE_FIELD_COLLISION
diff --git a/Engine/Shaders/SimpleElementPixelShader.usf b/Engine/Shaders/SimpleElementPixelShader.usf
--- a/Engine/Shaders/SimpleElementPixelShader.usf
+++ b/Engine/Shaders/SimpleElementPixelShader.usf
@@ -304,9 +304,6 @@ void GammaDistanceFieldMain(
 	float4 BaseTexture = Texture2DSample(InTexture, InTextureSampler,TextureCoordinate);
 	ReplicateChannelSimpleElementShader(BaseTexture);
 	float BaseDist = BaseTexture.a;
-#if COMPILER_METAL && FEATURE_LEVEL >= FEATURE_LEVEL_SM4 // For reasons best known to themselves Apple have made R8Unorm use the sRGB color space on desktop!
-	BaseDist = pow(BaseTexture.a, 2.2f);
-#endif
 	// smooth range between SmoothWidth texels around edge
 	// We want to control font edge smoothness in screen space, but the distance field is stored in texture space.
 	// Using ddx scales the smoothstep steepness based on texel-to-pixel scaling.
@@ -322,9 +319,6 @@ void GammaDistanceFieldMain(
 		float4 ShadowTexture = Texture2DSample(InTexture, InTextureSampler,TextureCoordinate + ShadowDirection);
 		ReplicateChannelSimpleElementShader(ShadowTexture);
 		float ShadowDist = ShadowTexture.a;
-#if COMPILER_METAL && FEATURE_LEVEL >= FEATURE_LEVEL_SM4 // For reasons best known to themselves Apple have made R8Unorm use the sRGB color space on desktop!
-		ShadowDist = pow(ShadowTexture.a, 2.2f);
-#endif
 
 		// smooth range between ShadowSmoothWidth texels around edge
 		float4 ShadowResult = ShadowColor * smoothstep(0.5-ShadowSmoothWidth, 0.5+ShadowSmoothWidth, ShadowDist);
diff --git a/Engine/Shaders/SlateElementPixelShader.usf b/Engine/Shaders/SlateElementPixelShader.usf
--- a/Engine/Shaders/SlateElementPixelShader.usf
+++ b/Engine/Shaders/SlateElementPixelShader.usf
@@ -24,10 +24,12 @@ float3 GammaCorrect(float3 InColor)
 {
 	float3 CorrectedColor = InColor;
 
+#if (!(ES2_PROFILE || ES3_1_PROFILE)) || USE_MATERIALS
 	FLATTEN if( GammaValues.y != 1.0f )
 	{
 		CorrectedColor = ApplyGammaCorrection(CorrectedColor, GammaValues.x);
 	}
+#endif
 
 	return CorrectedColor;
 }
@@ -114,13 +116,6 @@ float4 GetTextureColor( VertexToPixelInterpolants InVertex, float2 UV )
 	BaseColor.a = 1.0f;
 #endif
 
-#if (ES2_PROFILE || ES3_1_PROFILE)
-	FLATTEN if( GammaValues.y != 1.0f )
-	{
-		BaseColor.rgb = sRGBToLinear(BaseColor.rgb);
-	}
-#endif
-
 	return BaseColor;
 }
 #endif
diff --git a/Engine/Shaders/SlateVertexShader.usf b/Engine/Shaders/SlateVertexShader.usf
--- a/Engine/Shaders/SlateVertexShader.usf
+++ b/Engine/Shaders/SlateVertexShader.usf
@@ -25,7 +25,10 @@ VertexToPixelInterpolants Main(
 
 	float4 WorldPosition = float4(InPosition.xy,0,1);
 
+#if (!(ES2_PROFILE || ES3_1_PROFILE)) || USE_MATERIALS
 	InColor.rgb = sRGBToLinear(InColor.rgb);
+#endif
+
 	float4 FinalVertexColor = InColor FCOLOR_COMPONENT_SWIZZLE;
 
 	VOut.MaterialTexCoords = InMaterialTextureCoordinates;
diff --git a/Engine/Source/Developer/Apple/MetalShaderFormat/Private/MetalShaderCompiler.cpp b/Engine/Source/Developer/Apple/MetalShaderFormat/Private/MetalShaderCompiler.cpp
--- a/Engine/Source/Developer/Apple/MetalShaderFormat/Private/MetalShaderCompiler.cpp
+++ b/Engine/Source/Developer/Apple/MetalShaderFormat/Private/MetalShaderCompiler.cpp
@@ -539,6 +539,7 @@ void CompileShader_Metal(const FShaderCompilerInput& Input,FShaderCompilerOutput
 	else if (Input.ShaderFormat == NAME_SF_METAL_MACES3_1)
 	{
 		AdditionalDefines.SetDefine(TEXT("METAL_PROFILE"), 1);
+		AdditionalDefines.SetDefine(TEXT("FORCE_FLOATS"), 1); // Force floats to avoid radr://24884199 & radr://24884860
 		Standard = TEXT("-std=osx-metal1.1");
 		MetalCompilerTarget = HCT_FeatureLevelES3_1;
 	}
diff --git a/Engine/Source/Developer/DerivedDataCache/Private/DerivedDataCache.cpp b/Engine/Source/Developer/DerivedDataCache/Private/DerivedDataCache.cpp
--- a/Engine/Source/Developer/DerivedDataCache/Private/DerivedDataCache.cpp
+++ b/Engine/Source/Developer/DerivedDataCache/Private/DerivedDataCache.cpp
@@ -1037,6 +1037,9 @@ FDerivedDataCache& InternalSingleton()
  */
 class FDerivedDataCacheModule : public IDerivedDataCacheModule
 {
+	/** Cached reference to DDC singleton, helpful to control singleton's lifetime. */
+	FDerivedDataCache* DDC;
+
 public:
 	virtual FDerivedDataCacheInterface& GetDDC() override
 	{
@@ -1045,14 +1048,24 @@ public:
 
 	virtual void StartupModule() override
 	{
+
+		// make sure DDC gets created early, previously it might have happened in ShutdownModule() (for PrintLeaks()) when it was already too late
+		DDC = static_cast< FDerivedDataCache* >( &GetDDC() );
 	}
 
 	virtual void ShutdownModule() override
 	{
 		FDDCCleanup::Shutdown();
 
-		FDerivedDataCache& DDC = static_cast< FDerivedDataCache& >( GetDDC() );
-		DDC.PrintLeaks();
+		if (DDC)
+		{
+			DDC->PrintLeaks();
+		}
+	}
+
+	FDerivedDataCacheModule()
+		:	DDC(nullptr)
+	{
 	}
 };
 
diff --git a/Engine/Source/Developer/DesktopPlatform/Private/PlatformInfo.cpp b/Engine/Source/Developer/DesktopPlatform/Private/PlatformInfo.cpp
--- a/Engine/Source/Developer/DesktopPlatform/Private/PlatformInfo.cpp
+++ b/Engine/Source/Developer/DesktopPlatform/Private/PlatformInfo.cpp
@@ -81,17 +81,17 @@ static const FPlatformInfo PlatformInfoArray[] = {
 	BuildPlatformInfo(TEXT("WindowsNoEditor_Win32"),	TEXT("WindowsNoEditor"),	LOCTEXT("WindowsNoEditor_Win32", "Windows (32-bit)"),		EPlatformType::Game,		EPlatformFlags::BuildFlavor,	FPlatformIconPaths(TEXT("Launcher/WindowsTarget/Platform_WindowsNoEditor_24x"), TEXT("Launcher/WindowsTarget/Platform_WindowsNoEditor_128x")),		TEXT("-targetplatform=Win32"),							TEXT(""),			EPlatformSDKStatus::Unknown,	TEXT("/Engine/Tutorial/Installation/InstallingVisualStudioTutorial.InstallingVisualStudioTutorial"),	IsAvailableOnWindows,						TEXT("Win32"),		IsAvailableOnWindows,	true),
 	BuildPlatformInfo(TEXT("WindowsNoEditor_Win64"),	TEXT("WindowsNoEditor"),	LOCTEXT("WindowsNoEditor_Win64", "Windows (64-bit)"),		EPlatformType::Game,		EPlatformFlags::BuildFlavor,	FPlatformIconPaths(TEXT("Launcher/WindowsTarget/Platform_WindowsNoEditor_24x"), TEXT("Launcher/WindowsTarget/Platform_WindowsNoEditor_128x")),		TEXT("-targetplatform=Win64"),							TEXT(""),			EPlatformSDKStatus::Unknown,	TEXT("/Engine/Tutorial/Installation/InstallingVisualStudioTutorial.InstallingVisualStudioTutorial"),	IsAvailableOnWindows,						TEXT("Win64"),		IsAvailableOnWindows,	true),
 	BuildPlatformInfo(TEXT("Windows"),					TEXT("Windows"),			LOCTEXT("WindowsEditor", "Windows (Editor)"),				EPlatformType::Editor,		EPlatformFlags::None,			FPlatformIconPaths(TEXT("Launcher/WindowsTarget/Platform_Windows_24x"), TEXT("Launcher/WindowsTarget/Platform_Windows_128x")),						TEXT(""),												TEXT(""),			EPlatformSDKStatus::Unknown,	TEXT("/Engine/Tutorial/Installation/InstallingVisualStudioTutorial.InstallingVisualStudioTutorial"),	IsAvailableOnWindows,						TEXT("Win64"),		IsAvailableOnWindows,	true),
-	BuildPlatformInfo(TEXT("WindowsClient"),			TEXT("WindowsClient"),		LOCTEXT("WindowsClient", "Windows (Client-only)"),			EPlatformType::Client,		EPlatformFlags::None,			FPlatformIconPaths(TEXT("Launcher/WindowsTarget/Platform_Windows_24x"), TEXT("Launcher/WindowsTarget/Platform_Windows_128x")),						TEXT(""),												TEXT(""),			EPlatformSDKStatus::Unknown,	TEXT("/Engine/Tutorial/Installation/InstallingVisualStudioTutorial.InstallingVisualStudioTutorial"),	IsAvailableOnWindows,						TEXT("Win64"),		IsAvailableOnWindows,	true),
+	BuildPlatformInfo(TEXT("WindowsClient"),			TEXT("WindowsClient"),		LOCTEXT("WindowsClient", "Windows (Client-only)"),			EPlatformType::Client,		EPlatformFlags::None,			FPlatformIconPaths(TEXT("Launcher/WindowsTarget/Platform_Windows_24x"), TEXT("Launcher/WindowsTarget/Platform_Windows_128x")),						TEXT("-client"),										TEXT(""),			EPlatformSDKStatus::Unknown,	TEXT("/Engine/Tutorial/Installation/InstallingVisualStudioTutorial.InstallingVisualStudioTutorial"),	IsAvailableOnWindows,						TEXT("Win64"),		IsAvailableOnWindows,	true),
 	BuildPlatformInfo(TEXT("WindowsServer"),			TEXT("WindowsServer"),		LOCTEXT("WindowsServer", "Windows (Dedicated Server)"),		EPlatformType::Server,		EPlatformFlags::None,			FPlatformIconPaths(TEXT("Launcher/WindowsTarget/Platform_WindowsServer_24x"), TEXT("Launcher/WindowsTarget/Platform_WindowsServer_128x")),			TEXT(""),												TEXT(""),			EPlatformSDKStatus::Unknown,	TEXT("/Engine/Tutorial/Installation/InstallingVisualStudioTutorial.InstallingVisualStudioTutorial"),	IsAvailableOnWindows,						TEXT("Win64"),		IsAvailableOnWindows,	true),
 
 	BuildPlatformInfo(TEXT("MacNoEditor"),				TEXT("MacNoEditor"),		LOCTEXT("MacNoEditor", "Mac"),								EPlatformType::Game,		EPlatformFlags::None,			FPlatformIconPaths(TEXT("Launcher/MacTarget/Platform_Mac_24x"), TEXT("Launcher/MacTarget/Platform_Mac_128x")),										TEXT("-targetplatform=Mac"),							TEXT(""),			EPlatformSDKStatus::Unknown,	TEXT("/Engine/Tutorial/Installation/InstallingXCodeTutorial.InstallingXCodeTutorial"),					IsAvailableOnMac,							TEXT("Mac"),		IsAvailableOnMac,		true),
 	BuildPlatformInfo(TEXT("Mac"),						TEXT("Mac"),				LOCTEXT("MacEditor", "Mac (Editor)"),						EPlatformType::Editor,		EPlatformFlags::None,			FPlatformIconPaths(TEXT("Launcher/MacTarget/Platform_Mac_24x"), TEXT("Launcher/MacTarget/Platform_Mac_128x")),										TEXT(""),												TEXT(""),			EPlatformSDKStatus::Unknown,	TEXT("/Engine/Tutorial/Installation/InstallingXCodeTutorial.InstallingXCodeTutorial"),					IsAvailableOnMac,							TEXT("Mac"),		IsAvailableOnMac,		true),
-	BuildPlatformInfo(TEXT("MacClient"),				TEXT("MacClient"),			LOCTEXT("MacClient", "Mac (Client-only)"),					EPlatformType::Client,		EPlatformFlags::None,			FPlatformIconPaths(TEXT("Launcher/MacTarget/Platform_Mac_24x"), TEXT("Launcher/MacTarget/Platform_Mac_128x")),										TEXT(""),												TEXT(""),			EPlatformSDKStatus::Unknown,	TEXT("/Engine/Tutorial/Installation/InstallingXCodeTutorial.InstallingXCodeTutorial"),					IsAvailableOnMac,							TEXT("Mac"),		IsAvailableOnMac,		true),
+	BuildPlatformInfo(TEXT("MacClient"),				TEXT("MacClient"),			LOCTEXT("MacClient", "Mac (Client-only)"),					EPlatformType::Client,		EPlatformFlags::None,			FPlatformIconPaths(TEXT("Launcher/MacTarget/Platform_Mac_24x"), TEXT("Launcher/MacTarget/Platform_Mac_128x")),										TEXT("-client"),										TEXT(""),			EPlatformSDKStatus::Unknown,	TEXT("/Engine/Tutorial/Installation/InstallingXCodeTutorial.InstallingXCodeTutorial"),					IsAvailableOnMac,							TEXT("Mac"),		IsAvailableOnMac,		true),
 	BuildPlatformInfo(TEXT("MacServer"),				TEXT("MacServer"),			LOCTEXT("MacServer", "Mac (Dedicated Server)"),				EPlatformType::Server,		EPlatformFlags::None,			FPlatformIconPaths(TEXT("Launcher/MacTarget/Platform_Mac_24x"), TEXT("Launcher/MacTarget/Platform_Mac_128x")),										TEXT(""),												TEXT(""),			EPlatformSDKStatus::Unknown,	TEXT("/Engine/Tutorial/Installation/InstallingXCodeTutorial.InstallingXCodeTutorial"),					IsAvailableOnMac,							TEXT("Mac"),		IsAvailableOnMac,		true),
 
 	BuildPlatformInfo(TEXT("LinuxNoEditor"),			TEXT("LinuxNoEditor"),		LOCTEXT("LinuxNoEditor", "Linux"),							EPlatformType::Game,		EPlatformFlags::None,			FPlatformIconPaths(TEXT("Launcher/LinuxTarget/Platform_Linux_24x"), TEXT("Launcher/LinuxTarget/Platform_Linux_128x")),								TEXT("-targetplatform=Linux"),							TEXT("Linux_x64"),	EPlatformSDKStatus::Unknown,	TEXT("/Platforms/Linux/GettingStarted"),																IsAvailableOnLinux || IsAvailableOnWindows,	TEXT("Linux"),		IsAvailableOnLinux,		true),	
 	BuildPlatformInfo(TEXT("Linux"),					TEXT("Linux"),				LOCTEXT("LinuxEditor", "Linux (Editor)"),					EPlatformType::Editor,		EPlatformFlags::None,			FPlatformIconPaths(TEXT("Launcher/LinuxTarget/Platform_Linux_24x"), TEXT("Launcher/LinuxTarget/Platform_Linux_128x")),								TEXT(""),												TEXT("Linux_x64"),	EPlatformSDKStatus::Unknown,	TEXT("/Platforms/Linux/GettingStarted"),																IsAvailableOnLinux,							TEXT("Linux"),		IsAvailableOnLinux,		true),
-	BuildPlatformInfo(TEXT("LinuxClient"),				TEXT("LinuxClient"),		LOCTEXT("LinuxClient", "Linux (Client-only)"),				EPlatformType::Client,		EPlatformFlags::None,			FPlatformIconPaths(TEXT("Launcher/LinuxTarget/Platform_Linux_24x"), TEXT("Launcher/LinuxTarget/Platform_Linux_128x")),								TEXT(""),												TEXT("Linux_x64"),	EPlatformSDKStatus::Unknown,	TEXT("/Platforms/Linux/GettingStarted"),																IsAvailableOnLinux || IsAvailableOnWindows,	TEXT("Linux"),		IsAvailableOnLinux,		true),
+	BuildPlatformInfo(TEXT("LinuxClient"),				TEXT("LinuxClient"),		LOCTEXT("LinuxClient", "Linux (Client-only)"),				EPlatformType::Client,		EPlatformFlags::None,			FPlatformIconPaths(TEXT("Launcher/LinuxTarget/Platform_Linux_24x"), TEXT("Launcher/LinuxTarget/Platform_Linux_128x")),								TEXT("-client"),										TEXT("Linux_x64"),	EPlatformSDKStatus::Unknown,	TEXT("/Platforms/Linux/GettingStarted"),																IsAvailableOnLinux || IsAvailableOnWindows,	TEXT("Linux"),		IsAvailableOnLinux,		true),
 	BuildPlatformInfo(TEXT("LinuxServer"),				TEXT("LinuxServer"),		LOCTEXT("LinuxServer", "Linux (Dedicated Server)"),			EPlatformType::Server,		EPlatformFlags::None,			FPlatformIconPaths(TEXT("Launcher/LinuxTarget/Platform_Linux_24x"), TEXT("Launcher/LinuxTarget/Platform_Linux_128x")),								TEXT(""),												TEXT("Linux_x64"),	EPlatformSDKStatus::Unknown,	TEXT("/Platforms/Linux/GettingStarted"),																IsAvailableOnLinux || IsAvailableOnWindows,	TEXT("Linux"),		IsAvailableOnLinux,		true),
 
 	BuildPlatformInfo(TEXT("IOS"),						TEXT("IOS"),				LOCTEXT("IOS", "iOS"),										EPlatformType::Game,		EPlatformFlags::None,			FPlatformIconPaths(TEXT("Launcher/iOSTarget/Platform_iOS_24x"), TEXT("Launcher/iOSTarget/Platform_iOS_128x")),										TEXT(""),												TEXT(""),			EPlatformSDKStatus::Unknown,	IOSTutorial,																							IsAvailableOnWindows || IsAvailableOnMac,	TEXT("IOS"),		false,					true),
diff --git a/Engine/Source/Developer/GameplayDebugger/GameplayDebugger.Build.cs b/Engine/Source/Developer/GameplayDebugger/GameplayDebugger.Build.cs
--- a/Engine/Source/Developer/GameplayDebugger/GameplayDebugger.Build.cs
+++ b/Engine/Source/Developer/GameplayDebugger/GameplayDebugger.Build.cs
@@ -54,8 +54,16 @@ namespace UnrealBuildTool.Rules
 						}
 					);
 
-                PrivateDependencyModuleNames.AddRange(new string[] { "AIModule", "GameplayTasks", "GameplayAbilities" });
-                CircularlyReferencedDependentModules.AddRange(new string[] { "AIModule", "GameplayTasks", "GameplayAbilities" });
+                PrivateDependencyModuleNames.AddRange(new string[] { "AIModule", "GameplayTasks" });
+                CircularlyReferencedDependentModules.AddRange(new string[] { "AIModule", "GameplayTasks" });
+
+				DynamicallyLoadedModuleNames.AddRange(
+                new string[]
+				    {
+					    // Removed direct dependency on GameplayAbilities from GameplayDebugger to solve problems with its classes showing in blueprints editor as a result of the OS automatically loading the library
+					    "GameplayAbilities",
+					}
+				);
 			}
         }
     }
diff --git a/Engine/Source/Developer/GameplayDebugger/Private/GameplayDebuggingComponent.cpp b/Engine/Source/Developer/GameplayDebugger/Private/GameplayDebuggingComponent.cpp
--- a/Engine/Source/Developer/GameplayDebugger/Private/GameplayDebuggingComponent.cpp
+++ b/Engine/Source/Developer/GameplayDebugger/Private/GameplayDebuggingComponent.cpp
@@ -16,7 +16,6 @@
 #include "EnvironmentQuery/EnvQueryTest.h"
 #include "BehaviorTree/BlackboardComponent.h"
 #include "BehaviorTree/BehaviorTreeComponent.h"
-#include "AbilitySystemComponent.h"
 #include "AIController.h"
 #include "BrainComponent.h"
 #include "BehaviorTreeDelegates.h"
@@ -25,6 +24,7 @@
 #include "GameFramework/CharacterMovementComponent.h"
 #include "Engine/Channel.h"
 #include "Animation/AnimMontage.h"
+#include "GameplayTasksComponent.h"
 #include "Perception/AIPerceptionComponent.h"
 
 #if WITH_EDITOR
diff --git a/Engine/Source/Developer/LauncherServices/Private/Launcher/LauncherUATTask.h b/Engine/Source/Developer/LauncherServices/Private/Launcher/LauncherUATTask.h
--- a/Engine/Source/Developer/LauncherServices/Private/Launcher/LauncherUATTask.h
+++ b/Engine/Source/Developer/LauncherServices/Private/Launcher/LauncherUATTask.h
@@ -63,6 +63,7 @@ protected:
 		UATCommandLine += FParse::Param( FCommandLine::Get(), TEXT("development") ) || ChainState.Profile->IsBuildingUAT() ? TEXT("") : TEXT(" -nocompile");
 		// we never want to build the editor when launching from the editor or running with an installed engine (which can't rebuild itself)
 		UATCommandLine += GIsEditor || FApp::IsEngineInstalled() ? TEXT(" -nocompileeditor") : TEXT("");
+		UATCommandLine += FApp::IsEngineInstalled() ? TEXT(" -installed") : TEXT("");
 
 		// specify the path to the editor exe if necessary
 		if(EditorExe.Len() > 0)
diff --git a/Engine/Source/Developer/LauncherServices/Private/Launcher/LauncherWorker.cpp b/Engine/Source/Developer/LauncherServices/Private/Launcher/LauncherWorker.cpp
--- a/Engine/Source/Developer/LauncherServices/Private/Launcher/LauncherWorker.cpp
+++ b/Engine/Source/Developer/LauncherServices/Private/Launcher/LauncherWorker.cpp
@@ -251,11 +251,11 @@ FString FLauncherWorker::CreateUATCommand( const ILauncherProfileRef& InProfile,
 		{
 			Platforms += TEXT("+Linux");
 		}
-		else if (PlatformInfo->TargetPlatformName == FName("WindowsNoEditor") || PlatformInfo->TargetPlatformName == FName("Windows"))
+		else if (PlatformInfo->TargetPlatformName == FName("WindowsNoEditor") || PlatformInfo->TargetPlatformName == FName("Windows") || PlatformInfo->TargetPlatformName == FName("WindowsClient"))
 		{
 			Platforms += TEXT("+Win64");
 		}
-		else if (PlatformInfo->TargetPlatformName == FName("MacNoEditor"))
+		else if (PlatformInfo->TargetPlatformName == FName("MacNoEditor") || PlatformInfo->TargetPlatformName == FName("MacClient"))
 		{
 			Platforms += TEXT("+Mac");
 		}
diff --git a/Engine/Source/Developer/MeshUtilities/Private/MeshUtilities.cpp b/Engine/Source/Developer/MeshUtilities/Private/MeshUtilities.cpp
--- a/Engine/Source/Developer/MeshUtilities/Private/MeshUtilities.cpp
+++ b/Engine/Source/Developer/MeshUtilities/Private/MeshUtilities.cpp
@@ -5199,7 +5199,7 @@ static void RetrieveValidStaticMeshComponentsForMerging(AActor* InActor, TArray<
 	TInlineComponentArray<UStaticMeshComponent*> Components;
 	InActor->GetComponents<UStaticMeshComponent>(Components);
 	// TODO: support derived classes from static component
-	Components.RemoveAll([](UStaticMeshComponent* Val){ return !(Val->IsA(UStaticMeshComponent::StaticClass()) || Val->IsA(USplineMeshComponent::StaticClass())); });
+	Components.RemoveAll([](UStaticMeshComponent* Val){ return !(Val->GetClass() == UStaticMeshComponent::StaticClass() || Val->IsA(USplineMeshComponent::StaticClass())); });
 
 	// TODO: support non-opaque materials
 	//Components.RemoveAll(&NonOpaqueMaterialPredicate);
@@ -6979,7 +6979,6 @@ void FMeshUtilities::MergeStaticMeshComponents(const TArray<UStaticMeshComponent
 		}
 	}
 
-#if 0 // TODO AG: The FRawMesh structure has changed
 	if (InSettings.bMergePhysicsData)
 	{
 		for (int32 MeshId = 0; MeshId < ComponentsToMerge.Num(); ++MeshId)
@@ -6994,7 +6993,6 @@ void FMeshUtilities::MergeStaticMeshComponents(const TArray<UStaticMeshComponent
 			}
 		}
 	}
-#endif
 
 	
 
@@ -7230,16 +7228,14 @@ void FMeshUtilities::MergeStaticMeshComponents(const TArray<UStaticMeshComponent
 	}
 
 	// Transform physics primitives to merged mesh pivot
-#if 0 // TODO AG: The FRawMesh structure has changed
-	if (InSettings.bMergePhysicsData && !MergedMesh.Pivot.IsZero())
+	if (InSettings.bMergePhysicsData && !MergedAssetPivot.IsZero())
 	{
-		FTransform PivotTM(-MergedMesh.Pivot);
+		FTransform PivotTM(-MergedAssetPivot);
 		for (auto& SourceMesh : SourceMeshes)
 		{
 			TransformPhysicsGeometry(PivotTM, SourceMesh.AggGeom);
 		}
 	}
-#endif 
 
 	// Compute target lightmap channel for each LOD
 	// User can specify any index, but there are should not be empty gaps in UV channel list
@@ -7323,7 +7319,6 @@ void FMeshUtilities::MergeStaticMeshComponents(const TArray<UStaticMeshComponent
 			StaticMesh->Materials.Add(Material);
 		}
 
-#if 0 // TODO AG: The FRawMesh structure has changed
 		if (InSettings.bMergePhysicsData && BodySetupSource)
 		{
 			StaticMesh->CreateBodySetup();
@@ -7335,7 +7330,6 @@ void FMeshUtilities::MergeStaticMeshComponents(const TArray<UStaticMeshComponent
 				StaticMesh->BodySetup->AddCollisionFrom(SourceMesh.AggGeom);
 			}
 		}
-#endif
 
 		MainTask.EnterProgressFrame(10, LOCTEXT("MeshUtilities_MergeStaticMeshComponents_BuildingStaticMesh", "Building Static Mesh"));
 
diff --git a/Engine/Source/Developer/MeshUtilities/Public/MeshMergeData.h b/Engine/Source/Developer/MeshUtilities/Public/MeshMergeData.h
--- a/Engine/Source/Developer/MeshUtilities/Public/MeshMergeData.h
+++ b/Engine/Source/Developer/MeshUtilities/Public/MeshMergeData.h
@@ -3,6 +3,7 @@
 #pragma once
 
 #include "RawMesh.h"
+#include "PhysicsEngine/AggregateGeom.h"
 
 /** Structure holding intermediate mesh merging data that is used throughout the mesh merging and proxy creation processes */
 struct FMeshMergeData
@@ -26,5 +27,6 @@ struct FRawMeshExt
 	FRawMeshExt()
 	{}
 
-	FMeshMergeData MeshLODData[MAX_STATIC_MESH_LODS];
+	FMeshMergeData		MeshLODData[MAX_STATIC_MESH_LODS];
+	FKAggregateGeom		AggGeom;
 };
diff --git a/Engine/Source/Developer/TextureFormatAndroid/TextureFormatAndroid.Build.cs b/Engine/Source/Developer/TextureFormatAndroid/TextureFormatAndroid.Build.cs
--- a/Engine/Source/Developer/TextureFormatAndroid/TextureFormatAndroid.Build.cs
+++ b/Engine/Source/Developer/TextureFormatAndroid/TextureFormatAndroid.Build.cs
@@ -31,5 +31,10 @@ public class TextureFormatAndroid : ModuleRules
 			AddEngineThirdPartyPrivateStaticDependencies(Target, "QualcommTextureConverter");
 		}
 
+                // opt-out from precompile for Linux (this module cannot be built for Linux atm)
+                if (Target.Platform == UnrealTargetPlatform.Linux)
+                {
+                    PrecompileForTargets = PrecompileTargetsType.None;
+                }
 	}
 }
diff --git a/Engine/Source/Developer/iOS/IOSTargetPlatform/Private/IOSTargetPlatform.h b/Engine/Source/Developer/iOS/IOSTargetPlatform/Private/IOSTargetPlatform.h
--- a/Engine/Source/Developer/iOS/IOSTargetPlatform/Private/IOSTargetPlatform.h
+++ b/Engine/Source/Developer/iOS/IOSTargetPlatform/Private/IOSTargetPlatform.h
@@ -52,7 +52,7 @@ public:
 
     virtual FString IniPlatformName() const override
     {
-        return bIsTVOS ? "TVOS" : "IOS";
+        return "IOS";
     }
     
 	virtual void EnableDeviceCheck(bool OnOff) override;
diff --git a/Engine/Source/Editor/BlueprintGraph/Private/K2Node_SetFieldsInStruct.cpp b/Engine/Source/Editor/BlueprintGraph/Private/K2Node_SetFieldsInStruct.cpp
--- a/Engine/Source/Editor/BlueprintGraph/Private/K2Node_SetFieldsInStruct.cpp
+++ b/Engine/Source/Editor/BlueprintGraph/Private/K2Node_SetFieldsInStruct.cpp
@@ -191,7 +191,7 @@ void UK2Node_SetFieldsInStruct::ValidateNodeDuringCompilation(FCompilerResultsLo
 	Super::ValidateNodeDuringCompilation(MessageLog);
 
 	UEdGraphPin* FoundPin = FindPin(SetFieldsInStructHelper::StructRefPinName());
-	if (!ensure(FoundPin) || (FoundPin->LinkedTo.Num() <= 0))
+	if (!FoundPin || (FoundPin->LinkedTo.Num() <= 0))
 	{
 		FText ErrorMessage = LOCTEXT("SetStructFields_NoStructRefError", "The @@ pin must be connected to the struct that you wish to set.");
 		MessageLog.Error(*ErrorMessage.ToString(), FoundPin);
diff --git a/Engine/Source/Editor/GameplayAbilitiesEditor/Private/AttributeDetails.cpp b/Engine/Source/Editor/GameplayAbilitiesEditor/Private/AttributeDetails.cpp
--- a/Engine/Source/Editor/GameplayAbilitiesEditor/Private/AttributeDetails.cpp
+++ b/Engine/Source/Editor/GameplayAbilitiesEditor/Private/AttributeDetails.cpp
@@ -6,6 +6,7 @@
 #include "KismetEditorUtilities.h"
 #include "AttributeSet.h"
 #include "GameplayAbilitiesModule.h"
+#include "AbilitySystemGlobals.h"
 #include "SGameplayAttributeGraphPin.h"
 #include "SSearchBox.h"
 #include "STextComboBox.h"
diff --git a/Engine/Source/Editor/GameplayAbilitiesEditor/Private/GameplayCueTagDetails.cpp b/Engine/Source/Editor/GameplayAbilitiesEditor/Private/GameplayCueTagDetails.cpp
--- a/Engine/Source/Editor/GameplayAbilitiesEditor/Private/GameplayCueTagDetails.cpp
+++ b/Engine/Source/Editor/GameplayAbilitiesEditor/Private/GameplayCueTagDetails.cpp
@@ -6,6 +6,7 @@
 #include "KismetEditorUtilities.h"
 #include "AttributeSet.h"
 #include "GameplayAbilitiesModule.h"
+#include "AbilitySystemGlobals.h"
 #include "SGameplayAttributeGraphPin.h"
 #include "SSearchBox.h"
 #include "STextComboBox.h"
diff --git a/Engine/Source/Editor/GraphEditor/Private/KismetNodes/SGraphNodeMakeStruct.cpp b/Engine/Source/Editor/GraphEditor/Private/KismetNodes/SGraphNodeMakeStruct.cpp
--- a/Engine/Source/Editor/GraphEditor/Private/KismetNodes/SGraphNodeMakeStruct.cpp
+++ b/Engine/Source/Editor/GraphEditor/Private/KismetNodes/SGraphNodeMakeStruct.cpp
@@ -419,7 +419,10 @@ void SGraphNodeMakeStruct::CreatePinWidgets()
 	{
 		UEdGraphPin* CurPin = GraphNode->Pins[PinIndex];
 
-		AddPin(CreatePinWidget(CurPin).ToSharedRef());
+		if (ShouldPinBeHidden(CurPin))
+		{
+			AddPin(CreatePinWidget(CurPin).ToSharedRef());
+		}
 	}
 }
 
diff --git a/Engine/Source/Editor/GraphEditor/Private/SGraphNode.cpp b/Engine/Source/Editor/GraphEditor/Private/SGraphNode.cpp
--- a/Engine/Source/Editor/GraphEditor/Private/SGraphNode.cpp
+++ b/Engine/Source/Editor/GraphEditor/Private/SGraphNode.cpp
@@ -974,7 +974,7 @@ void SGraphNode::CreateAdvancedViewArrow(TSharedPtr<SVerticalBox> MainBox)
 	}
 }
 
-void SGraphNode::CreateStandardPinWidget(UEdGraphPin* CurPin)
+bool SGraphNode::ShouldPinBeHidden(const UEdGraphPin* InPin) const
 {
 	const UEdGraphSchema_K2* K2Schema = Cast<const UEdGraphSchema_K2>(GraphNode->GetSchema());
 
@@ -982,7 +982,7 @@ void SGraphNode::CreateStandardPinWidget(UEdGraphPin* CurPin)
 	bool bHideNoConnectionNoDefaultPins = false;
 
 	// Not allowed to hide exec pins 
-	const bool bCanHidePin = (K2Schema && (CurPin->PinType.PinCategory != K2Schema->PC_Exec));
+	const bool bCanHidePin = (K2Schema && (InPin->PinType.PinCategory != K2Schema->PC_Exec));
 
 	if (OwnerGraphPanelPtr.IsValid() && bCanHidePin)
 	{
@@ -990,17 +990,24 @@ void SGraphNode::CreateStandardPinWidget(UEdGraphPin* CurPin)
 		bHideNoConnectionNoDefaultPins = OwnerGraphPanelPtr.Pin()->GetPinVisibility() == SGraphEditor::Pin_HideNoConnectionNoDefault;
 	}
 
-	const bool bIsOutputPin= CurPin->Direction == EGPD_Output;
-	const bool bPinHasDefaultValue = !CurPin->DefaultValue.IsEmpty() || (CurPin->DefaultObject != NULL);
-	const bool bIsSelfTarget = K2Schema && (CurPin->PinType.PinCategory == K2Schema->PC_Object) && (CurPin->PinName == K2Schema->PN_Self);
+	const bool bIsOutputPin = InPin->Direction == EGPD_Output;
+	const bool bPinHasDefaultValue = !InPin->DefaultValue.IsEmpty() || (InPin->DefaultObject != NULL);
+	const bool bIsSelfTarget = K2Schema && (InPin->PinType.PinCategory == K2Schema->PC_Object) && (InPin->PinName == K2Schema->PN_Self);
 	const bool bPinHasValidDefault = !bIsOutputPin && (bPinHasDefaultValue || bIsSelfTarget);
-	const bool bPinHasConections = CurPin->LinkedTo.Num() > 0;
+	const bool bPinHasConections = InPin->LinkedTo.Num() > 0;
 
-	const bool bPinDesiresToBeHidden = CurPin->bHidden || (bHideNoConnectionPins && !bPinHasConections) || (bHideNoConnectionNoDefaultPins && !bPinHasConections && !bPinHasValidDefault); 
+	const bool bPinDesiresToBeHidden = InPin->bHidden || (bHideNoConnectionPins && !bPinHasConections) || (bHideNoConnectionNoDefaultPins && !bPinHasConections && !bPinHasValidDefault);
 
 	// No matter how strong the desire, a pin with connections can never be hidden!
 	const bool bShowPin = !bPinDesiresToBeHidden || bPinHasConections;
 
+	return bShowPin;
+}
+
+void SGraphNode::CreateStandardPinWidget(UEdGraphPin* CurPin)
+{
+	const bool bShowPin = ShouldPinBeHidden(CurPin);
+
 	if (bShowPin)
 	{
 		TSharedPtr<SGraphPin> NewPin = CreatePinWidget(CurPin);
diff --git a/Engine/Source/Editor/GraphEditor/Public/SGraphNode.h b/Engine/Source/Editor/GraphEditor/Public/SGraphNode.h
--- a/Engine/Source/Editor/GraphEditor/Public/SGraphNode.h
+++ b/Engine/Source/Editor/GraphEditor/Public/SGraphNode.h
@@ -330,6 +330,9 @@ protected:
 	/* Populate a meta data tag with information about this graph node */
 	virtual void PopulateMetaTag(class FGraphNodeMetaData* TagMeta) const;
 
+	/** Returns TRUE if the input pin should be hidden from view */
+	bool ShouldPinBeHidden(const UEdGraphPin* InPin) const;
+
 protected:
 	/** Input pin widgets on this node */
 	TArray< TSharedRef<SGraphPin> > InputPins;
diff --git a/Engine/Source/Editor/Kismet/Private/SBlueprintPalette.cpp b/Engine/Source/Editor/Kismet/Private/SBlueprintPalette.cpp
--- a/Engine/Source/Editor/Kismet/Private/SBlueprintPalette.cpp
+++ b/Engine/Source/Editor/Kismet/Private/SBlueprintPalette.cpp
@@ -1468,7 +1468,7 @@ FText SBlueprintPaletteItem::GetToolTipText() const
 		else if (PaletteAction->GetTypeId() == FEdGraphSchemaAction_K2Graph::StaticGetTypeId())
 		{
 			FEdGraphSchemaAction_K2Graph* GraphAction = (FEdGraphSchemaAction_K2Graph*)PaletteAction.Get();
-			if (GraphAction->EdGraph != NULL)
+			if (GraphAction->EdGraph != NULL && GraphAction->EdGraph->GetSchema())
 			{
 				if (const UEdGraphSchema* GraphSchema = GraphAction->EdGraph->GetSchema())
 				{
diff --git a/Engine/Source/Editor/KismetCompiler/Private/KismetCompilerVMBackend.cpp b/Engine/Source/Editor/KismetCompiler/Private/KismetCompilerVMBackend.cpp
--- a/Engine/Source/Editor/KismetCompiler/Private/KismetCompilerVMBackend.cpp
+++ b/Engine/Source/Editor/KismetCompiler/Private/KismetCompilerVMBackend.cpp
@@ -15,6 +15,7 @@
 #include "Editor/UnrealEd/Public/Kismet2/KismetDebugUtilities.h"
 #include "Engine/UserDefinedStruct.h"
 
+#define LOCTEXT_NAMESPACE "KismetCompilerVMBackend"
 //////////////////////////////////////////////////////////////////////////
 // FScriptBytecodeWriter
 
@@ -565,6 +566,7 @@ public:
 							// Create a new term for each property, and serialize it out
 							FBPTerminal NewTerm;
 							NewTerm.bIsLiteral = true;
+							NewTerm.Source = Term->Source;
 							Prop->ExportText_InContainer(ArrayIter, NewTerm.Name, StructData, StructData, NULL, PPF_None);
 							if (Prop->IsA(UTextProperty::StaticClass()))
 							{
@@ -600,6 +602,7 @@ public:
 				{
 					FBPTerminal NewTerm;
 					NewTerm.bIsLiteral = true;
+					NewTerm.Source = Term->Source;
 					uint8* RawElemData = ScriptArrayHelper.GetRawPtr(ElemIdx);
 					InnerProp->ExportText_Direct(NewTerm.Name, RawElemData, RawElemData, NULL, PPF_None);
 					if (InnerProp->IsA(UTextProperty::StaticClass()))
@@ -671,7 +674,13 @@ public:
 			// Cannot assign a literal to a multicast delegate; it should be added instead of assigned
 			else
 			{
-				ensureMsgf(false, TEXT("It is not possible to express this type as a literal value! (%s)"), *CoerceProperty->GetFullName());
+				if (ensure(CurrentCompilerContext))
+				{
+					FFormatNamedArguments Args;
+					Args.Add(TEXT("PropertyType"), CoerceProperty->GetClass()->GetDisplayNameText());
+					Args.Add(TEXT("PropertyName"), CoerceProperty->GetDisplayNameText());
+					CurrentCompilerContext->MessageLog.Error(*FText::Format(LOCTEXT("InvalidProperty", "It is not possible to express this type ({PropertyType}) as a literal value for the property {PropertyName} on pin @@! If it is inside a struct, you can add a Make struct node to resolve this issue!"), Args).ToString(), Term->Source);
+				}
 			}
 		}
 		else
@@ -1689,3 +1698,5 @@ void FKismetCompilerVMBackend::ConstructFunction(FKismetFunctionContext& Functio
 	static_assert(sizeof(CodeSkipSizeType) == 4, "Update this code as size changed.");
 #endif
 }
+
+#undef LOCTEXT_NAMESPACE
\ No newline at end of file
diff --git a/Engine/Source/Editor/UMGEditor/Private/WidgetBlueprintCompiler.cpp b/Engine/Source/Editor/UMGEditor/Private/WidgetBlueprintCompiler.cpp
--- a/Engine/Source/Editor/UMGEditor/Private/WidgetBlueprintCompiler.cpp
+++ b/Engine/Source/Editor/UMGEditor/Private/WidgetBlueprintCompiler.cpp
@@ -272,6 +272,10 @@ void FWidgetBlueprintCompiler::FinishCompilingClass(UClass* Class)
 	if ( Blueprint->SkeletonGeneratedClass != Class )
 	{
 		UWidgetBlueprintGeneratedClass* BPGClass = CastChecked<UWidgetBlueprintGeneratedClass>(Class);
+		if( !Blueprint->bHasBeenRegenerated )
+		{
+			FBlueprintEditorUtils::ForceLoadMembers(Blueprint->WidgetTree);
+		}
 		BPGClass->WidgetTree = BPGClass->DesignerWidgetTree = DuplicateObject<UWidgetTree>(Blueprint->WidgetTree, BPGClass);
 
 		for ( const UWidgetAnimation* Animation : Blueprint->Animations )
diff --git a/Engine/Source/Editor/UnrealEd/Private/Animation/DebugSkelMeshComponent.cpp b/Engine/Source/Editor/UnrealEd/Private/Animation/DebugSkelMeshComponent.cpp
--- a/Engine/Source/Editor/UnrealEd/Private/Animation/DebugSkelMeshComponent.cpp
+++ b/Engine/Source/Editor/UnrealEd/Private/Animation/DebugSkelMeshComponent.cpp
@@ -419,10 +419,9 @@ void UDebugSkelMeshComponent::GenSpaceBases(TArray<FTransform>& OutSpaceBases)
 {
 	TArray<FTransform> TempLocalAtoms;
 	TempLocalAtoms.AddUninitialized(OutSpaceBases.Num());
-	TArray<FActiveVertexAnim> TempVertexAnims;
 	FVector TempRootBoneTranslation;
 	FBlendedHeapCurve TempCurve;
-	PerformAnimationEvaluation(SkeletalMesh, AnimScriptInstance, OutSpaceBases, TempLocalAtoms, TempVertexAnims, TempRootBoneTranslation, TempCurve);
+	PerformAnimationEvaluation(SkeletalMesh, AnimScriptInstance, OutSpaceBases, TempLocalAtoms, TempRootBoneTranslation, TempCurve);
 }
 
 void UDebugSkelMeshComponent::RefreshBoneTransforms(FActorComponentTickFunction* TickFunction)
diff --git a/Engine/Source/Editor/UnrealEd/Private/CookOnTheFlyServer.cpp b/Engine/Source/Editor/UnrealEd/Private/CookOnTheFlyServer.cpp
--- a/Engine/Source/Editor/UnrealEd/Private/CookOnTheFlyServer.cpp
+++ b/Engine/Source/Editor/UnrealEd/Private/CookOnTheFlyServer.cpp
@@ -1622,7 +1622,14 @@ uint32 UCookOnTheFlyServer::TickCookOnTheSide( const float TimeSlice, uint32 &Co
 			for (const auto& TargetPlatformName : ToBuild.GetPlatformnames())
 			{
 				const FString SandboxFilename = ConvertToFullSandboxPath(ToBuild.GetFilename().ToString(), true, TargetPlatformName.ToString());
-				check(IFileManager::Get().FileExists(*SandboxFilename) == false);
+				if (IFileManager::Get().FileExists(*SandboxFilename))
+				{
+					// if we find the file this means it was cooked on a previous cook, however source package can't be found now. 
+					// this could be because the source package was deleted or renamed, and we are using iterative cooking
+					// perhaps in this case we should delete it?
+					UE_LOG(LogCook, Warning, TEXT("Found cooked file which shouldn't exist as it failed loading %s"), *SandboxFilename); 
+					IFileManager::Get().Delete(*SandboxFilename);
+				}
 			}
 #endif
 			CookedPackages.Add( FFilePlatformCookedPackage( ToBuild.GetFilename(), TargetPlatformNames, false) );
@@ -1969,6 +1976,14 @@ uint32 UCookOnTheFlyServer::TickCookOnTheSide( const float TimeSlice, uint32 &Co
 					COOK_STAT(FScopedDurationTimer TickCookOnTheSideSaveCookedPackageTimer(DetailedCookStats::TickCookOnTheSideSaveCookedPackageTimeSec));
 					SCOPE_TIMER(SaveCookedPackage);
 					uint32 SaveFlags = SAVE_KeepGUID | (bShouldSaveAsync ? SAVE_Async : SAVE_None) | (IsCookFlagSet(ECookInitializationFlags::Unversioned) ? SAVE_Unversioned : 0);
+
+					bool KeepEditorOnlyPackages = false;
+					// removing editor only packages only works when cooking in commandlet and non iterative cooking
+					// also doesn't work in multiprocess cooking
+					KeepEditorOnlyPackages = !(IsCookByTheBookMode() && !IsCookingInEditor()) || IsCookFlagSet(ECookInitializationFlags::Iterative);
+					KeepEditorOnlyPackages |= IsChildCooker() || (CookByTheBookOptions && CookByTheBookOptions->ChildCookers.Num() > 0);
+					SaveFlags |= KeepEditorOnlyPackages ? SAVE_KeepEditorOnlyCookedPackages : SAVE_None;
+
 					GOutputCookingWarnings = IsCookFlagSet(ECookInitializationFlags::OutputVerboseCookerWarnings);
 					SavePackageResult = SaveCookedPackage(Package, SaveFlags, bWasUpToDate, AllTargetPlatformNames);
 					GOutputCookingWarnings = false;
@@ -3411,14 +3426,19 @@ void UCookOnTheFlyServer::AddFileToCook( TArray<FName>& InOutFilesToCook, const
 { 
 	if (!FPackageName::IsScriptPackage(InFilename))
 	{
+		FName InFilenameName = FName(*InFilename );
+		if ( InFilenameName == NAME_None)
+		{
+			return;
+		}
 #if 0 // randomize cook file order, don't check in enabled...
 		if ( !InOutFilesToCook.Contains(InFilename) )
 		{
 			int Index = FMath::RandRange(0,InOutFilesToCook.Num()-1);
-			InOutFilesToCook.Insert(InFilename, Index );
+			InOutFilesToCook.Insert(InFilenameName, Index );
 		}
 #else
-		InOutFilesToCook.AddUnique(FName(*InFilename));
+		InOutFilesToCook.AddUnique(InFilenameName);
 #endif
 	}
 }
diff --git a/Engine/Source/Editor/UnrealEd/Private/EditorTransaction.cpp b/Engine/Source/Editor/UnrealEd/Private/EditorTransaction.cpp
--- a/Engine/Source/Editor/UnrealEd/Private/EditorTransaction.cpp
+++ b/Engine/Source/Editor/UnrealEd/Private/EditorTransaction.cpp
@@ -721,7 +721,7 @@ bool UTransBuffer::CanRedo( FText* Text )
 
 const FTransaction* UTransBuffer::GetTransaction( int32 QueueIndex ) const
 {
-	if (UndoBuffer.Num() > QueueIndex)
+	if (UndoBuffer.Num() > QueueIndex && QueueIndex != INDEX_NONE)
 	{
 		return &UndoBuffer[QueueIndex];
 	}
diff --git a/Engine/Source/Editor/UnrealEd/Private/FileHelpers.cpp b/Engine/Source/Editor/UnrealEd/Private/FileHelpers.cpp
--- a/Engine/Source/Editor/UnrealEd/Private/FileHelpers.cpp
+++ b/Engine/Source/Editor/UnrealEd/Private/FileHelpers.cpp
@@ -3346,8 +3346,9 @@ void FEditorFileUtils::LoadDefaultMapAtStartup()
 		GConfig->GetString(TEXT("EditorStartup"), TEXT("LastLevel"), EditorStartupMap, GEditorPerProjectIni);
 	}
 	// Default project map.
-	if (EditorStartupMap.IsEmpty()) {
-		EditorStartupMap = GetDefault<UGameMapsSettings>()->EditorStartupMap;
+	if (EditorStartupMap.IsEmpty()) 
+	{
+		EditorStartupMap = GetDefault<UGameMapsSettings>()->EditorStartupMap.GetLongPackageName();
 	}
 	
 	const bool bIncludeReadOnlyRoots = true;
diff --git a/Engine/Source/Editor/UnrealEd/Private/HierarchicalLOD.cpp b/Engine/Source/Editor/UnrealEd/Private/HierarchicalLOD.cpp
--- a/Engine/Source/Editor/UnrealEd/Private/HierarchicalLOD.cpp
+++ b/Engine/Source/Editor/UnrealEd/Private/HierarchicalLOD.cpp
@@ -365,12 +365,12 @@ bool FHierarchicalLODBuilder::ShouldGenerateCluster(AActor* Actor, const bool bP
 	if (Components.Num() > 0)
 	{
 		for (auto& ComponentIter : Components)
-		{
-			
+		{			
 			if (ComponentIter->GetLODParentPrimitive())
 			{
 				auto ParentActor = CastChecked<ALODActor>(ComponentIter->GetLODParentPrimitive()->GetOwner());
-				if (!ParentActor->GetStaticMeshComponent()->StaticMesh || !bPreviewBuild)
+				
+				if (ParentActor && bPreviewBuild)
 				{
 					return false;
 				}
diff --git a/Engine/Source/Editor/UnrealEd/Private/Kismet2/BlueprintEditorUtils.cpp b/Engine/Source/Editor/UnrealEd/Private/Kismet2/BlueprintEditorUtils.cpp
--- a/Engine/Source/Editor/UnrealEd/Private/Kismet2/BlueprintEditorUtils.cpp
+++ b/Engine/Source/Editor/UnrealEd/Private/Kismet2/BlueprintEditorUtils.cpp
@@ -236,7 +236,7 @@ static void RenameVariableReferencesInGraph(UBlueprint* InBlueprint, UClass* InV
 
 		if(UK2Node_ComponentBoundEvent* const ComponentBoundEventNode = Cast<UK2Node_ComponentBoundEvent>(GraphNode))
 		{
-			if(InOldVarName == ComponentBoundEventNode->ComponentPropertyName)
+			if( InOldVarName == ComponentBoundEventNode->ComponentPropertyName && InVariableClass->IsChildOf(InBlueprint->GeneratedClass) )
 			{
 				ComponentBoundEventNode->Modify();
 				ComponentBoundEventNode->ComponentPropertyName = InNewVarName;
@@ -1423,6 +1423,11 @@ static void RemoveStaleFunctions(UBlueprintGeneratedClass* Class, UBlueprint* Bl
 	Blueprint->GeneratedClass->StaticLink(true);
 }
 
+void FBlueprintEditorUtils::ForceLoadMembers(UObject* Object)
+{
+	FRegenerationHelper::ForcedLoadMembers(Object);
+}
+
 UClass* FBlueprintEditorUtils::RegenerateBlueprintClass(UBlueprint* Blueprint, UClass* ClassToRegenerate, UObject* PreviousCDO, TArray<UObject*>& ObjLoaded)
 {
 	bool bRegenerated = false;
diff --git a/Engine/Source/Editor/UnrealEd/Private/ObjectTools.cpp b/Engine/Source/Editor/UnrealEd/Private/ObjectTools.cpp
--- a/Engine/Source/Editor/UnrealEd/Private/ObjectTools.cpp
+++ b/Engine/Source/Editor/UnrealEd/Private/ObjectTools.cpp
@@ -886,14 +886,7 @@ namespace ObjectTools
 				UObject* CurObjOuter = CurObjToConsolidate->GetOuter();
 				UPackage* CurObjPackage = CurObjToConsolidate->GetOutermost();
 				FName CurObjName = CurObjToConsolidate->GetFName();
-
-				// null out the CDO of our current generated class so that DeleteSingleObject does not find it and set it's ClassGeneratedBy to the replacing type.
-				// That would trigger a type mismatch assertion...
 				UBlueprint* BlueprintToConsolidate = Cast<UBlueprint>(CurObjToConsolidate);
-				if (BlueprintToConsolidateTo != nullptr && BlueprintToConsolidate != nullptr && BlueprintToConsolidate->GeneratedClass != nullptr)
-				{
-					BlueprintToConsolidate->GeneratedClass->ClassDefaultObject = nullptr;
-				}
 
 				// Attempt to delete the object that was consolidated
 				if ( DeleteSingleObject( CurObjToConsolidate ) )
diff --git a/Engine/Source/Editor/UnrealEd/Private/PlayLevel.cpp b/Engine/Source/Editor/UnrealEd/Private/PlayLevel.cpp
--- a/Engine/Source/Editor/UnrealEd/Private/PlayLevel.cpp
+++ b/Engine/Source/Editor/UnrealEd/Private/PlayLevel.cpp
@@ -1433,6 +1433,7 @@ private:
 
 void UEditorEngine::HandleStageStarted(const FString& InStage, TWeakPtr<SNotificationItem> NotificationItemPtr)
 {
+	bool bSetNotification = true;
 	FFormatNamedArguments Arguments;
 	FText NotificationText;
 	if (InStage.Contains(TEXT("Cooking")) || InStage.Contains(TEXT("Cook Task")))
@@ -1466,7 +1467,7 @@ void UEditorEngine::HandleStageStarted(const FString& InStage, TWeakPtr<SNotific
 	else if (InStage.Contains(TEXT("Deploy Task")))
 	{
 		Arguments.Add(TEXT("DeviceName"), FText::FromString(PlayUsingLauncherDeviceName));
-		if(PlayUsingLauncherDeviceName.Len() == 0)
+		if (PlayUsingLauncherDeviceName.Len() == 0)
 		{
 			NotificationText = FText::Format(LOCTEXT("LauncherTaskStageNotificationNoDevice", "Deploying Executable and Assets..."), Arguments);
 		}
@@ -1479,17 +1480,24 @@ void UEditorEngine::HandleStageStarted(const FString& InStage, TWeakPtr<SNotific
 	{
 		Arguments.Add(TEXT("GameName"), FText::FromString(FApp::GetGameName()));
 		Arguments.Add(TEXT("DeviceName"), FText::FromString(PlayUsingLauncherDeviceName));
-		if(PlayUsingLauncherDeviceName.Len() == 0)
+		if (PlayUsingLauncherDeviceName.Len() == 0)
 		{
-			NotificationText = FText::Format(LOCTEXT("LauncherTaskStageNotificationNoDevice", "Running {GameName}..."), Arguments);
+			NotificationText = FText::Format(LOCTEXT("LauncherTaskRunNotificationNoDevice", "Running {GameName}..."), Arguments);
 		}
 		else
 		{
-			NotificationText = FText::Format(LOCTEXT("LauncherTaskStageNotification", "Running {GameName} on {DeviceName}..."), Arguments);
+			NotificationText = FText::Format(LOCTEXT("LauncherTaskRunNotification", "Running {GameName} on {DeviceName}..."), Arguments);
 		}
 	}
+	else
+	{
+		bSetNotification = false;
+	}
 
-	NotificationItemPtr.Pin()->SetText(NotificationText);
+	if (bSetNotification)
+	{
+		NotificationItemPtr.Pin()->SetText(NotificationText);
+	}
 }
 
 void UEditorEngine::HandleStageCompleted(const FString& InStage, double StageTime, bool bHasCode, TWeakPtr<SNotificationItem> NotificationItemPtr)
diff --git a/Engine/Source/Editor/UnrealEd/Private/Settings/GameMapsSettingsCustomization.h b/Engine/Source/Editor/UnrealEd/Private/Settings/GameMapsSettingsCustomization.h
--- a/Engine/Source/Editor/UnrealEd/Private/Settings/GameMapsSettingsCustomization.h
+++ b/Engine/Source/Editor/UnrealEd/Private/Settings/GameMapsSettingsCustomization.h
@@ -26,14 +26,6 @@ public:
 
 	virtual void CustomizeDetails( IDetailLayoutBuilder& LayoutBuilder ) override
 	{
-		IDetailCategoryBuilder& DefaultMapsCategory = LayoutBuilder.EditCategory("DefaultMaps");
-		{
-			CustomizeMapSetting(LayoutBuilder, DefaultMapsCategory, "GameDefaultMap");
-			CustomizeMapSetting(LayoutBuilder, DefaultMapsCategory, "ServerDefaultMap");
-			CustomizeMapSetting(LayoutBuilder, DefaultMapsCategory, "EditorStartupMap");
-			CustomizeMapSetting(LayoutBuilder, DefaultMapsCategory, "TransitionMap");
-		}
-
 		// Add extra info around 'Global Default Game Mode'
 		IDetailCategoryBuilder& DefaultModesCategory = LayoutBuilder.EditCategory("DefaultModes");
 		{
diff --git a/Engine/Source/Editor/UnrealEd/Private/Toolkits/SGlobalTabSwitchingDialog.cpp b/Engine/Source/Editor/UnrealEd/Private/Toolkits/SGlobalTabSwitchingDialog.cpp
--- a/Engine/Source/Editor/UnrealEd/Private/Toolkits/SGlobalTabSwitchingDialog.cpp
+++ b/Engine/Source/Editor/UnrealEd/Private/Toolkits/SGlobalTabSwitchingDialog.cpp
@@ -9,6 +9,10 @@
 #include "LevelEditor.h"
 #include "Editor/WorkspaceMenuStructure/Public/WorkspaceMenuStructureModule.h"
 
+#if PLATFORM_MAC
+#include "MacApplication.h"
+#endif
+
 #define LOCTEXT_NAMESPACE "SGlobalTabSwitchingDialog"
 
 //////////////////////////////////////////////////////////////////////////
@@ -193,6 +197,15 @@ protected:
 
 bool SGlobalTabSwitchingDialog::bIsAlreadyOpen = false;
 
+SGlobalTabSwitchingDialog::~SGlobalTabSwitchingDialog()
+{
+	bIsAlreadyOpen = false;
+
+#if PLATFORM_MAC
+	MacApplication->SetIsRightClickEmulationEnabled(true);
+#endif
+}
+
 SGlobalTabSwitchingDialog::FTabListItemPtr SGlobalTabSwitchingDialog::GetMainTabListSelectedItem() const
 {
 	TArray<FTabListItemPtr> SelectedItems = MainTabsListWidget->GetSelectedItems();
@@ -349,6 +362,11 @@ void SGlobalTabSwitchingDialog::Construct(const FArguments& InArgs, FVector2D In
 	check(!bIsAlreadyOpen);
 	bIsAlreadyOpen = true;
 
+#if PLATFORM_MAC
+	// On Mac we emulate right click with ctrl+left click. This needs to be disabled for tab navigator, so that users can click on its widgets while they keep the ctrl key pressed
+	MacApplication->SetIsRightClickEmulationEnabled(false);
+#endif
+
 	TriggerChord = InTriggerChord;
 
 	AssetThumbnailPool = MakeShareable(new FAssetThumbnailPool(128));
diff --git a/Engine/Source/Editor/UnrealEd/Private/Toolkits/SGlobalTabSwitchingDialog.h b/Engine/Source/Editor/UnrealEd/Private/Toolkits/SGlobalTabSwitchingDialog.h
--- a/Engine/Source/Editor/UnrealEd/Private/Toolkits/SGlobalTabSwitchingDialog.h
+++ b/Engine/Source/Editor/UnrealEd/Private/Toolkits/SGlobalTabSwitchingDialog.h
@@ -14,10 +14,7 @@ public:
 public:
 	void Construct(const FArguments& InArgs, FVector2D InSize, FInputChord InTriggerChord);
 
-	~SGlobalTabSwitchingDialog()
-	{
-		bIsAlreadyOpen = false;
-	}
+	~SGlobalTabSwitchingDialog();
 
 	// SWidget interface
 	virtual FReply OnPreviewKeyDown(const FGeometry& MyGeometry, const FKeyEvent& InKeyEvent) override;
diff --git a/Engine/Source/Editor/UnrealEd/Private/UnrealEdMisc.cpp b/Engine/Source/Editor/UnrealEd/Private/UnrealEdMisc.cpp
--- a/Engine/Source/Editor/UnrealEd/Private/UnrealEdMisc.cpp
+++ b/Engine/Source/Editor/UnrealEd/Private/UnrealEdMisc.cpp
@@ -292,7 +292,7 @@ void FUnrealEdMisc::OnInit()
 		{
 			if (!bMapLoaded && GEditor)
 			{
-				const FString& StartupMap = GetDefault<UGameMapsSettings>()->EditorStartupMap;
+				const FString StartupMap = GetDefault<UGameMapsSettings>()->EditorStartupMap.ToString();
 
 				if ((StartupMap.Len() > 0) && (GetDefault<UEditorLoadingSavingSettings>()->LoadLevelAtStartup != ELoadLevelAtStartup::None))
 				{
diff --git a/Engine/Source/Editor/UnrealEd/Public/Kismet2/BlueprintEditorUtils.h b/Engine/Source/Editor/UnrealEd/Public/Kismet2/BlueprintEditorUtils.h
--- a/Engine/Source/Editor/UnrealEd/Public/Kismet2/BlueprintEditorUtils.h
+++ b/Engine/Source/Editor/UnrealEd/Public/Kismet2/BlueprintEditorUtils.h
@@ -150,6 +150,11 @@ public:
 	 */
 	static void PatchNewCDOIntoLinker(UObject* CDO, FLinkerLoad* Linker, int32 ExportIndex, TArray<UObject*>& ObjLoaded);
 
+	/**
+	 * Helper for loading dependent data needed for compile, would not be needed if compile on load ran after serialization
+	 */
+	static void ForceLoadMembers(UObject* Object);
+
 	/**
 	 * Regenerates the class at class load time, and refreshes the blueprint
 	 */
diff --git a/Engine/Source/Programs/AutomationTool/AutomationUtils/DeploymentContext.cs b/Engine/Source/Programs/AutomationTool/AutomationUtils/DeploymentContext.cs
--- a/Engine/Source/Programs/AutomationTool/AutomationUtils/DeploymentContext.cs
+++ b/Engine/Source/Programs/AutomationTool/AutomationUtils/DeploymentContext.cs
@@ -468,7 +468,9 @@ public class DeploymentContext //: ProjectParams
                     bool OtherPlatform = false;
                     foreach (UnrealTargetPlatform Plat in Enum.GetValues(typeof(UnrealTargetPlatform)))
                     {
-                        if (Plat != StageTargetPlatform.PlatformType && Plat != UnrealTargetPlatform.Unknown)
+						bool bMatchesIniPlatform = (AllFile.EndsWith(".ini") && Plat == StageTargetPlatform.IniPlatformType); // filter ini files for the ini file platform
+						bool bMatchesTargetPlatform = (Plat == StageTargetPlatform.PlatformType || Plat == UnrealTargetPlatform.Unknown); // filter platform files for the target platform
+						if (!bMatchesIniPlatform && !bMatchesTargetPlatform)
                         {
                             var Search = FileToCopy;
                             if (Search.StartsWith(LocalRoot, StringComparison.InvariantCultureIgnoreCase))
diff --git a/Engine/Source/Programs/AutomationTool/AutomationUtils/Platform.cs b/Engine/Source/Programs/AutomationTool/AutomationUtils/Platform.cs
--- a/Engine/Source/Programs/AutomationTool/AutomationUtils/Platform.cs
+++ b/Engine/Source/Programs/AutomationTool/AutomationUtils/Platform.cs
@@ -92,10 +92,12 @@ namespace AutomationTool
 		#endregion
 
 		protected UnrealTargetPlatform TargetPlatformType = UnrealTargetPlatform.Unknown;
+		protected UnrealTargetPlatform TargetIniPlatformType = UnrealTargetPlatform.Unknown;
 
 		public Platform(UnrealTargetPlatform PlatformType)
 		{
 			TargetPlatformType = PlatformType;
+			TargetIniPlatformType = PlatformType;
 		}
 
 		/// <summary>
@@ -405,6 +407,14 @@ namespace AutomationTool
 			get { return TargetPlatformType; }
 		}
 
+		/// <summary>
+		/// UnrealTargetPlatform type for this platform.
+		/// </summary>
+		public UnrealTargetPlatform IniPlatformType
+		{
+			get { return TargetIniPlatformType; }
+		}
+
 		/// <summary>
 		/// True if this platform is supported.
 		/// </summary>
diff --git a/Engine/Source/Programs/AutomationTool/AutomationUtils/ProcessUtils.cs b/Engine/Source/Programs/AutomationTool/AutomationUtils/ProcessUtils.cs
--- a/Engine/Source/Programs/AutomationTool/AutomationUtils/ProcessUtils.cs
+++ b/Engine/Source/Programs/AutomationTool/AutomationUtils/ProcessUtils.cs
@@ -552,13 +552,24 @@ namespace AutomationTool
             Process[] AllProcs = Process.GetProcesses();
             foreach (Process KillCandidate in AllProcs)
             {
-                HashSet<int> VisitedPids = new HashSet<int>();
-                if (ProcessManager.CanBeKilled(KillCandidate.ProcessName) && IsOurDescendant(ProcessToCheck, KillCandidate.Id, VisitedPids))
-                {
-					CommandUtils.LogLog("Descendant pid={0}, name={1}, filename={2}", KillCandidate.Id, KillCandidate.ProcessName,
-						KillCandidate.MainModule != null ? KillCandidate.MainModule.FileName : "unknown");
-                    return true;
-                }
+				// Silently skip InvalidOperationExceptions here, because it depends on the process still running. It may have terminated.
+				string ProcessName;
+				try
+				{
+					ProcessName = KillCandidate.ProcessName;
+				}
+				catch(InvalidOperationException)
+				{
+					continue;
+				}
+
+				// Check if it's still running
+				HashSet<int> VisitedPids = new HashSet<int>();
+				if (ProcessManager.CanBeKilled(ProcessName) && IsOurDescendant(ProcessToCheck, KillCandidate.Id, VisitedPids))
+				{
+					CommandUtils.LogLog("Descendant pid={0}, name={1}", KillCandidate.Id, ProcessName);
+					return true;
+				}
             }
             return false;
         }
diff --git a/Engine/Source/Programs/AutomationTool/GUBP/LegacyNodes.cs b/Engine/Source/Programs/AutomationTool/GUBP/LegacyNodes.cs
--- a/Engine/Source/Programs/AutomationTool/GUBP/LegacyNodes.cs
+++ b/Engine/Source/Programs/AutomationTool/GUBP/LegacyNodes.cs
@@ -351,11 +351,13 @@ partial class GUBP
     public class ToolsForCompileNode : CompileNode
     {
 		bool bHasLauncherParam;
+		bool bNoInstalledEngine;
 
         public ToolsForCompileNode(GUBP.GUBPBranchConfig InBranchConfig, UnrealTargetPlatform InHostPlatform, bool bInHasLauncherParam)
             : base(InBranchConfig, InHostPlatform, false)
         {
 			bHasLauncherParam = bInHasLauncherParam;
+			bNoInstalledEngine = InBranchConfig.BranchOptions.bNoInstalledEngine;
 
 			if (InHostPlatform != UnrealTargetPlatform.Win64)
             {
@@ -414,7 +416,7 @@ partial class GUBP
             }
             string AddArgs = "-CopyAppBundleBackToDevice";
 
-            Agenda.AddTargets(new string[] { "UnrealHeaderTool" }, HostPlatform, UnrealTargetConfiguration.Development, InAddArgs: AddArgs);
+            Agenda.AddTargets(new string[] { "UnrealHeaderTool" }, HostPlatform, UnrealTargetConfiguration.Development, InAddArgs: AddArgs + (bNoInstalledEngine? "" : " -precompile"));
             return Agenda;
         }
         public override void PostBuild(GUBP bp, UE4Build UE4Build)
@@ -2026,9 +2028,9 @@ partial class GUBP
             {
                 AddDependency(GamePlatformCookedAndCompiledNode.StaticGetFullName(HostPlatform, GameProj, Plat));
             }
-        }
+		}
 
-        public static string StaticGetFullName(BranchInfo.BranchUProject InGameProj, UnrealTargetPlatform InHostPlatform, List<UnrealTargetPlatform> InClientTargetPlatforms = null, List<UnrealTargetConfiguration> InClientConfigs = null, List<UnrealTargetPlatform> InServerTargetPlatforms = null, List<UnrealTargetConfiguration> InServerConfigs = null, bool InClientNotGame = false)
+		public static string StaticGetBaseName(BranchInfo.BranchUProject InGameProj, UnrealTargetPlatform InHostPlatform, List<UnrealTargetPlatform> InClientTargetPlatforms = null, List<UnrealTargetConfiguration> InClientConfigs = null, List<UnrealTargetPlatform> InServerTargetPlatforms = null, List<UnrealTargetConfiguration> InServerConfigs = null, bool InClientNotGame = false)
         {
             string Infix = "";
             if (InClientNotGame)
@@ -2061,9 +2063,14 @@ partial class GUBP
             {
                 Infix += "_Serv_" + InServerConfigs[0].ToString();
             }
-            return InGameProj.GameName + Infix + "_MakeBuild" + HostPlatformNode.StaticGetHostPlatformSuffix(InHostPlatform);
+			return InGameProj.GameName + Infix;
         }
-        public override string GetFullName()
+		public static string StaticGetFullName(BranchInfo.BranchUProject InGameProj, UnrealTargetPlatform InHostPlatform, List<UnrealTargetPlatform> InClientTargetPlatforms = null, List<UnrealTargetConfiguration> InClientConfigs = null, List<UnrealTargetPlatform> InServerTargetPlatforms = null, List<UnrealTargetConfiguration> InServerConfigs = null, bool InClientNotGame = false)
+		{
+			return StaticGetBaseName(InGameProj, InHostPlatform, InClientTargetPlatforms, InClientConfigs, InServerTargetPlatforms, InServerConfigs, InClientNotGame) + "_MakeBuild" + HostPlatformNode.StaticGetHostPlatformSuffix(InHostPlatform);
+		}
+
+		public override string GetFullName()
         {
             return StaticGetFullName(GameProj, HostPlatform, ClientTargetPlatforms, ClientConfigs, ServerTargetPlatforms, ServerConfigs, ClientNotGame);
         }
@@ -2109,10 +2116,9 @@ partial class GUBP
             {
                 BuildShareName = "UE4";
             }
-            string BaseDir = CommandUtils.CombinePaths(CommandUtils.RootBuildStorageDirectory(), BuildShareName);
-            string NodeName = StaticGetFullName(InGameProj, InHostPlatform, InClientTargetPlatforms, InClientConfigs, InServerTargetPlatforms, InServerConfigs, InClientNotGame);
-            string Inner = P4Env.BuildRootEscaped + "-" + BranchConfig.JobInfo.BuildName;
-			string ArchiveDirectory = CombinePaths(BaseDir, NodeName, Inner);
+            string BaseDir = CommandUtils.CombinePaths(CommandUtils.RootBuildStorageDirectory(), BuildShareName, "PackagedBuilds", P4Env.BuildRootEscaped, BranchConfig.JobInfo.BuildName);
+			string NodeName = StaticGetBaseName(InGameProj, InHostPlatform, InClientTargetPlatforms, InClientConfigs, InServerTargetPlatforms, InServerConfigs, InClientNotGame);
+			string ArchiveDirectory = CombinePaths(BaseDir, NodeName);
             return ArchiveDirectory;
         }
         public override void DoBuild(GUBP bp)
diff --git a/Engine/Source/Programs/AutomationTool/Scripts/BuildPluginCommand.Automation.cs b/Engine/Source/Programs/AutomationTool/Scripts/BuildPluginCommand.Automation.cs
--- a/Engine/Source/Programs/AutomationTool/Scripts/BuildPluginCommand.Automation.cs
+++ b/Engine/Source/Programs/AutomationTool/Scripts/BuildPluginCommand.Automation.cs
@@ -18,22 +18,43 @@ class BuildPlugin : BuildCommand
 	public override void ExecuteBuild()
 	{
 		// Get the plugin filename
-		string PluginFileName = ParseParamValue("Plugin");
-		if(PluginFileName == null)
+		string PluginParam = ParseParamValue("Plugin");
+		if(PluginParam == null)
 		{
 			throw new AutomationException("Plugin file name was not specified via the -plugin argument");
 		}
 
 		// Read the plugin
-		PluginDescriptor Plugin = PluginDescriptor.FromFile(new FileReference(PluginFileName));
+		FileReference PluginFile = new FileReference(PluginParam);
+		DirectoryReference PluginDirectory = PluginFile.Directory;
+		PluginDescriptor Plugin = PluginDescriptor.FromFile(PluginFile);
 
 		// Clean the intermediate build directory
-		string IntermediateBuildDirectory = Path.Combine(Path.GetDirectoryName(PluginFileName), "Intermediate", "Build");
-		if(CommandUtils.DirectoryExists(IntermediateBuildDirectory))
+		DirectoryReference IntermediateBuildDirectory = DirectoryReference.Combine(PluginDirectory, "Intermediate", "Build");
+		if(CommandUtils.DirectoryExists(IntermediateBuildDirectory.FullName))
 		{
-			CommandUtils.DeleteDirectory(IntermediateBuildDirectory);
+			CommandUtils.DeleteDirectory(IntermediateBuildDirectory.FullName);
 		}
 
+		// Create a host project for the plugin. For script generator plugins, we need to have UHT be able to load it - and that can only happen if it's enabled in a project.
+		DirectoryReference HostProjectDirectory = DirectoryReference.Combine(new DirectoryReference(CommandUtils.CmdEnv.LocalRoot), "HostProject");
+		if (CommandUtils.DirectoryExists(HostProjectDirectory.FullName))
+		{
+			CommandUtils.DeleteDirectory(HostProjectDirectory.FullName);
+		}
+
+		DirectoryReference HostProjectPluginDirectory = DirectoryReference.Combine(HostProjectDirectory, "Plugins", PluginFile.GetFileNameWithoutExtension());
+
+		string[] CopyPluginFiles = Directory.EnumerateFiles(PluginDirectory.FullName, "*", SearchOption.AllDirectories).ToArray();
+		foreach (string CopyPluginFile in CopyPluginFiles)
+		{
+			CommandUtils.CopyFile(CopyPluginFile, CommandUtils.MakeRerootedFilePath(CopyPluginFile, PluginDirectory.FullName, HostProjectPluginDirectory.FullName));
+		}
+
+		FileReference HostProjectPluginFile = FileReference.Combine(HostProjectPluginDirectory, PluginFile.GetFileName());
+		FileReference HostProjectFile = FileReference.Combine(HostProjectDirectory, "HostProject.uproject");
+		File.WriteAllText(HostProjectFile.FullName, "{ \"FileVersion\": 3, \"Plugins\": [ { \"Name\": \"" + PluginFile.GetFileNameWithoutExtension() + "\", \"Enabled\": true } ] }");
+
 		// Get any additional arguments from the commandline
 		string AdditionalArgs = "";
 
@@ -43,7 +64,11 @@ class BuildPlugin : BuildCommand
 		UnrealTargetPlatform HostPlatform = BuildHostPlatform.Current.Platform;
 		if(!ParseParam("NoHostPlatform"))
 		{
-			AddPluginToAgenda(Agenda, PluginFileName, Plugin, "UE4Editor", TargetRules.TargetType.Editor, HostPlatform, UnrealTargetConfiguration.Development, ReceiptFileNames, AdditionalArgs);
+			if (Plugin.bCanBeUsedWithUnrealHeaderTool)
+			{
+				AddPluginToAgenda(Agenda, PluginFile, Plugin, null, "UnrealHeaderTool", TargetRules.TargetType.Program, HostPlatform, UnrealTargetConfiguration.Development, ReceiptFileNames, String.Format("{0} -plugin {1}", AdditionalArgs, CommandUtils.MakePathSafeToUseWithCommandLine(HostProjectPluginFile.FullName)));
+			}
+			AddPluginToAgenda(Agenda, PluginFile, Plugin, HostProjectFile, "UE4Editor", TargetRules.TargetType.Editor, HostPlatform, UnrealTargetConfiguration.Development, ReceiptFileNames, AdditionalArgs);
 		}
 
 		// Add the game targets
@@ -52,25 +77,25 @@ class BuildPlugin : BuildCommand
 		{
 			if(Rocket.RocketBuild.IsCodeTargetPlatform(HostPlatform, TargetPlatform))
 			{
-				AddPluginToAgenda(Agenda, PluginFileName, Plugin, "UE4Game", TargetRules.TargetType.Game, TargetPlatform, UnrealTargetConfiguration.Development, ReceiptFileNames, AdditionalArgs);
-				AddPluginToAgenda(Agenda, PluginFileName, Plugin, "UE4Game", TargetRules.TargetType.Game, TargetPlatform, UnrealTargetConfiguration.Shipping, ReceiptFileNames, AdditionalArgs);
+				AddPluginToAgenda(Agenda, PluginFile, Plugin, HostProjectFile, "UE4Game", TargetRules.TargetType.Game, TargetPlatform, UnrealTargetConfiguration.Development, ReceiptFileNames, AdditionalArgs);
+				AddPluginToAgenda(Agenda, PluginFile, Plugin, HostProjectFile, "UE4Game", TargetRules.TargetType.Game, TargetPlatform, UnrealTargetConfiguration.Shipping, ReceiptFileNames, AdditionalArgs);
 			}
 		}
 
 		// Build it
 		UE4Build Build = new UE4Build(this);
-		Build.Build(Agenda, InDeleteBuildProducts: true, InUpdateVersionFiles: false);
+		Build.Build(Agenda, InDeleteBuildProducts: false, InUpdateVersionFiles: false);
 
 		// Package the plugin to the output folder
 		string PackageDirectory = ParseParamValue("Package");
 		if(PackageDirectory != null)
 		{
-			List<BuildProduct> BuildProducts = GetBuildProductsFromReceipts(ReceiptFileNames);
-			PackagePlugin(PluginFileName, BuildProducts, PackageDirectory);
+			List<BuildProduct> BuildProducts = GetBuildProductsFromReceipts(UnrealBuildTool.UnrealBuildTool.EngineDirectory, HostProjectDirectory, ReceiptFileNames);
+			PackagePlugin(HostProjectPluginFile, BuildProducts, PackageDirectory);
 		}
 	}
 
-	static void AddPluginToAgenda(UE4Build.BuildAgenda Agenda, string PluginFileName, PluginDescriptor Plugin, string TargetName, TargetRules.TargetType TargetType, UnrealTargetPlatform Platform, UnrealTargetConfiguration Configuration, List<string> ReceiptFileNames, string InAdditionalArgs)
+	static void AddPluginToAgenda(UE4Build.BuildAgenda Agenda, FileReference PluginFile, PluginDescriptor Plugin, FileReference HostProjectFile, string TargetName, TargetRules.TargetType TargetType, UnrealTargetPlatform Platform, UnrealTargetConfiguration Configuration, List<string> ReceiptFileNames, string InAdditionalArgs)
 	{
 		// Find a list of modules that need to be built for this plugin
 		List<string> ModuleNames = new List<string>();
@@ -87,13 +112,13 @@ class BuildPlugin : BuildCommand
 		// Add these modules to the build agenda
 		if(ModuleNames.Count > 0)
 		{
-			string Arguments = String.Format("-plugin {0}", CommandUtils.MakePathSafeToUseWithCommandLine(PluginFileName));
+			string Arguments = "";// String.Format("-plugin {0}", CommandUtils.MakePathSafeToUseWithCommandLine(PluginFile.FullName));
 			foreach(string ModuleName in ModuleNames)
 			{
 				Arguments += String.Format(" -module {0}", ModuleName);
 			}
 
-			string ReceiptFileName = TargetReceipt.GetDefaultPath(Path.GetDirectoryName(PluginFileName), TargetName, Platform, Configuration, "");
+			string ReceiptFileName = TargetReceipt.GetDefaultPath(Path.GetDirectoryName(PluginFile.FullName), TargetName, Platform, Configuration, "");
 			Arguments += String.Format(" -receipt {0}", CommandUtils.MakePathSafeToUseWithCommandLine(ReceiptFileName));
 			ReceiptFileNames.Add(ReceiptFileName);
 			
@@ -102,11 +127,11 @@ class BuildPlugin : BuildCommand
 				Arguments += InAdditionalArgs;
 			}
 
-			Agenda.AddTarget(TargetName, Platform, Configuration, InAddArgs: Arguments);
+			Agenda.AddTarget(TargetName, Platform, Configuration, InUprojectPath: HostProjectFile, InAddArgs: Arguments);
 		}
 	}
 
-	static List<BuildProduct> GetBuildProductsFromReceipts(List<string> ReceiptFileNames)
+	static List<BuildProduct> GetBuildProductsFromReceipts(DirectoryReference EngineDir, DirectoryReference ProjectDir, List<string> ReceiptFileNames)
 	{
 		List<BuildProduct> BuildProducts = new List<BuildProduct>();
 		foreach(string ReceiptFileName in ReceiptFileNames)
@@ -116,32 +141,33 @@ class BuildPlugin : BuildCommand
 			{
 				throw new AutomationException("Missing or invalid target receipt ({0})", ReceiptFileName);
 			}
+			Receipt.ExpandPathVariables(EngineDir, ProjectDir);
 			BuildProducts.AddRange(Receipt.BuildProducts);
 		}
 		return BuildProducts;
 	}
 
-	static void PackagePlugin(string PluginFileName, List<BuildProduct> BuildProducts, string PackageDirectory)
+	static void PackagePlugin(FileReference PluginFile, List<BuildProduct> BuildProducts, string PackageDirectory)
 	{
 		// Clear the output directory
 		CommandUtils.DeleteDirectoryContents(PackageDirectory);
 
 		// Copy all the files to the output directory
-		List<string> MatchingFileNames = FilterPluginFiles(PluginFileName, BuildProducts);
+		List<string> MatchingFileNames = FilterPluginFiles(PluginFile.FullName, BuildProducts);
 		foreach(string MatchingFileName in MatchingFileNames)
 		{
-			string SourceFileName = Path.Combine(Path.GetDirectoryName(PluginFileName), MatchingFileName);
+			string SourceFileName = Path.Combine(Path.GetDirectoryName(PluginFile.FullName), MatchingFileName);
 			string TargetFileName = Path.Combine(PackageDirectory, MatchingFileName);
 			CommandUtils.CopyFile(SourceFileName, TargetFileName);
 			CommandUtils.SetFileAttributes(TargetFileName, ReadOnly: false);
 		}
 
 		// Get the output plugin filename
-		string TargetPluginFileName = CommandUtils.MakeRerootedFilePath(Path.GetFullPath(PluginFileName), Path.GetDirectoryName(Path.GetFullPath(PluginFileName)), PackageDirectory);
-		PluginDescriptor NewDescriptor = PluginDescriptor.FromFile(new FileReference(TargetPluginFileName));
+		FileReference TargetPluginFile = FileReference.Combine(new DirectoryReference(PackageDirectory), PluginFile.GetFileName());
+		PluginDescriptor NewDescriptor = PluginDescriptor.FromFile(TargetPluginFile);
 		NewDescriptor.bEnabledByDefault = true;
 		NewDescriptor.bInstalled = true;
-		NewDescriptor.Save(TargetPluginFileName);
+		NewDescriptor.Save(TargetPluginFile.FullName);
 	}
 
 	static List<string> FilterPluginFiles(string PluginFileName, List<BuildProduct> BuildProducts)
diff --git a/Engine/Source/Programs/AutomationTool/Scripts/CopyBuildToStagingDirectory.Automation.cs b/Engine/Source/Programs/AutomationTool/Scripts/CopyBuildToStagingDirectory.Automation.cs
--- a/Engine/Source/Programs/AutomationTool/Scripts/CopyBuildToStagingDirectory.Automation.cs
+++ b/Engine/Source/Programs/AutomationTool/Scripts/CopyBuildToStagingDirectory.Automation.cs
@@ -264,16 +264,23 @@ public partial class Project : CommandUtils
             SC.StageFiles(StagedFileType.NonUFS, CombinePaths(SC.ProjectRoot, "Plugins", DLCName, "Binaries"), "UE4Server-*.dll", true, null, null, true);
 
             // Put all of the cooked dir into the staged dir
-            // Dedicated server cook doesn't save shaders so no Engine dir is created
-            if ((!SC.DedicatedServer) && (!Params.DLCIncludeEngineContent))
+            string PlatformCookDir = CombinePaths(SC.ProjectRoot, "Plugins", DLCName, "Saved", "Cooked", SC.CookPlatform);
+            string[] ExcludeWildCards = {"AssetRegistry.bin"};
+
+            // Stage any loose files in the root folder
+            SC.StageFiles(StagedFileType.UFS, PlatformCookDir, "*", false, ExcludeWildCards, "", true, !Params.UsePak(SC.StageTargetPlatform));
+
+            // Stage each sub directory separately so that we can skip Engine if need be
+            string[] SubDirs = CommandUtils.FindDirectories(true, "*", false, new string[] { PlatformCookDir });
+            foreach (string SubDir in SubDirs)
             {
-                if (Directory.Exists(CombinePaths(SC.ProjectRoot, "Plugins", DLCName, "Saved", "Cooked", SC.CookPlatform, "Engine")))
+                // Dedicated server cook doesn't save shaders so no Engine dir is created
+                if ((!SC.DedicatedServer) && (!Params.DLCIncludeEngineContent) && CommandUtils.GetLastDirectoryName(SubDir).Equals("Engine", StringComparison.InvariantCultureIgnoreCase))
                 {
-                    SC.StageFiles(StagedFileType.UFS, CombinePaths(SC.ProjectRoot, "Plugins", DLCName, "Saved", "Cooked", SC.CookPlatform), "*", true, new[] { CommandUtils.CombinePaths("Engine", "*") }, "", true, !Params.UsePak(SC.StageTargetPlatform));
+                    continue;
                 }
+                SC.StageFiles(StagedFileType.UFS, SubDir, "*", true, ExcludeWildCards, "", true, !Params.UsePak(SC.StageTargetPlatform));
             }
-            
-            SC.StageFiles(StagedFileType.UFS, CombinePaths(SC.ProjectRoot, "Plugins", DLCName, "Saved", "Cooked", SC.CookPlatform), "*", true, new[] { "AssetRegistry.bin" }, "", true, !Params.UsePak(SC.StageTargetPlatform));
 
             return;
         }
@@ -308,7 +315,7 @@ public partial class Project : CommandUtils
 		}
 		SC.StageFiles(StagedFileType.NonUFS, GetIntermediateCommandlineDir(SC), CommandLineFile, false, null, "", true, false);
 
-        ConfigCacheIni PlatformGameConfig = ConfigCacheIni.CreateConfigCacheIni(SC.StageTargetPlatform.PlatformType, "Game", new DirectoryReference(CommandUtils.GetDirectoryName(Params.RawProjectPath.FullName)));
+        ConfigCacheIni PlatformGameConfig = ConfigCacheIni.CreateConfigCacheIni(SC.StageTargetPlatform.IniPlatformType, "Game", new DirectoryReference(CommandUtils.GetDirectoryName(Params.RawProjectPath.FullName)));
         var ProjectContentRoot = CombinePaths(SC.ProjectRoot, "Content");
         var StageContentRoot = CombinePaths(SC.RelativeProjectRootForStage, "Content");
         
@@ -460,18 +467,19 @@ public partial class Project : CommandUtils
                 if (PlatformGameConfig.GetArray("/Script/UnrealEd.ProjectPackagingSettings", "DirectoriesToAlwaysStageAsNonUFS", out ExtraNonUFSDirs))
                 {
                     // Each string has the format '(Path="TheDirToStage")'
+					// NonUFS files are never in pak files and should always be remapped
                     foreach (var PathStr in ExtraNonUFSDirs)
                     {
 						var PathParts = PathStr.Split('"');
 						if (PathParts.Length == 3)
 						{
 							var RelativePath = PathParts[1];
-							SC.StageFiles(StagedFileType.NonUFS, CombinePaths(ProjectContentRoot, RelativePath), "*", true, null, CombinePaths(StageContentRoot, RelativePath), true, !Params.UsePak(SC.StageTargetPlatform));
+							SC.StageFiles(StagedFileType.NonUFS, CombinePaths(ProjectContentRoot, RelativePath), "*", true, null, CombinePaths(StageContentRoot, RelativePath), true, true);
 						}
 						else if (PathParts.Length == 1)
 						{
 							var RelativePath = PathParts[0];
-							SC.StageFiles(StagedFileType.NonUFS, CombinePaths(ProjectContentRoot, RelativePath), "*", true, null, CombinePaths(StageContentRoot, RelativePath), true, !Params.UsePak(SC.StageTargetPlatform));
+							SC.StageFiles(StagedFileType.NonUFS, CombinePaths(ProjectContentRoot, RelativePath), "*", true, null, CombinePaths(StageContentRoot, RelativePath), true, true);
 						}
                     }
                 }
diff --git a/Engine/Source/Programs/AutomationTool/TVOS/TVOSPlatform.Automation.cs b/Engine/Source/Programs/AutomationTool/TVOS/TVOSPlatform.Automation.cs
--- a/Engine/Source/Programs/AutomationTool/TVOS/TVOSPlatform.Automation.cs
+++ b/Engine/Source/Programs/AutomationTool/TVOS/TVOSPlatform.Automation.cs
@@ -13,6 +13,7 @@ public class TVOSPlatform : IOSPlatform
 	public TVOSPlatform()
 		:base(UnrealTargetPlatform.TVOS)
 	{
+		TargetIniPlatformType = UnrealTargetPlatform.IOS;
 	}
 
 	public override UnrealBuildTool.UEDeployIOS GetDeployHandler()
diff --git a/Engine/Source/Programs/CrashReporter/CrashReportClient/CrashReportClient.Build.cs b/Engine/Source/Programs/CrashReporter/CrashReportClient/CrashReportClient.Build.cs
--- a/Engine/Source/Programs/CrashReporter/CrashReportClient/CrashReportClient.Build.cs
+++ b/Engine/Source/Programs/CrashReporter/CrashReportClient/CrashReportClient.Build.cs
@@ -27,6 +27,7 @@ public class CrashReportClient : ModuleRules
 				"XmlParser",
 				"Analytics",
 				"AnalyticsET",
+				"DesktopPlatform",
 				"QoSReporter",
 			}
 			);
