UAT: Fix code which deletes packaged builds after 4 days. Currently causing warnings in //UE4/Main due to adjusted directory structure after merge down from Release-4.11. Directory structure is now PackagedBuilds/<Branch>/<CL>/<Game Name>.

diff --git a/Engine/Source/Programs/AutomationTool/AutomationUtils/CommandUtils.cs b/Engine/Source/Programs/AutomationTool/AutomationUtils/CommandUtils.cs
--- a/Engine/Source/Programs/AutomationTool/AutomationUtils/CommandUtils.cs
+++ b/Engine/Source/Programs/AutomationTool/AutomationUtils/CommandUtils.cs
@@ -1978,83 +1978,38 @@ namespace AutomationTool
 			return false;
 		}
 
-        public static void CleanFormalBuilds(string DirectoryForThisBuild, string CLString = "", int MaximumDaysToKeepTempStorage = 4)
+		public static void CleanFormalBuilds(string ParentDir, string SearchPattern, int MaximumDaysToKeepTempStorage = 4)
         {
-            if (CLString == "" && (!IsBuildMachine || !DirectoryForThisBuild.StartsWith(RootBuildStorageDirectory()) || !P4Enabled))
+            if (!IsBuildMachine || !ParentDir.StartsWith(RootBuildStorageDirectory()))
             {
                 return;
             }
             try
             {
-                if (P4Enabled && CLString == "")
-                {
-                    CLString = P4Env.ChangelistString;
-                }
-                string ParentDir = Path.GetDirectoryName(CombinePaths(DirectoryForThisBuild));
-                if (!DirectoryExists_NoExceptions(ParentDir))
-                {
-                    throw new AutomationException("Not cleaning formal builds, because the parent directory {0} does not exist.", ParentDir);
-                }
-                string MyDir = Path.GetFileName(CombinePaths(DirectoryForThisBuild));
-                int CLStart = MyDir.IndexOf(CLString);
-                if (CLStart < 0)
-                {
-                    throw new AutomationException("Not cleaning formal builds, because the directory {0} does not contain the CL {1}.", DirectoryForThisBuild, CLString);
-                }
-                string StartString = MyDir.Substring(0, CLStart);
-                string EndString = MyDir.Substring(CLStart + CLString.Length);
-
-
                 DirectoryInfo DirInfo = new DirectoryInfo(ParentDir);
-                var TopLevelDirs = DirInfo.GetDirectories();
-				Log("Looking for directories to delete in {0}   {1} dirs", ParentDir, TopLevelDirs.Length);
-                foreach (var TopLevelDir in TopLevelDirs)
+				Log("Looking for directories to delete in {0}", ParentDir);
+                foreach (DirectoryInfo ThisDirInfo in DirInfo.EnumerateDirectories(SearchPattern))
                 {
-                    if (DirectoryExists_NoExceptions(TopLevelDir.FullName))
+					double AgeDays = (DateTime.UtcNow - ThisDirInfo.CreationTimeUtc).TotalDays;
+					if (AgeDays > MaximumDaysToKeepTempStorage)
                     {
-                        var JustDir = Path.GetFileName(CombinePaths(TopLevelDir.FullName));
-                        if (JustDir.StartsWith(StartString, StringComparison.InvariantCultureIgnoreCase) && (String.IsNullOrEmpty(EndString) || JustDir.EndsWith(EndString, StringComparison.InvariantCultureIgnoreCase)))
-                        {
-                            string CLPart = JustDir.Substring(StartString.Length, JustDir.Length - StartString.Length - EndString.Length);
-                            if (CLPart.Contains("-PF-") || // always delete preflights
-                                (!CLPart.Contains("-") && !CLPart.Contains("+"))) // never delete anything else that is weird, this is probably another branch
-                            {
-                                DirectoryInfo ThisDirInfo = new DirectoryInfo(TopLevelDir.FullName);
-                                bool bOld = false;
-
-                                if ((DateTime.UtcNow - ThisDirInfo.CreationTimeUtc).TotalDays > MaximumDaysToKeepTempStorage)
-                                {
-                                    bOld = true;
-                                }
-                                if (bOld)
-                                {
-                                    LogVerbose("Deleting temp storage directory {0}, because it is more than {1} days old.", TopLevelDir.FullName, MaximumDaysToKeepTempStorage);
-                                    DeleteDirectory_NoExceptions(true, TopLevelDir.FullName);
-                                }
-                                else
-                                {
-                                    LogVerbose("Not Deleteing temp storage directory {0}, because it is less than {1} days old.", TopLevelDir.FullName, MaximumDaysToKeepTempStorage);
-                                }
-                            }
-                            else
-                            {
-                                LogVerbose("skipping {0}, because the CL part {1} had weird characters", JustDir, CLPart);
-                            }
-                        }
-                        else
-                        {
-                            LogVerbose("skipping {0}, because it didn't start with {1} or end with {2}", JustDir, StartString, EndString);
-                        }
+                        Log("Deleting formal build directory {0}, because it is {1} days old (maximum {2}).", ThisDirInfo.FullName, (int)AgeDays, MaximumDaysToKeepTempStorage);
+                        DeleteDirectory_NoExceptions(true, ThisDirInfo.FullName);
+                    }
+                    else
+                    {
+						LogVerbose("Not deleting formal build directory {0}, because it is {1} days old (maximum {2}).", ThisDirInfo.FullName, (int)AgeDays, MaximumDaysToKeepTempStorage);
                     }
                 }
             }
             catch (Exception Ex)
             {
-                LogWarning("Unable to Clean Directory with DirectoryForThisBuild {0}", DirectoryForThisBuild);
+                LogWarning("Unable to clean formal builds from directory: {0}", ParentDir);
                 LogWarning(" Exception was {0}", LogUtils.FormatException(Ex));
             }
         }
 
+
 		/// <summary>
 		/// Returns the generic name for a given platform (eg. "Windows" for Win32/Win64)
 		/// </summary>
diff --git a/Engine/Source/Programs/AutomationTool/GUBP/LegacyNodes.cs b/Engine/Source/Programs/AutomationTool/GUBP/LegacyNodes.cs
--- a/Engine/Source/Programs/AutomationTool/GUBP/LegacyNodes.cs
+++ b/Engine/Source/Programs/AutomationTool/GUBP/LegacyNodes.cs
@@ -2109,15 +2109,22 @@ partial class GUBP
 		{
 			return base.CISFrequencyQuantumShift(BranchConfig) + 3;
 		}
+
+		public static string GetBranchArchiveDirectory(GUBP.GUBPBranchConfig BranchConfig, BranchInfo.BranchUProject InGameProj)
+		{
+			// Find the build share where formal builds will be placed for this game.
+			string BuildShareName;
+			if (!BranchConfig.BranchOptions.GameNameToBuildShareMapping.TryGetValue(InGameProj.GameName, out BuildShareName))
+			{
+				BuildShareName = "UE4";
+			}
+			return CommandUtils.CombinePaths(CommandUtils.RootBuildStorageDirectory(), BuildShareName, "PackagedBuilds", P4Env.BuildRootEscaped);
+		}
+
 		public static string GetArchiveDirectory(GUBP.GUBPBranchConfig BranchConfig, BranchInfo.BranchUProject InGameProj, UnrealTargetPlatform InHostPlatform, List<UnrealTargetPlatform> InClientTargetPlatforms = null, List<UnrealTargetConfiguration> InClientConfigs = null, List<UnrealTargetPlatform> InServerTargetPlatforms = null, List<UnrealTargetConfiguration> InServerConfigs = null, bool InClientNotGame = false)
         {
-            // Find the build share where formal builds will be placed for this game.
-            string BuildShareName;
-            if (!BranchConfig.BranchOptions.GameNameToBuildShareMapping.TryGetValue(InGameProj.GameName, out BuildShareName))
-            {
-                BuildShareName = "UE4";
-            }
-            string BaseDir = CommandUtils.CombinePaths(CommandUtils.RootBuildStorageDirectory(), BuildShareName, "PackagedBuilds", P4Env.BuildRootEscaped, BranchConfig.JobInfo.BuildName);
+			string BranchArchiveDir = GetBranchArchiveDirectory(BranchConfig, InGameProj);
+            string BaseDir = CommandUtils.CombinePaths(BranchArchiveDir, BranchConfig.JobInfo.BuildName);
 			string NodeName = StaticGetBaseName(InGameProj, InHostPlatform, InClientTargetPlatforms, InClientConfigs, InServerTargetPlatforms, InServerConfigs, InClientNotGame);
 			string ArchiveDirectory = CombinePaths(BaseDir, NodeName);
             return ArchiveDirectory;
@@ -2245,7 +2252,7 @@ partial class GUBP
 					}
 					CreateDirectory_NoExceptions(IntermediateArchiveDirectory);
 				}
-				CleanFormalBuilds(FinalArchiveDirectory);
+				CleanFormalBuilds(FinalArchiveDirectory, "CL-*");
 				if (DirectoryExists_NoExceptions(FinalArchiveDirectory))
                 {
                     if (IsBuildMachine)
diff --git a/Engine/Source/Programs/AutomationTool/Scripts/Tests.Automation.cs b/Engine/Source/Programs/AutomationTool/Scripts/Tests.Automation.cs
--- a/Engine/Source/Programs/AutomationTool/Scripts/Tests.Automation.cs
+++ b/Engine/Source/Programs/AutomationTool/Scripts/Tests.Automation.cs
@@ -1646,9 +1646,8 @@ class TestCleanFormalBuilds : BuildCommand
 	public override void ExecuteBuild()
 	{
 		Log("*********************** TestCleanFormalBuilds");
-		var Dir = ParseParamValue("Dir", @"P:\Builds\Soul\Soul_Android_Shipping_MakeBuild\++depot+UE4-CL-2077154");
-		var CLString = ParseParamValue("CL", "2077154");
-		CleanFormalBuilds(Dir, CLString);
+		var Dir = ParseParamValue("Dir", @"P:\Builds\UE4\PackagedBuilds\++UE4+Release-4.11");
+		CleanFormalBuilds(Dir, "CL-*");
 	}
 }
 
