UAT: Fix search directory when looking for outdated formal build folders to delete.

diff --git a/Engine/Source/Programs/AutomationTool/GUBP/LegacyNodes.cs b/Engine/Source/Programs/AutomationTool/GUBP/LegacyNodes.cs
--- a/Engine/Source/Programs/AutomationTool/GUBP/LegacyNodes.cs
+++ b/Engine/Source/Programs/AutomationTool/GUBP/LegacyNodes.cs
@@ -2252,7 +2252,7 @@ partial class GUBP
 					}
 					CreateDirectory_NoExceptions(IntermediateArchiveDirectory);
 				}
-				CleanFormalBuilds(FinalArchiveDirectory, "CL-*");
+				CleanFormalBuilds(GetBranchArchiveDirectory(BranchConfig, GameProj), "CL-*");
 				if (DirectoryExists_NoExceptions(FinalArchiveDirectory))
                 {
                     if (IsBuildMachine)
