Backing out changelist 2935527 and replacing with a more narrowly scoped fix. Path.GetFullPath() on runtime dependencies in receipt for PS4 was causing p4-style "..." wildcards to be expanded as ".." parent directory separators for Slate content. New fix only operates on the actual file causing problems.

diff --git a/Engine/Source/Programs/AutomationTool/Win/WinPlatform.Automation.cs b/Engine/Source/Programs/AutomationTool/Win/WinPlatform.Automation.cs
--- a/Engine/Source/Programs/AutomationTool/Win/WinPlatform.Automation.cs
+++ b/Engine/Source/Programs/AutomationTool/Win/WinPlatform.Automation.cs
@@ -48,7 +48,8 @@ public abstract class BaseWinPlatform : Platform
 				if(Executable != null)
 				{
 					// only create bootstraps for executables
-					if (SC.NonUFSStagingFiles.ContainsKey(Executable.Path) && Path.GetExtension(Executable.Path) == ".exe")
+					string FullExecutablePath = Path.GetFullPath(Executable.Path);
+					if (SC.NonUFSStagingFiles.ContainsKey(FullExecutablePath) && Path.GetExtension(FullExecutablePath) == ".exe")
 					{
 						string BootstrapArguments = "";
 						if (!SC.IsCodeBasedProject && !ShouldStageCommandLine(Params, SC))
@@ -59,7 +60,7 @@ public abstract class BaseWinPlatform : Platform
 						string BootstrapExeName;
 						if(SC.StageTargetConfigurations.Count > 1)
 						{
-							BootstrapExeName = Path.GetFileName(Executable.Path);
+							BootstrapExeName = Path.GetFileName(FullExecutablePath);
 						}
 						else if(Params.IsCodeBasedProject)
 						{
@@ -70,9 +71,9 @@ public abstract class BaseWinPlatform : Platform
 							BootstrapExeName = SC.ShortProjectName + ".exe";
 						}
 
-						foreach (string StagePath in SC.NonUFSStagingFiles[Executable.Path])
+						foreach (string StagePath in SC.NonUFSStagingFiles[FullExecutablePath])
 						{
-							StageBootstrapExecutable(SC, BootstrapExeName, Executable.Path, StagePath, BootstrapArguments);
+							StageBootstrapExecutable(SC, BootstrapExeName, FullExecutablePath, StagePath, BootstrapArguments);
 						}
 					}
 				}
diff --git a/Engine/Source/Programs/UnrealBuildTool/System/TargetReceipt.cs b/Engine/Source/Programs/UnrealBuildTool/System/TargetReceipt.cs
--- a/Engine/Source/Programs/UnrealBuildTool/System/TargetReceipt.cs
+++ b/Engine/Source/Programs/UnrealBuildTool/System/TargetReceipt.cs
@@ -386,11 +386,11 @@ namespace UnrealBuildTool
 			// Replace all the variables in the paths
 			foreach (BuildProduct BuildProduct in BuildProducts)
 			{
-				BuildProduct.Path = Path.GetFullPath(Utils.ExpandVariables(BuildProduct.Path, Variables));
+				BuildProduct.Path = Utils.ExpandVariables(BuildProduct.Path, Variables);
 			}
 			foreach (RuntimeDependency RuntimeDependency in RuntimeDependencies)
 			{
-				RuntimeDependency.Path = Path.GetFullPath(Utils.ExpandVariables(RuntimeDependency.Path, Variables));
+				RuntimeDependency.Path = Utils.ExpandVariables(RuntimeDependency.Path, Variables);
 			}
 		}
 
