Copying //UE4/Dev-Build to //UE4/Main (Source: //UE4/Dev-Build @ 2944521)

==========================
MAJOR FEATURES + CHANGES
==========================

Change 2909886 on 2016/03/15 by [REDACTED_PERSON]

	Adding a build exception to give a message instead of crashing when trying to generate all project files from an installed build.

Change 2911727 on 2016/03/16 by [REDACTED_PERSON]

	Added Platform Type and Architecture to Installed Platform Info
	Reworked the different IsValid... functions to use lamdas to reduce duplicated code looping and checking receipts
	Moved the code to write config file entries into InstalledPlatformInfo so that it can be reused by anyone wanting to make installed builds
	Added temporary hack to write Android architecture until I can get it from build process

Change 2913692 on 2016/03/17 by [REDACTED_PERSON]

	UAT: Move script to archive a build for UGS into a public folder.

Change 2915445 on 2016/03/18 by [REDACTED_PERSON]

	UAT: Reduce the number of redundant log warnings/errors after a reported build failure, and simplify calls to ParallelExecutor which don't need retrying.

Change 2915450 on 2016/03/18 by [REDACTED_PERSON]

	UAT: Suppress warning messages trying to kill child processes if the operation failed because it's already exited.

Change 2925830 on 2016/03/29 by [REDACTED_PERSON]

	Added new selective download tags
	Added a test for whether installed platforms are missing required files so that we can try to open the launcher to the installer settings

Change 2926437 on 2016/03/29 by [REDACTED_PERSON]

	PR #2210: Fix "Rebuild.bat" for paths with parentheses (Contributed by [REDACTED_PERSON])

Change 2927399 on 2016/03/30 by [REDACTED_PERSON]

	Updating use of PDBCopy to look in VS2015 folder and fall back to VS2013 version if it doesn't exist.

Change 2933093 on 2016/04/05 by [REDACTED_PERSON]

	PR #2232: Updated copyright text to 2016 (Contributed by [REDACTED_PERSON])

Change 2936221 on 2016/04/07 by [REDACTED_PERSON]

	Adding checks on architecture for android config options

Change 2938021 on 2016/04/08 by [REDACTED_PERSON]

	UAT: Prevent UnauthorizedAccessException when enumerating crash files on Mac from a restricted user account.

Change 2939332 on 2016/04/11 by [REDACTED_PERSON]

	Added AdditionalBundleResources to external file list so that they should be included in Launcher releases

Change 2939767 on 2016/04/11 by [REDACTED_PERSON]

	BuildGraph: Add a -preprocess option, which will cause the preprocessed and culled graph out to an XML file for debugging.

Change 2941611 on 2016/04/12 by [REDACTED_PERSON]

	UAT: Prevent warning about commands requiring P4 if -p4 is specified on the command line.

Change 2942037 on 2016/04/13 by [REDACTED_PERSON]

	UBT: Only print 'Detailed Action Stats' message footer if there were any detailed action stats.

Change 2942640 on 2016/04/13 by [REDACTED_PERSON]

	GUBP: Trigger GitHub promotions by triggering a new procedure rather than scanning for labels.

Change 2942728 on 2016/04/13 by [REDACTED_PERSON]

	BuildGraph: Rename "AgentGroup" to "Agent" for consistency with XML.

Change 2942735 on 2016/04/13 by [REDACTED_PERSON]

	BuildGraph: Few renames to match class names (Build.cs -> BuildGraph.cs, AgentGroup.cs -> Agent.cs)

Change 2943568 on 2016/04/14 by [REDACTED_PERSON]

	EC: Print out the log folder at the start of each job.

Change 2944421 on 2016/04/14 by [REDACTED_PERSON]

	EC: Add GitHub dashboard page which shows the current syncing state

diff --git a/Engine/Build/BatchFiles/Rebuild.bat b/Engine/Build/BatchFiles/Rebuild.bat
--- a/Engine/Build/BatchFiles/Rebuild.bat
+++ b/Engine/Build/BatchFiles/Rebuild.bat
@@ -7,13 +7,13 @@ REM %3 is the configuration name
 IF EXIST "%~dp0\Clean.bat" (
 	call "%~dp0\Clean.bat" %*
 ) ELSE (
-	ECHO Clean.bat not found in %~dp0 
+	ECHO Clean.bat not found in "%~dp0"
 	EXIT /B 999
 )
 
 IF EXIST "%~dp0\Build.bat" (
 	call "%~dp0\Build.bat" %*
 ) ELSE (
-	ECHO Build.bat not found in %~dp0 
+	ECHO Build.bat not found in "%~dp0"
 	EXIT /B 999
 )
diff --git a/Engine/Build/InstalledEngineFilters.ini b/Engine/Build/InstalledEngineFilters.ini
--- a/Engine/Build/InstalledEngineFilters.ini
+++ b/Engine/Build/InstalledEngineFilters.ini
@@ -94,7 +94,6 @@
 ; Mac-specific dependencies
 {Mac} /Engine/Binaries/ThirdParty/ICU/icu4c-53_1/Mac/*
 {Mac} /Engine/Binaries/ThirdParty/Mono/Mac/...
-{Mac} /Engine/Source/ThirdParty/Mac/RadioEffectUnit/...
 {Mac} /Engine/Source/ThirdParty/CEF3/cef_binary_3.2357.1291.g47e6d4b_macosx64/Release/locale/...
 
 {Mac} /Engine/Build/BatchFiles/Mac/Build.sh
@@ -258,3 +257,26 @@
 [LauncherAttribute.tag:templates]
 /FeaturePacks/...
 /Templates/...
+
+[LauncherAttribute.tag:engine_source]
+/Engine/Source/....cpp
+
+[LauncherAttribute.tag:art_tools]
+/Engine/Extras/3dsMaxScripts/...
+/Engine/Extras/Maya_AnimationRiggingTools/...
+/Engine/Extras/MayaVelocityGridExporter/...
+
+[LauncherAttribute.tag:platform_IOS]
+/Engine/Binaries/IOS/...
+
+[LauncherAttribute.tag:platform_Android]
+/Engine/Binaries/Android/...
+
+[LauncherAttribute.tag:platform_HTML5]
+/Engine/Binaries/HTML5/...
+
+[LauncherAttribute.tag:platform_Linux]
+/Engine/Binaries/Linux/...
+
+[LauncherAttribute.tag:platform_TVOS]
+/Engine/Binaries/TVOS/...
diff --git a/Engine/Source/Developer/Android/AndroidPlatformEditor/Private/AndroidTargetSettingsCustomization.cpp b/Engine/Source/Developer/Android/AndroidPlatformEditor/Private/AndroidTargetSettingsCustomization.cpp
--- a/Engine/Source/Developer/Android/AndroidPlatformEditor/Private/AndroidTargetSettingsCustomization.cpp
+++ b/Engine/Source/Developer/Android/AndroidPlatformEditor/Private/AndroidTargetSettingsCustomization.cpp
@@ -16,6 +16,7 @@
 #include "SNotificationList.h"
 #include "NotificationManager.h"
 #include "EngineBuildSettings.h"
+#include "InstalledPlatformInfo.h"
 
 #define LOCTEXT_NAMESPACE "AndroidRuntimeSettings"
 
@@ -200,6 +201,13 @@ void FAndroidTargetSettingsCustomization::BuildAppManifestSection(IDetailLayoutB
 	GooglePlayCategory.AddProperty(GooglePlayLicenseKeyProperty)
 		.EditCondition(SetupForGooglePlayAttribute, NULL);
 
+#define SETUP_ANDROIDARCH_PROP(ArchFragment, PropName, Category, Tip) \
+	{ \
+		TSharedRef<IPropertyHandle> PropertyHandle = DetailLayout.GetProperty(GET_MEMBER_NAME_CHECKED(UAndroidRuntimeSettings, PropName)); \
+		Category.AddProperty(PropertyHandle) \
+			.IsEnabled(FInstalledPlatformInfo::Get().IsValidPlatformArchitecture(TEXT("Android"), ArchFragment)) \
+			.ToolTip(FInstalledPlatformInfo::Get().IsValidPlatformArchitecture(TEXT("Android"), ArchFragment) ? Tip : FAndroidTargetSettingsCustomizationConstants::DisabledTip); \
+	}
 
 #define SETUP_SOURCEONLY_PROP(PropName, Category, Tip) \
 	{ \
@@ -208,11 +216,11 @@ void FAndroidTargetSettingsCustomization::BuildAppManifestSection(IDetailLayoutB
 			.IsEnabled(FEngineBuildSettings::IsSourceDistribution()) \
 			.ToolTip(FEngineBuildSettings::IsSourceDistribution() ? Tip : FAndroidTargetSettingsCustomizationConstants::DisabledTip); \
 	}
-	SETUP_SOURCEONLY_PROP(bBuildForArmV7, BuildCategory, LOCTEXT("BuildForArmV7ToolTip", "Enable ArmV7 CPU architecture support? (this will be used if all CPU architecture types are unchecked)"));
-	SETUP_SOURCEONLY_PROP(bBuildForX86, BuildCategory, LOCTEXT("BuildForX86ToolTip", "Enable X86 CPU architecture support?"));
-	SETUP_SOURCEONLY_PROP(bBuildForX8664, BuildCategory, LOCTEXT("BuildForX8664ToolTip", "Enable X86-64 CPU architecture support?"));
-	SETUP_SOURCEONLY_PROP(bBuildForES2, BuildCategory, LOCTEXT("BuildForES2ToolTip", "Enable OpenGL ES2 rendering support? (this will be used if rendering types are unchecked)"));
-	SETUP_SOURCEONLY_PROP(bBuildForES31, BuildCategory, LOCTEXT("BuildForES31ToolTip", "Enable OpenGL ES31 + AEP (Android Extension Pack) rendering support?"));
+	SETUP_ANDROIDARCH_PROP(TEXT("-armv7"), bBuildForArmV7, BuildCategory, LOCTEXT("BuildForArmV7ToolTip", "Enable ArmV7 CPU architecture support? (this will be used if all CPU architecture types are unchecked)"));
+	SETUP_ANDROIDARCH_PROP(TEXT("-x86"), bBuildForX86, BuildCategory, LOCTEXT("BuildForX86ToolTip", "Enable X86 CPU architecture support?"));
+	SETUP_ANDROIDARCH_PROP(TEXT("-x64"), bBuildForX8664, BuildCategory, LOCTEXT("BuildForX8664ToolTip", "Enable X86-64 CPU architecture support?"));
+	SETUP_ANDROIDARCH_PROP(TEXT("-es2"), bBuildForES2, BuildCategory, LOCTEXT("BuildForES2ToolTip", "Enable OpenGL ES2 rendering support? (this will be used if rendering types are unchecked)"));
+	SETUP_ANDROIDARCH_PROP(TEXT("-es31"), bBuildForES31, BuildCategory, LOCTEXT("BuildForES31ToolTip", "Enable OpenGL ES31 + AEP (Android Extension Pack) rendering support?"));
 	SETUP_SOURCEONLY_PROP(bSupportsVulkan, BuildCategory, LOCTEXT("SupportsVulkanToolTip", "Enable Vulkan rendering support?"));
 
 	// @todo android fat binary: Put back in when we expose those
diff --git a/Engine/Source/Developer/DesktopPlatform/Private/PlatformInfo.cpp b/Engine/Source/Developer/DesktopPlatform/Private/PlatformInfo.cpp
--- a/Engine/Source/Developer/DesktopPlatform/Private/PlatformInfo.cpp
+++ b/Engine/Source/Developer/DesktopPlatform/Private/PlatformInfo.cpp
@@ -210,6 +210,31 @@ TArray<FVanillaPlatformEntry> BuildPlatformHierarchy(const EPlatformFilter InFil
 	return VanillaPlatforms;
 }
 
+EPlatformType EPlatformTypeFromString(const FString& PlatformTypeName)
+{
+	if (FCString::Strcmp(*PlatformTypeName, TEXT("Game")) == 0)
+	{
+		return PlatformInfo::EPlatformType::Game;
+	}
+	else if (FCString::Strcmp(*PlatformTypeName, TEXT("Editor")) == 0)
+	{
+		return PlatformInfo::EPlatformType::Editor;
+	}
+	else if (FCString::Strcmp(*PlatformTypeName, TEXT("Client")) == 0)
+	{
+		return PlatformInfo::EPlatformType::Client;
+	}
+	else if (FCString::Strcmp(*PlatformTypeName, TEXT("Server")) == 0)
+	{
+		return PlatformInfo::EPlatformType::Server;
+	}
+	else
+	{
+		UE_LOG(LogDesktopPlatform, Warning, TEXT("Unable to read Platform Type from %s, defaulting to Game"), *PlatformTypeName);
+		return PlatformInfo::EPlatformType::Game;
+	}
+}
+
 } // namespace PlatformFamily
 
 #undef LOCTEXT_NAMESPACE
diff --git a/Engine/Source/Developer/DesktopPlatform/Public/PlatformInfo.h b/Engine/Source/Developer/DesktopPlatform/Public/PlatformInfo.h
--- a/Engine/Source/Developer/DesktopPlatform/Public/PlatformInfo.h
+++ b/Engine/Source/Developer/DesktopPlatform/Public/PlatformInfo.h
@@ -329,4 +329,11 @@ namespace PlatformInfo
 	* @param InDisplayName - The new display name
 	*/
 	DESKTOPPLATFORM_API void UpdatePlatformDisplayName(FString InPlatformName, FText InDisplayName);
+
+	/**
+	* Returns an EPlatformType value from a string representation.
+	* @param PlatformTypeName The string to get the EPlatformType for.
+	* @return An EPlatformType value.
+	*/
+	DESKTOPPLATFORM_API EPlatformType EPlatformTypeFromString(const FString& PlatformTypeName);
 }
diff --git a/Engine/Source/Developer/TargetPlatform/Private/InstalledPlatformInfo.cpp b/Engine/Source/Developer/TargetPlatform/Private/InstalledPlatformInfo.cpp
--- a/Engine/Source/Developer/TargetPlatform/Private/InstalledPlatformInfo.cpp
+++ b/Engine/Source/Developer/TargetPlatform/Private/InstalledPlatformInfo.cpp
@@ -2,6 +2,10 @@
 
 #include "TargetPlatformPrivatePCH.h"
 #include "InstalledPlatformInfo.h"
+#include "PlatformInfo.h"
+#include "DesktopPlatformModule.h"
+
+#define LOCTEXT_NAMESPACE "InstalledPlatformInfo"
 
 DEFINE_LOG_CATEGORY_STATIC(LogInstalledPlatforms, Log, All);
 
@@ -65,6 +69,16 @@ void FInstalledPlatformInfo::ParsePlatformConfiguration(FString PlatformConfigur
 		bCanCreateEntry = false;
 	}
 
+	FString PlatformTypeName;
+	PlatformInfo::EPlatformType PlatformType = PlatformInfo::EPlatformType::Game;
+	if (FParse::Value(*PlatformConfiguration, TEXT("PlatformType="), PlatformTypeName))
+	{
+		PlatformType = PlatformInfo::EPlatformTypeFromString(PlatformTypeName);
+	}
+
+	FString Architecture;
+	FParse::Value(*PlatformConfiguration, TEXT("Architecture="), Architecture);
+
 	FString RequiredFile;
 	if (FParse::Value(*PlatformConfiguration, TEXT("RequiredFile="), RequiredFile))
 	{
@@ -88,106 +102,154 @@ void FInstalledPlatformInfo::ParsePlatformConfiguration(FString PlatformConfigur
 	
 	if (bCanCreateEntry)
 	{
-		FInstalledPlatformConfiguration NewConfig = {Configuration, PlatformName, RequiredFile, ProjectType, bCanBeDisplayed};
+		FInstalledPlatformConfiguration NewConfig = {Configuration, PlatformName, PlatformType, Architecture, RequiredFile, ProjectType, bCanBeDisplayed};
 		InstalledPlatformConfigurations.Add(NewConfig);
 	}
 }
 
 bool FInstalledPlatformInfo::IsValidConfiguration(const EBuildConfigurations::Type Configuration, EProjectType ProjectType) const
+{
+	return ContainsValidConfiguration(
+		[Configuration, ProjectType](const FInstalledPlatformConfiguration& CurConfig)
+		{
+			return CurConfig.Configuration == Configuration
+				&& (ProjectType == EProjectType::Any || CurConfig.ProjectType == EProjectType::Any
+					|| CurConfig.ProjectType == ProjectType);
+		}
+	);
+}
+
+bool FInstalledPlatformInfo::IsValidPlatform(const FString& PlatformName, EProjectType ProjectType) const
+{
+	return ContainsValidConfiguration(
+		[PlatformName, ProjectType](const FInstalledPlatformConfiguration& CurConfig)
+		{
+			return CurConfig.PlatformName == PlatformName
+				&& (ProjectType == EProjectType::Any || CurConfig.ProjectType == EProjectType::Any
+					|| CurConfig.ProjectType == ProjectType);
+		}
+	);
+}
+
+bool FInstalledPlatformInfo::IsValidPlatformAndConfiguration(const EBuildConfigurations::Type Configuration, const FString& PlatformName, EProjectType ProjectType) const
+{
+	return ContainsValidConfiguration(
+		[Configuration, PlatformName, ProjectType](const FInstalledPlatformConfiguration& CurConfig)
+		{
+			return CurConfig.Configuration == Configuration
+				&& CurConfig.PlatformName == PlatformName
+				&& (ProjectType == EProjectType::Any || CurConfig.ProjectType == EProjectType::Any
+					|| CurConfig.ProjectType == ProjectType);
+		}
+	);
+}
+
+bool FInstalledPlatformInfo::CanDisplayPlatform(const FString& PlatformName, EProjectType ProjectType) const
+{
+	return ContainsMatchingConfiguration(
+		[PlatformName, ProjectType](const FInstalledPlatformConfiguration& CurConfig)
+		{
+			return CurConfig.PlatformName == PlatformName && (CurConfig.bCanBeDisplayed
+				|| ProjectType == EProjectType::Any || CurConfig.ProjectType == EProjectType::Any
+				|| CurConfig.ProjectType == ProjectType);
+		}
+	);
+}
+
+bool FInstalledPlatformInfo::IsValidPlatformType(PlatformInfo::EPlatformType PlatformType) const
+{
+	return ContainsValidConfiguration(
+		[PlatformType](const FInstalledPlatformConfiguration& CurConfig)
+		{
+			return CurConfig.PlatformType == PlatformType;
+		}
+	);
+}
+
+bool FInstalledPlatformInfo::IsValidPlatformArchitecture(const FString& PlatformName, const FString& Architecture) const
+{
+	return ContainsValidConfiguration(
+		[PlatformName, Architecture](const FInstalledPlatformConfiguration& CurConfig)
+		{
+			return CurConfig.PlatformName == PlatformName && CurConfig.Architecture.Contains(Architecture);
+		}
+	);
+}
+
+bool FInstalledPlatformInfo::IsPlatformMissingRequiredFile(const FString& PlatformName) const
 {
 	if (FApp::IsEngineInstalled())
 	{
-		for (const FInstalledPlatformConfiguration& PlatformConfiguration : InstalledPlatformConfigurations)
-		{
-			if (PlatformConfiguration.Configuration == Configuration)
+		return ContainsMatchingConfiguration(
+			[PlatformName](const FInstalledPlatformConfiguration& CurConfig)
 			{
-				if (ProjectType == EProjectType::Any || PlatformConfiguration.ProjectType == EProjectType::Any
-				 || PlatformConfiguration.ProjectType == ProjectType)
-				{
-					if (PlatformConfiguration.RequiredFile.IsEmpty()
-						|| FPaths::FileExists(PlatformConfiguration.RequiredFile))
-					{
-						return true;
-					}
-				}
+				return CurConfig.PlatformName == PlatformName
+					&& !CurConfig.RequiredFile.IsEmpty()
+					&& !FPaths::FileExists(CurConfig.RequiredFile);
 			}
-		}
-
-		return false;
+		);
 	}
-	return true;
+	return false;
 }
 
-bool FInstalledPlatformInfo::IsValidPlatform(const FString& PlatformName, EProjectType ProjectType) const
+bool FInstalledPlatformInfo::OpenInstallerOptions()
 {
-	if (FApp::IsEngineInstalled())
+	IDesktopPlatform* DesktopPlatform = FDesktopPlatformModule::Get();
+
+	if (DesktopPlatform != nullptr)
 	{
-		for (const FInstalledPlatformConfiguration& PlatformConfiguration : InstalledPlatformConfigurations)
+		FString CurrentIdentifier = DesktopPlatform->GetCurrentEngineIdentifier();
+		if (DesktopPlatform->IsStockEngineRelease(CurrentIdentifier))
 		{
-			if (PlatformConfiguration.PlatformName == PlatformName)
+			if (FMessageDialog::Open(EAppMsgType::YesNo, LOCTEXT("NotInstalled_SelectedPlatform", "The Binaries for this Target Platform are not currently installed, would you like to use the Launcher to download them?")) == EAppReturnType::Yes)
 			{
-				if (ProjectType == EProjectType::Any || PlatformConfiguration.ProjectType == EProjectType::Any
-				 || PlatformConfiguration.ProjectType == ProjectType)
+				// TODO: Ensure that this URL opens the launcher correctly before this is included in a release
+				FString InstallerURL = FString::Printf(TEXT("ue/library/engines/UE_%s/installer"), *DesktopPlatform->GetEngineDescription(CurrentIdentifier));
+				FOpenLauncherOptions OpenOptions(InstallerURL);
+				if (DesktopPlatform->OpenLauncher(OpenOptions))
 				{
-					if (PlatformConfiguration.RequiredFile.IsEmpty()
-						|| FPaths::FileExists(PlatformConfiguration.RequiredFile))
-					{
-						return true;
-					}
+					return true;
 				}
 			}
 		}
-
-		return false;
 	}
-	return true;
+	return false;
 }
 
-bool FInstalledPlatformInfo::IsValidPlatformAndConfiguration(const EBuildConfigurations::Type Configuration, const FString& PlatformName, EProjectType ProjectType) const
+bool FInstalledPlatformInfo::ContainsValidConfiguration(TFunctionRef<bool(const FInstalledPlatformConfiguration)> ConfigFilter) const
 {
 	if (FApp::IsEngineInstalled())
 	{
 		for (const FInstalledPlatformConfiguration& PlatformConfiguration : InstalledPlatformConfigurations)
 		{
-			if (PlatformConfiguration.Configuration == Configuration
-			 && PlatformConfiguration.PlatformName == PlatformName)
+			// Check whether filter accepts this configuration and it has required file
+			if (ConfigFilter(PlatformConfiguration)
+				&& (PlatformConfiguration.RequiredFile.IsEmpty()
+					|| FPaths::FileExists(PlatformConfiguration.RequiredFile)))
 			{
-				if (ProjectType == EProjectType::Any || PlatformConfiguration.ProjectType == EProjectType::Any
-				 || PlatformConfiguration.ProjectType == ProjectType)
-				{
-					if (PlatformConfiguration.RequiredFile.IsEmpty()
-					 || FPaths::FileExists(PlatformConfiguration.RequiredFile))
-					{
-						return true;
-					}
-				}
+				return true;
 			}
 		}
-
 		return false;
 	}
 	return true;
 }
 
-bool FInstalledPlatformInfo::CanDisplayPlatform(const FString& PlatformName, EProjectType ProjectType) const
+bool FInstalledPlatformInfo::ContainsMatchingConfiguration(TFunctionRef<bool(const FInstalledPlatformConfiguration)> ConfigFilter) const
 {
 	if (FApp::IsEngineInstalled())
 	{
 		for (const FInstalledPlatformConfiguration& PlatformConfiguration : InstalledPlatformConfigurations)
 		{
-			if (PlatformConfiguration.PlatformName == PlatformName)
+			// Check whether filter accepts this configuration
+			if (ConfigFilter(PlatformConfiguration))
 			{
-				if ( PlatformConfiguration.bCanBeDisplayed
-					|| ProjectType == EProjectType::Any || PlatformConfiguration.ProjectType == EProjectType::Any
-					|| PlatformConfiguration.ProjectType == ProjectType )
-				{
-					return true;
-				}
+				return true;
 			}
 		}
-
 		return false;
 	}
 	return true;
 }
 
+#undef LOCTEXT_NAMESPACE
diff --git a/Engine/Source/Developer/TargetPlatform/Public/InstalledPlatformInfo.h b/Engine/Source/Developer/TargetPlatform/Public/InstalledPlatformInfo.h
--- a/Engine/Source/Developer/TargetPlatform/Public/InstalledPlatformInfo.h
+++ b/Engine/Source/Developer/TargetPlatform/Public/InstalledPlatformInfo.h
@@ -2,6 +2,12 @@
 
 #pragma once
 
+// Forward declaration
+namespace PlatformInfo
+{
+	enum class EPlatformType : uint8;
+}
+
 enum class EProjectType : uint8
 {
 	Unknown,
@@ -23,6 +29,12 @@ struct FInstalledPlatformConfiguration
 	/** Name of the Platform for this combination */
 	FString PlatformName;
 
+	/** Type of Platform for this combination */
+	PlatformInfo::EPlatformType PlatformType;
+
+	/** Name of the Architecture for this combination */
+	FString Architecture;
+
 	/** Location of a file that must exist for this combination to be valid (optional) */
 	FString RequiredFile;
 
@@ -63,8 +75,35 @@ public:
 	 */
 	bool IsValidPlatformAndConfiguration(const EBuildConfigurations::Type Configuration, const FString& PlatformName, EProjectType ProjectType = EProjectType::Any) const;
 
+	/**
+	 * Queries whether a platform can be displayed as an option, even if it's not supported for the specified project type
+	 */
 	bool CanDisplayPlatform(const FString& PlatformName, EProjectType ProjectType) const;
 
+	/**
+	 * Queries whether a platform type is valid for any configuration
+	 */
+	bool IsValidPlatformType(PlatformInfo::EPlatformType PlatformType) const;
+
+	/**
+	 * Queries whether a platform architecture is valid for any configuration
+	 * @param PlatformName Name of the platform's binary folder (eg. Win64, Android)
+	 * @param Architecture Either a full architecture name or a partial substring for CPU/GPU combinations (eg. "-armv7", "-es2")
+	 */
+	bool IsValidPlatformArchitecture(const FString& PlatformName, const FString& Architecture) const;
+
+	/**
+	 * Queries whether a platform has any missing required files
+	 */
+	bool IsPlatformMissingRequiredFile(const FString& PlatformName) const;
+
+	/**
+	 * Attempts to open the Launcher to the Installer options so that additional platforms can be downloaded
+	 *
+	 * @return false if the engine is not a stock release, user cancels action or launcher fails to load
+	 */
+	static bool OpenInstallerOptions();
+
 private:
 	/**
 	 * Constructor
@@ -76,6 +115,17 @@ private:
 	 */
 	void ParsePlatformConfiguration(FString PlatformConfiguration);
 
+	/**
+	 * Given a filter function, checks whether any configuration passes that filter and has required file
+	 */
+	bool ContainsValidConfiguration(TFunctionRef<bool(const FInstalledPlatformConfiguration)> ConfigFilter) const;
+
+	/**
+	 * Given a filter function, checks whether any configuration passes that filter
+	 * Doesn't check whether required file exists, so that we can find platforms that can be optionally installed
+	 */
+	bool ContainsMatchingConfiguration(TFunctionRef<bool(const FInstalledPlatformConfiguration)> ConfigFilter) const;
+
 	/** List of installed platform configuration combinations */
 	TArray<FInstalledPlatformConfiguration> InstalledPlatformConfigurations;
 };
diff --git a/Engine/Source/Editor/MainFrame/Private/Frame/MainFrameActions.cpp b/Engine/Source/Editor/MainFrame/Private/Frame/MainFrameActions.cpp
--- a/Engine/Source/Editor/MainFrame/Private/Frame/MainFrameActions.cpp
+++ b/Engine/Source/Editor/MainFrame/Private/Frame/MainFrameActions.cpp
@@ -320,6 +320,15 @@ void FMainFrameActionCallbacks::CookContent(const FName InPlatformInfoName)
 	const PlatformInfo::FPlatformInfo* const PlatformInfo = PlatformInfo::FindPlatformInfo(InPlatformInfoName);
 	check(PlatformInfo);
 
+	if (FInstalledPlatformInfo::Get().IsPlatformMissingRequiredFile(PlatformInfo->BinaryFolderName))
+	{
+		if (!FInstalledPlatformInfo::OpenInstallerOptions())
+		{
+			FMessageDialog::Open(EAppMsgType::Ok, LOCTEXT("MissingPlatformFilesCook", "Missing required files to cook for this platform."));
+		}
+		return;
+	}
+
 	FString OptionalParams;
 
 	if (!FModuleManager::LoadModuleChecked<IProjectTargetPlatformEditorModule>("ProjectTargetPlatformEditor").ShowUnsupportedTargetWarning(PlatformInfo->VanillaPlatformName))
@@ -404,6 +413,15 @@ void FMainFrameActionCallbacks::PackageProject( const FName InPlatformInfoName )
 	const PlatformInfo::FPlatformInfo* const PlatformInfo = PlatformInfo::FindPlatformInfo(InPlatformInfoName);
 	check(PlatformInfo);
 
+	if (FInstalledPlatformInfo::Get().IsPlatformMissingRequiredFile(PlatformInfo->BinaryFolderName))
+	{
+		if (!FInstalledPlatformInfo::OpenInstallerOptions())
+		{
+			FMessageDialog::Open(EAppMsgType::Ok, LOCTEXT("MissingPlatformFilesPackage", "Missing required files to package this platform."));
+		}
+		return;
+	}
+
 	if (PlatformInfo->SDKStatus == PlatformInfo::EPlatformSDKStatus::NotInstalled || (bProjectHasCode && PlatformInfo->bUsesHostCompiler && !FSourceCodeNavigation::IsCompilerAvailable()))
 	{
 		IMainFrameModule& MainFrameModule = FModuleManager::GetModuleChecked<IMainFrameModule>(TEXT("MainFrame"));
diff --git a/Engine/Source/Editor/UnrealEd/Private/EditorUndoClient.cpp b/Engine/Source/Editor/UnrealEd/Private/EditorUndoClient.cpp
--- a/Engine/Source/Editor/UnrealEd/Private/EditorUndoClient.cpp
+++ b/Engine/Source/Editor/UnrealEd/Private/EditorUndoClient.cpp
@@ -1,4 +1,4 @@
-// Copyright 1998-2015 Epic Games, Inc. All Rights Reserved.
+// Copyright 1998-2016 Epic Games, Inc. All Rights Reserved.
 
 #include "UnrealEd.h"
 #include "EditorUndoClient.h"
diff --git a/Engine/Source/Editor/UnrealEd/Private/Kismet2/DebuggerCommands.cpp b/Engine/Source/Editor/UnrealEd/Private/Kismet2/DebuggerCommands.cpp
--- a/Engine/Source/Editor/UnrealEd/Private/Kismet2/DebuggerCommands.cpp
+++ b/Engine/Source/Editor/UnrealEd/Private/Kismet2/DebuggerCommands.cpp
@@ -1982,6 +1982,18 @@ void FInternalPlayWorldCommandCallbacks::LaunchOnDevice( const FString& DeviceId
 	FTargetDeviceId TargetDeviceId;
 	if (FTargetDeviceId::Parse(DeviceId, TargetDeviceId))
 	{
+		const PlatformInfo::FPlatformInfo* const PlatformInfo = PlatformInfo::FindPlatformInfo(*TargetDeviceId.GetPlatformName());
+		check(PlatformInfo);
+
+		if (FInstalledPlatformInfo::Get().IsPlatformMissingRequiredFile(PlatformInfo->BinaryFolderName))
+		{
+			if (!FInstalledPlatformInfo::OpenInstallerOptions())
+			{
+				FMessageDialog::Open(EAppMsgType::Ok, LOCTEXT("MissingPlatformFilesLaunch", "Missing required files to launch on this platform."));
+			}
+			return;
+		}
+
 		if (FModuleManager::LoadModuleChecked<IProjectTargetPlatformEditorModule>("ProjectTargetPlatformEditor").ShowUnsupportedTargetWarning(*TargetDeviceId.GetPlatformName()))
 		{
 			GUnrealEd->RequestPlaySession(DeviceId, DeviceName);
diff --git a/Engine/Source/Programs/AutomationTool/AutomationUtils/CommandletUtils.cs b/Engine/Source/Programs/AutomationTool/AutomationUtils/CommandletUtils.cs
--- a/Engine/Source/Programs/AutomationTool/AutomationUtils/CommandletUtils.cs
+++ b/Engine/Source/Programs/AutomationTool/AutomationUtils/CommandletUtils.cs
@@ -291,10 +291,17 @@ namespace AutomationTool
 				List<FileInfo> CrashFileInfos = new List<FileInfo>();
 				foreach(string CrashDir in CrashDirs)
 				{
-					DirectoryInfo CrashDirInfo = new DirectoryInfo(CrashDir);
-					if(CrashDirInfo.Exists)
+					try
 					{
-						CrashFileInfos.AddRange(CrashDirInfo.EnumerateFiles("*.crash", SearchOption.TopDirectoryOnly).Where(x => x.LastWriteTimeUtc >= StartTime));
+						DirectoryInfo CrashDirInfo = new DirectoryInfo(CrashDir);
+						if (CrashDirInfo.Exists)
+						{
+							CrashFileInfos.AddRange(CrashDirInfo.EnumerateFiles("*.crash", SearchOption.TopDirectoryOnly).Where(x => x.LastWriteTimeUtc >= StartTime));
+						}
+					}
+					catch (UnauthorizedAccessException)
+					{
+						// Not all account types can access /Library/Logs/DiagnosticReports
 					}
 				}
 				
diff --git a/Engine/Source/Programs/AutomationTool/AutomationUtils/P4Utils.cs b/Engine/Source/Programs/AutomationTool/AutomationUtils/P4Utils.cs
--- a/Engine/Source/Programs/AutomationTool/AutomationUtils/P4Utils.cs
+++ b/Engine/Source/Programs/AutomationTool/AutomationUtils/P4Utils.cs
@@ -519,7 +519,10 @@ namespace AutomationTool
 					var RequireP4Attributes = Command.GetCustomAttributes(typeof(RequireP4Attribute), true);	
 					if (!CommandUtils.IsNullOrEmpty(RequireP4Attributes))
 					{
-						LogWarning("Command {0} requires P4 functionality.", Command.Name);
+						if(!GlobalCommandLine.P4)
+						{
+							LogWarning("Command {0} requires P4 functionality.", Command.Name);
+						}
 						bRequireP4 = true;
 
 						var DoesNotNeedP4CLAttributes = Command.GetCustomAttributes(typeof(DoesNotNeedP4CLAttribute), true);
diff --git a/Engine/Source/Programs/AutomationTool/AutomationUtils/ProcessUtils.cs b/Engine/Source/Programs/AutomationTool/AutomationUtils/ProcessUtils.cs
--- a/Engine/Source/Programs/AutomationTool/AutomationUtils/ProcessUtils.cs
+++ b/Engine/Source/Programs/AutomationTool/AutomationUtils/ProcessUtils.cs
@@ -534,8 +534,11 @@ namespace AutomationTool
 						}
 						catch (Exception Ex)
 						{
-							CommandUtils.LogWarning("Failed to kill descendant:");
-							CommandUtils.LogWarning(LogUtils.FormatException(Ex));
+							if(!KillCandidate.HasExited)
+							{
+								CommandUtils.LogWarning("Failed to kill descendant:");
+								CommandUtils.LogWarning(LogUtils.FormatException(Ex));
+							}
 						}
 						break;  // exit the loop as who knows what else died, so let's get processes anew
 					}
diff --git a/Engine/Source/Programs/AutomationTool/AutomationUtils/UE4Build.cs b/Engine/Source/Programs/AutomationTool/AutomationUtils/UE4Build.cs
--- a/Engine/Source/Programs/AutomationTool/AutomationUtils/UE4Build.cs
+++ b/Engine/Source/Programs/AutomationTool/AutomationUtils/UE4Build.cs
@@ -658,7 +658,7 @@ namespace AutomationTool
 			/// Used for special jobs to test things. At the moment, debugging some megaXGE thing
 			public bool SpecialTestFlag = false;
 
-			public bool DoRetries = IsBuildMachine;
+			public bool DoRetries = false;
 
 			/// <summary>
 			/// Adds a target with the specified configuration.
@@ -814,67 +814,73 @@ namespace AutomationTool
 				}
 
 				CombineXGEStopwatch.Finish();
-               
-				using(TelemetryStopwatch ProcessXGEStopwatch = new TelemetryStopwatch("ProcessXGE.{0}", Path.GetFileNameWithoutExtension(XGETool)))
+
+				if(XGETool == null)
 				{
-					int Retries = DoRetries ? 2 : 1;
-					int ConnectionRetries = 4;
-					for (int i = 0; i < Retries; i++)
+					PushDir(CombinePaths(CmdEnv.LocalRoot, @"\Engine\Source"));
+					try
 					{
-						try
+						int ExitCode = ParallelExecutor.Execute(TaskFilePath);
+						if(ExitCode != 0)
 						{
-							while (true)
+							return false;
+						}
+					}
+					finally
+					{
+						PopDir();
+					}
+				}
+				else
+				{
+					using (TelemetryStopwatch ProcessXGEStopwatch = new TelemetryStopwatch("ProcessXGE.{0}", Path.GetFileNameWithoutExtension(XGETool)))
+					{
+						int Retries = DoRetries ? 2 : 1;
+						int ConnectionRetries = 4;
+						for (int i = 0; i < Retries; i++)
+						{
+							try
 							{
-								LogVerbose("Running {0} *******", XGETool);
-								PushDir(CombinePaths(CmdEnv.LocalRoot, @"\Engine\Source"));
-								int SuccesCode;
-								string LogFile;
-								bool bOutputContainsProject;
-								if(XGETool == null)
+								while (true)
 								{
-									LogFile = GetRunAndLogOnlyName(CmdEnv, "ParallelExecutor");
-									SuccesCode = ParallelExecutor.Execute(TaskFilePath);
-									bOutputContainsProject = false;
-								}
-								else
-								{
-									LogFile = GetRunAndLogOnlyName(CmdEnv, XGETool);
+									LogVerbose("Running {0} *******", XGETool);
+									PushDir(CombinePaths(CmdEnv.LocalRoot, @"\Engine\Source"));
+									int SuccesCode;
+									string LogFile = GetRunAndLogOnlyName(CmdEnv, XGETool);
 									string Output = RunAndLog(XGETool, Args, out SuccesCode, LogFile);
-									bOutputContainsProject = Output.Contains("------Project:");
-								}
-								PopDir();
-								if (ConnectionRetries > 0 && (SuccesCode == 4 || SuccesCode == 2) && !bOutputContainsProject && XGETool != null)
-								{
-									LogWarning(String.Format("{0} failure on the local connection timeout", XGETool));
-									if (ConnectionRetries < 2)
+									bool bOutputContainsProject = Output.Contains("------Project:");
+									PopDir();
+									if (ConnectionRetries > 0 && (SuccesCode == 4 || SuccesCode == 2) && !bOutputContainsProject && XGETool != null)
 									{
-										System.Threading.Thread.Sleep(60000);
+										LogWarning(String.Format("{0} failure on the local connection timeout", XGETool));
+										if (ConnectionRetries < 2)
+										{
+											System.Threading.Thread.Sleep(60000);
+										}
+										ConnectionRetries--;
+										continue;
 									}
-									ConnectionRetries--;
-									continue;
-								}
-								else if (SuccesCode != 0)
-								{
-									LogWarning("{0} did not succeed *******", XGETool);
-									throw new AutomationException("Command failed (Result:{3}): {0} {1}. See logfile for details: '{2}' ",
-																	XGETool, Args, Path.GetFileName(LogFile), SuccesCode);
+									else if (SuccesCode != 0)
+									{
+										throw new AutomationException("Command failed (Result:{3}): {0} {1}. See logfile for details: '{2}' ",
+																		XGETool, Args, Path.GetFileName(LogFile), SuccesCode);
+									}
+									LogVerbose("{0} {1} Done *******", XGETool, Args);
+									break;
 								}
-								LogVerbose("{0} {1} Done *******", XGETool, Args);
 								break;
 							}
-							break;
-						}
-						catch (Exception Ex)
-						{
-							LogWarning("{0} failed on try {1}: {2}", XGETool, i + 1, Ex.ToString());
-							if (i + 1 >= Retries)
-							{
-								return false;
-							}
-							LogWarning("Deleting build products to force a rebuild.");
-							foreach (var Item in Actions)
+							catch (Exception Ex)
 							{
-								XGEDeleteBuildProducts(Item.Manifest);
+								if (i + 1 >= Retries)
+								{
+									throw;
+								}
+								LogWarning("{0} failed on try {1}, deleting products to force a rebuild: {2}", XGETool, i + 1, Ex.ToString());
+								foreach (var Item in Actions)
+								{
+									XGEDeleteBuildProducts(Item.Manifest);
+								}
 							}
 						}
 					}
diff --git a/Engine/Source/Programs/AutomationTool/BuildGraph/Agent.cs b/Engine/Source/Programs/AutomationTool/BuildGraph/Agent.cs
new file mode 100644
--- /dev/null
+++ b/Engine/Source/Programs/AutomationTool/BuildGraph/Agent.cs
@@ -0,0 +1,67 @@
+﻿using System;
+using System.Collections.Generic;
+using System.Diagnostics;
+using System.Linq;
+using System.Text;
+using System.Threading.Tasks;
+using System.Xml;
+
+namespace AutomationTool
+{
+	/// <summary>
+	/// Stores a list of nodes which can be executed on a single agent
+	/// </summary>
+	[DebuggerDisplay("{Name}")]
+	class Agent
+	{
+		/// <summary>
+		/// Name of this agent. Used for display purposes in a build system.
+		/// </summary>
+		public string Name;
+
+		/// <summary>
+		/// Array of valid agent types that these nodes may run on. When running in the build system, this determines the class of machine that should
+		/// be selected to run these nodes. The first defined agent type for this branch will be used.
+		/// </summary>
+		public string[] PossibleTypes;
+
+		/// <summary>
+		/// List of nodes in this agent group.
+		/// </summary>
+		public List<Node> Nodes = new List<Node>();
+
+		/// <summary>
+		/// Constructor
+		/// </summary>
+		/// <param name="InName">Name of this agent group</param>
+		/// <param name="InPossibleTypes">Array of valid agent types. See comment for AgentTypes member.</param>
+		public Agent(string InName, string[] InPossibleTypes)
+		{
+			Name = InName;
+			PossibleTypes = InPossibleTypes;
+		}
+
+		/// <summary>
+		/// Writes this agent group out to a file, filtering nodes by a controlling trigger
+		/// </summary>
+		/// <param name="Writer">The XML writer to output to</param>
+		/// <param name="ControllingTrigger">The controlling trigger to filter by</param>
+		public void Write(XmlWriter Writer, ManualTrigger ControllingTrigger)
+		{
+			if (Nodes.Any(x => x.ControllingTrigger == ControllingTrigger))
+			{
+				Writer.WriteStartElement("Agent");
+				Writer.WriteAttributeString("Name", Name);
+				Writer.WriteAttributeString("Type", String.Join(";", PossibleTypes));
+				foreach (Node Node in Nodes)
+				{
+					if (Node.ControllingTrigger == ControllingTrigger)
+					{
+						Node.Write(Writer);
+					}
+				}
+				Writer.WriteEndElement();
+			}
+		}
+	}
+}
diff --git a/Engine/Source/Programs/AutomationTool/BuildGraph/AgentGroup.cs b/Engine/Source/Programs/AutomationTool/BuildGraph/AgentGroup.cs
deleted file mode 100644
diff --git a/Engine/Source/Programs/AutomationTool/BuildGraph/Build.cs b/Engine/Source/Programs/AutomationTool/BuildGraph/Build.cs
deleted file mode 100644
diff --git a/Engine/Source/Programs/AutomationTool/BuildGraph/BuildGraph.Automation.csproj b/Engine/Source/Programs/AutomationTool/BuildGraph/BuildGraph.Automation.csproj
--- a/Engine/Source/Programs/AutomationTool/BuildGraph/BuildGraph.Automation.csproj
+++ b/Engine/Source/Programs/AutomationTool/BuildGraph/BuildGraph.Automation.csproj
@@ -61,14 +61,14 @@
     <Compile Include="..\..\DotNETCommon\MetaData.cs">
       <Link>Properties\MetaData.cs</Link>
     </Compile>
-    <Compile Include="AgentGroup.cs" />
+    <Compile Include="Agent.cs" />
+    <Compile Include="BuildGraph.cs" />
     <Compile Include="Graph.cs" />
     <Compile Include="Node.cs" />
     <Compile Include="Task.cs" />
     <Compile Include="JobContext.cs" />
     <Compile Include="ManualTrigger.cs" />
     <Compile Include="Condition.cs" />
-    <Compile Include="Build.cs" />
     <Compile Include="Schema.cs" />
     <Compile Include="Script.cs" />
     <Compile Include="Tasks\ChunkTask.cs" />
diff --git a/Engine/Source/Programs/AutomationTool/BuildGraph/BuildGraph.cs b/Engine/Source/Programs/AutomationTool/BuildGraph/BuildGraph.cs
new file mode 100644
--- /dev/null
+++ b/Engine/Source/Programs/AutomationTool/BuildGraph/BuildGraph.cs
@@ -0,0 +1,479 @@
+using System;
+using System.Collections.Generic;
+using System.Linq;
+using System.Text;
+using System.Threading.Tasks;
+using UnrealBuildTool;
+using BuildGraph;
+using System.Reflection;
+using System.Collections;
+using System.IO;
+
+namespace AutomationTool
+{
+	/// <summary>
+	/// Tool to execute build processes for UE4 projects, which can be run locally or in parallel across a build farm (assuming synchronization and resource allocation implemented by a separate system).
+	///
+	/// Build graphs are declared using an XML script using syntax similar to MSBuild, ANT or NAnt, and consist of the following components:
+	///
+	/// - Tasks:        Building blocks which can be executed as part of the build process. Many predefined tasks are provided ('Cook', 'Compile', 'Copy', 'Stage', 'Log', 'PakFile', etc...), and additional tasks may be 
+	///                 added be declaring classes derived from AutomationTool.CustomTask in other UAT modules. 
+	/// - Nodes:        A named sequence of tasks which are executed in order to produce outputs. Nodes may have dependencies on other nodes for their outputs before they can be executed. Declared with the 'Node' element.
+	/// - Agents:		A machine which can execute a sequence of nodes, if running as part of a build system. Has no effect when building locally. Declared with the 'Agent' element.
+	/// - Triggers:     Container for agents which should only be executed when explicitly triggered (using the -Trigger=... or -SkipTriggers command line argument). Declared with the 'Trigger' element.
+	/// - Notifiers:    Specifies email recipients for failures in one or more nodes, whether they should receive notifications on warnings, and so on.
+	/// 
+	/// Properties can be passed in to a script on the command line, or set procedurally with the &ltProperty Name="Foo" Value="Bar"/&gt; syntax. Properties referenced with the $(Property Name) notation are valid within 
+	/// all strings, and will be expanded as macros when the script is read. If a property name is not set explicitly, it defaults to the contents of an environment variable with the same name. 
+	/// Local properties, which only affect the scope of the containing XML element (node, agent, etc...) are declared with the &lt;Local Name="Foo" Value="Bar"/&gt; element, and will override a similarly named global 
+	/// property for the local property's scope.
+	///
+	/// Any elements can be conditionally defined via the "If" attribute, which follows a syntax similar to MSBuild. Literals in conditions may be quoted with single (') or double (") quotes, or an unquoted sequence of 
+	/// letters, digits and underscore characters. All literals are considered identical regardless of how they are declared, and are considered case-insensitive for comparisons (so true equals 'True', equals "TRUE"). 
+	/// Available operators are "==", "!=", "And", "Or", "!", "(...)", "Exists(...)" and "HasTrailingSlash(...)". A full grammar is written up in Condition.cs.
+	/// 
+	/// File manipulation is done using wildcards and tags. Any attribute that accepts a list of files may consist of: a Perforce-style wildcard (matching any number of "...", "*" and "?" patterns in any location), a 
+	/// full path name, or a reference to a tagged collection of files, denoted by prefixing with a '#' character. Files may be added to a tag set using the &lt;Tag&gt; Task, which also allows performing set union/difference 
+	/// style operations. Each node can declare multiple outputs in the form of a list of named tags, which other nodes can then depend on.
+	/// 
+	/// Build graphs may be executed in parallel as part build system. To do so, the initial graph configuration is generated by running with the -Export=... argument (producing a JSON file listing the nodes 
+	/// and dependencies to execute). Each participating agent should be synced to the same changelist, and UAT should be re-run with the appropriate -Node=... argument. Outputs from different nodes are transferred between 
+	/// agents via shared storage, typically a network share, the path to which can be specified on the command line using the -SharedStorageDir=... argument. Note that the allocation of machines, and coordination between 
+	/// them, is assumed to be managed by an external system based on the contents of the script generated by -Export=....
+	/// 
+	/// A schema for the known set of tasks can be generated by running UAT with the -Schema=... option. Generating a schema and referencing it from a BuildGraph script allows Visual Studio to validate and auto-complete 
+	/// elements as you type.
+	/// </summary>
+	[Help("Tool for creating extensible build processes in UE4 which can be run locally or in parallel across a build farm.")]
+	[Help("Script=<FileName>", "Path to the script describing the graph")]
+	[Help("Target=<Name>", "Name of the node or output tag to be built")]
+	[Help("Schema=<FileName>", "Generate a schema describing valid script documents, including all the known tasks")]
+	[Help("Set:<Property>=<Value>", "Sets a named property to the given value")]
+	[Help("Clean", "Cleans all cached state of completed build nodes before running")]
+	[Help("CleanNode=<Name>[+<Name>...]", "Cleans just the given nodes before running")]
+	[Help("ListOnly", "Shows the contents of the preprocessed graph, but does not execute it")]
+	[Help("ShowDeps", "Show node dependencies in the graph output")]
+	[Help("ShowNotifications", "Show notifications that will be sent for each node in the output")]
+	[Help("Trigger=<Name>[+<Name>...]", "Activates the given triggers, including all the nodes behind them in the graph")]
+	[Help("SkipTriggers", "Activate all triggers")]
+	[Help("Preprocess=<FileName>", "Writes the preprocessed graph to the given file")]
+	[Help("Export=<FileName>", "Exports a JSON file containing the preprocessed build graph, for use as part of a build system")]
+	[Help("PublicTasksOnly", "Only include built-in tasks in the schema, excluding any other UAT modules")]
+	[Help("SharedStorageDir=<DirName>", "Sets the directory to use to transfer build products between agents in a build farm")]
+	[Help("SingleNode=<Name>", "Run only the given node. Intended for use on a build system after running with -Export.")]
+	[Help("WriteToSharedStorage", "Allow writing to shared storage. If not set, but -SharedStorageDir is specified, build products will read but not written")]
+	public class BuildGraph : BuildCommand
+	{
+		/// <summary>
+		/// Main entry point for the BuildGraph command
+		/// </summary>
+		public override ExitCode Execute()
+		{
+			// Parse the command line parameters
+			string ScriptFileName = ParseParamValue("Script", null);
+			if(ScriptFileName == null)
+			{
+				LogError("Missing -Script= parameter for BuildGraph");
+				return ExitCode.Error_Unknown;
+			}
+
+			string TargetNames = ParseParamValue("Target", null);
+			if(TargetNames == null)
+			{
+				LogError("Missing -Target= parameter for BuildGraph");
+				return ExitCode.Error_Unknown;
+			}
+
+			string SchemaFileName = ParseParamValue("Schema", null);
+			string ExportFileName = ParseParamValue("Export", null);
+			string PreprocessedFileName = ParseParamValue("Preprocess", null);
+
+			string SharedStorageDir = ParseParamValue("SharedStorageDir", null);
+			string SingleNodeName = ParseParamValue("SingleNode", null);
+			string[] TriggerNames = ParseParamValue("Trigger", "").Split(new char[]{ '+' }, StringSplitOptions.RemoveEmptyEntries).ToArray();
+			bool bSkipTriggers = ParseParam("SkipTriggers");
+			bool bClean = ParseParam("Clean");
+			bool bListOnly = ParseParam("ListOnly");
+			bool bWriteToSharedStorage = ParseParam("WriteToSharedStorage") || CommandUtils.IsBuildMachine;
+			bool bPublicTasksOnly = ParseParam("PublicTasksOnly");
+
+			GraphPrintOptions PrintOptions = 0;
+			if(ParseParam("ShowDeps"))
+			{
+				PrintOptions |= GraphPrintOptions.ShowDependencies;
+			}
+			if(ParseParam("ShowNotifications"))
+			{
+				PrintOptions |= GraphPrintOptions.ShowNotifications;
+			}
+
+			// Parse any specific nodes to clean
+			List<string> CleanNodes = new List<string>();
+			foreach(string NodeList in ParseParamValues("CleanNode"))
+			{
+				foreach(string NodeName in NodeList.Split('+'))
+				{
+					CleanNodes.Add(NodeName);
+				}
+			}
+
+			// Read any environment variables
+			Dictionary<string, string> DefaultProperties = new Dictionary<string,string>(StringComparer.InvariantCultureIgnoreCase);
+			foreach(DictionaryEntry Entry in Environment.GetEnvironmentVariables())
+			{
+				DefaultProperties[Entry.Key.ToString()] = Entry.Value.ToString();
+			}
+
+			// Add any additional custom parameters from the command line (of the form -Set:X=Y)
+			foreach(string Param in Params)
+			{
+				const string Prefix = "set:";
+				if(Param.StartsWith(Prefix, StringComparison.InvariantCultureIgnoreCase))
+				{
+					int EqualsIdx = Param.IndexOf('=');
+					if(EqualsIdx >= 0)
+					{
+						DefaultProperties[Param.Substring(Prefix.Length, EqualsIdx - Prefix.Length)] = Param.Substring(EqualsIdx + 1);
+					}
+					else
+					{
+						LogWarning("Missing value for '{0}'", Param.Substring(Prefix.Length));
+					}
+				}
+			}
+
+			// Set up the standard properties which build scripts might need
+			DefaultProperties["Branch"] = P4Enabled? P4Env.BuildRootP4 : "Unknown";
+			DefaultProperties["EscapedBranch"] = P4Enabled? P4Env.BuildRootEscaped : "Unknown";
+			DefaultProperties["Change"] = P4Enabled? P4Env.Changelist.ToString() : "0";
+			DefaultProperties["RootDir"] = CommandUtils.RootDirectory.FullName;
+			DefaultProperties["IsBuildMachine"] = IsBuildMachine? "true" : "false";
+			DefaultProperties["HostPlatform"] = HostPlatform.Current.HostEditorPlatform.ToString();
+
+			// Find all the tasks from the loaded assemblies
+			Dictionary<string, ScriptTask> NameToTask = new Dictionary<string,ScriptTask>();
+			if(!FindAvailableTasks(NameToTask, bPublicTasksOnly))
+			{
+				return ExitCode.Error_Unknown;
+			}
+
+			// Create a schema for the given tasks
+			ScriptSchema Schema = new ScriptSchema(NameToTask);
+			if(SchemaFileName != null)
+			{
+				Schema.Export(new FileReference(SchemaFileName));
+			}
+
+			// Read the script from disk
+			Graph Graph;
+			if(!ScriptReader.TryRead(new FileReference(ScriptFileName), DefaultProperties, Schema, out Graph))
+			{
+				return ExitCode.Error_Unknown;
+			}
+
+			// Create the temp storage handler
+			DirectoryReference RootDir = new DirectoryReference(CommandUtils.CmdEnv.LocalRoot);
+			TempStorage Storage = new TempStorage(RootDir, DirectoryReference.Combine(RootDir, "Engine", "Saved", "BuildGraph"), (SharedStorageDir == null)? null : new DirectoryReference(SharedStorageDir), bWriteToSharedStorage);
+			if(bClean)
+			{
+				Storage.CleanLocal();
+			}
+			foreach(string CleanNode in CleanNodes)
+			{
+				Storage.CleanLocalNode(CleanNode);
+			}
+
+			// Convert the supplied target references into nodes 
+			HashSet<Node> TargetNodes = new HashSet<Node>();
+			foreach(string TargetName in TargetNames.Split(new char[]{ '+' }, StringSplitOptions.RemoveEmptyEntries).Select(x => x.Trim()))
+			{
+				Node[] Nodes;
+				if(!Graph.TryResolveReference(TargetName, out Nodes))
+				{
+					LogError("Target '{0}' is not in graph", TargetName);
+					return ExitCode.Error_Unknown;
+				}
+				TargetNodes.UnionWith(Nodes);
+			}
+
+			// Cull the graph to include only those nodes
+			Graph.Select(TargetNodes);
+
+			// Write out the preprocessed script
+			if (PreprocessedFileName != null)
+			{
+				Graph.Write(new FileReference(PreprocessedFileName), (SchemaFileName != null)? new FileReference(SchemaFileName) : null);
+			}
+
+			// Find the triggers which are explicitly activated, and all of its upstream triggers.
+			HashSet<ManualTrigger> Triggers = new HashSet<ManualTrigger>();
+			foreach(string TriggerName in TriggerNames)
+			{
+				ManualTrigger Trigger;
+				if(!Graph.NameToTrigger.TryGetValue(TriggerName, out Trigger))
+				{
+					LogError("Couldn't find trigger '{0}'", TriggerName);
+					return ExitCode.Error_Unknown;
+				}
+				while(Trigger != null)
+				{
+					Triggers.Add(Trigger);
+					Trigger = Trigger.Parent;
+				}
+			}
+			if(bSkipTriggers)
+			{
+				Triggers.UnionWith(Graph.NameToTrigger.Values);
+			}
+
+			// If we're just building a single node, find it 
+			Node SingleNode = null;
+			if(SingleNodeName != null && !Graph.NameToNode.TryGetValue(SingleNodeName, out SingleNode))
+			{
+				LogError("Node '{0}' is not in the trimmed graph", SingleNodeName);
+				return ExitCode.Error_Unknown;
+			}
+
+			// Print out all the diagnostic messages which still apply, unless we're running a step as part of a build system. 
+			if(SingleNode == null)
+			{
+				IEnumerable<GraphDiagnostic> Diagnostics = Graph.Diagnostics.Where(x => x.EnclosingTrigger == null || Triggers.Contains(x.EnclosingTrigger));
+				foreach(GraphDiagnostic Diagnostic in Diagnostics)
+				{
+					if(Diagnostic.EventType == LogEventType.Warning)
+					{
+						CommandUtils.LogWarning(Diagnostic.Message);
+					}
+					else
+					{
+						CommandUtils.LogError(Diagnostic.Message);
+					}
+				}
+				if(Diagnostics.Any(x => x.EventType == LogEventType.Error))
+				{
+					return ExitCode.Error_Unknown;
+				}
+			}
+
+			// Execute the command
+			if(bListOnly)
+			{ 
+				HashSet<Node> CompletedNodes = FindCompletedNodes(Graph, Storage);
+				Graph.Print(CompletedNodes, PrintOptions);
+			}
+			else if(ExportFileName != null)
+			{
+				HashSet<Node> CompletedNodes = FindCompletedNodes(Graph, Storage);
+				Graph.Print(CompletedNodes, PrintOptions);
+				Graph.Export(new FileReference(ExportFileName), Triggers, CompletedNodes);
+			}
+			else if(SingleNode != null)
+			{
+				if(!BuildNode(new JobContext(this), Graph, SingleNode, Storage, bWithBanner: true))
+				{
+					return ExitCode.Error_Unknown;
+				}
+			}
+			else
+			{
+				if(!BuildAllNodes(new JobContext(this), Graph, Storage))
+				{
+					return ExitCode.Error_Unknown;
+				}
+			}
+			return ExitCode.Success;
+		}
+
+		/// <summary>
+		/// Find all the tasks which are available from the loaded assemblies
+		/// </summary>
+		/// <param name="TaskNameToReflectionInfo">Mapping from task name to information about how to serialize it</param>
+		/// <param name="bPublicTasksOnly">Whether to include just public tasks, or all the tasks in any loaded assemblies</param>
+		static bool FindAvailableTasks(Dictionary<string, ScriptTask> NameToTask, bool bPublicTasksOnly)
+		{
+			Assembly[] LoadedAssemblies = AppDomain.CurrentDomain.GetAssemblies();
+			if(bPublicTasksOnly)
+			{
+				LoadedAssemblies = LoadedAssemblies.Where(x => IsPublicAssembly(new FileReference(x.Location))).ToArray();
+			}
+			foreach (Assembly LoadedAssembly in LoadedAssemblies)
+			{
+				Type[] Types = LoadedAssembly.GetTypes();
+				foreach(Type Type in Types)
+				{
+					foreach(TaskElementAttribute ElementAttribute in Type.GetCustomAttributes<TaskElementAttribute>())
+					{
+						if(!Type.IsSubclassOf(typeof(CustomTask)))
+						{
+							CommandUtils.LogError("Class '{0}' has TaskElementAttribute, but is not derived from 'Task'", Type.Name);
+							return false;
+						}
+						if(NameToTask.ContainsKey(ElementAttribute.Name))
+						{
+							CommandUtils.LogError("Found multiple handlers for task elements called '{0}'", ElementAttribute.Name);
+							return false;
+						}
+						NameToTask.Add(ElementAttribute.Name, new ScriptTask(ElementAttribute.Name, Type, ElementAttribute.ParametersType));
+					}
+				}
+			}
+			return true;
+		}
+
+		/// <summary>
+		/// Checks whether the given assembly is a publically distributed engine assembly.
+		/// </summary>
+		/// <param name="File">Assembly location</param>
+		/// <returns>True if the assembly is distributed publically</returns>
+		static bool IsPublicAssembly(FileReference File)
+		{
+			DirectoryReference EngineDirectory = UnrealBuildTool.UnrealBuildTool.EngineDirectory;
+			if(File.IsUnderDirectory(EngineDirectory))
+			{
+				string[] PathFragments = File.MakeRelativeTo(EngineDirectory).Split(Path.DirectorySeparatorChar, Path.AltDirectorySeparatorChar);
+				if(PathFragments.All(x => !x.Equals("NotForLicensees", StringComparison.InvariantCultureIgnoreCase) && !x.Equals("NoRedist", StringComparison.InvariantCultureIgnoreCase)))
+				{
+					return true;
+				}
+			}
+			return false;
+		}
+
+		/// <summary>
+		/// Find all the nodes in the graph which are already completed
+		/// </summary>
+		/// <param name="Graph">The graph instance</param>
+		/// <param name="Storage">The temp storage backend which stores the shared state</param>
+		HashSet<Node> FindCompletedNodes(Graph Graph, TempStorage Storage)
+		{
+			HashSet<Node> CompletedNodes = new HashSet<Node>();
+			foreach(Node Node in Graph.Agents.SelectMany(x => x.Nodes))
+			{
+				if(Storage.IsComplete(Node.Name))
+				{
+					CompletedNodes.Add(Node);
+				}
+			}
+			return CompletedNodes;
+		}
+
+		/// <summary>
+		/// Builds all the nodes in the graph
+		/// </summary>
+		/// <param name="Job">Information about the current job</param>
+		/// <param name="Graph">The graph instance</param>
+		/// <returns>True if everything built successfully</returns>
+		bool BuildAllNodes(JobContext Job, Graph Graph, TempStorage Storage)
+		{
+			// Build a flat list of nodes to execute, in order
+			Node[] NodesToExecute = Graph.Agents.SelectMany(x => x.Nodes).ToArray();
+
+			// Check the integrity of any local nodes that have been completed. It's common to run formal builds locally between regular development builds, so we may have 
+			// stale local state. Rather than failing later, detect and clean them up now.
+			HashSet<Node> CleanedNodes = new HashSet<Node>();
+			foreach(Node NodeToExecute in NodesToExecute)
+			{
+				if(NodeToExecute.InputDependencies.Any(x => CleanedNodes.Contains(x)) || !Storage.CheckLocalIntegrity(NodeToExecute.Name, NodeToExecute.Outputs.Select(x => x.Name)))
+				{
+					Storage.CleanLocalNode(NodeToExecute.Name);
+					CleanedNodes.Add(NodeToExecute);
+				}
+			}
+
+			// Execute them in order
+			int NodeIdx = 0;
+			foreach(Node NodeToExecute in NodesToExecute)
+			{
+				Log("****** [{0}/{1}] {2}", ++NodeIdx, NodesToExecute.Length, NodeToExecute.Name);
+				if(!Storage.IsComplete(NodeToExecute.Name))
+				{
+					Log("");
+					if(!BuildNode(Job, Graph, NodeToExecute, Storage, false))
+					{
+						return false;
+					} 
+					Log("");
+				}
+			}
+			return true;
+		}
+
+		/// <summary>
+		/// Build a node
+		/// </summary>
+		/// <param name="Job">Information about the current job</param>
+		/// <param name="Graph">The graph to which the node belongs. Used to determine which outputs need to be transferred to temp storage.</param>
+		/// <param name="Node">The node to build</param>
+		/// <returns>True if the node built successfully, false otherwise.</returns>
+		bool BuildNode(JobContext Job, Graph Graph, Node Node, TempStorage Storage, bool bWithBanner)
+		{
+			// Create a mapping from tag name to the files it contains, and seed it with invalid entries for everything in the graph
+			Dictionary<string, HashSet<FileReference>> TagNameToFileSet = new Dictionary<string,HashSet<FileReference>>();
+			foreach(NodeOutput Output in Graph.Agents.SelectMany(x => x.Nodes).SelectMany(x => x.Outputs))
+			{
+				TagNameToFileSet[Output.Name] = null;
+			}
+
+			// Fetch all the input dependencies for this node, and fill in the tag names with those files
+			DirectoryReference RootDir = new DirectoryReference(CommandUtils.CmdEnv.LocalRoot);
+			foreach(NodeOutput Input in Node.Inputs)
+			{
+				TempStorageManifest Manifest = Storage.Retreive(Input.ProducingNode.Name, Input.Name);
+				TagNameToFileSet[Input.Name] = new HashSet<FileReference>(Manifest.Files.Select(x => x.ToFileReference(RootDir)));
+			}
+
+			// Add placeholder outputs for the current node
+			foreach(NodeOutput Output in Node.Outputs)
+			{
+				TagNameToFileSet[Output.Name] = new HashSet<FileReference>();
+			}
+
+			// Execute the node
+			if(bWithBanner)
+			{
+				Console.WriteLine();
+				CommandUtils.Log("========== Starting: {0} ==========", Node.Name);
+			}
+			if(!Node.Build(Job, TagNameToFileSet))
+			{
+				return false;
+			}
+			if(bWithBanner)
+			{
+				CommandUtils.Log("========== Finished: {0} ==========", Node.Name);
+				Console.WriteLine();
+			}
+
+			// Determine all the outputs which are required to be copied to temp storage (because they're referenced by nodes in another agent)
+			HashSet<NodeOutput> ReferencedOutputs = new HashSet<NodeOutput>();
+			foreach(Agent Agent in Graph.Agents)
+			{
+				bool bSameAgent = Agent.Nodes.Contains(Node);
+				foreach(Node OtherNode in Agent.Nodes)
+				{
+					if(!bSameAgent || Node.ControllingTrigger != OtherNode.ControllingTrigger)
+					{
+						ReferencedOutputs.UnionWith(OtherNode.Inputs);
+					}
+				}
+			}
+
+			// Publish all the outputs
+			foreach(NodeOutput Output in Node.Outputs)
+			{
+				Storage.Archive(Node.Name, Output.Name, TagNameToFileSet[Output.Name].ToArray(), ReferencedOutputs.Contains(Output));
+			}
+
+			// Mark the node as succeeded
+			Storage.MarkAsComplete(Node.Name);
+			return true;
+		}
+	}
+
+	/// <summary>
+	/// Legacy command name for compatibility.
+	/// </summary>
+	public class Build : BuildGraph
+	{
+	}
+}
diff --git a/Engine/Source/Programs/AutomationTool/BuildGraph/Graph.cs b/Engine/Source/Programs/AutomationTool/BuildGraph/Graph.cs
--- a/Engine/Source/Programs/AutomationTool/BuildGraph/Graph.cs
+++ b/Engine/Source/Programs/AutomationTool/BuildGraph/Graph.cs
@@ -51,9 +51,9 @@ namespace AutomationTool
 		public Node EnclosingNode;
 
 		/// <summary>
-		/// The agent group that this diagnostic is declared in. If the entire group is culled from the graph, the message will not be displayed.
+		/// The agent that this diagnostic is declared in. If the entire agent is culled from the graph, the message will not be displayed.
 		/// </summary>
-		public AgentGroup EnclosingGroup;
+		public Agent EnclosingAgent;
 
 		/// <summary>
 		/// The trigger that this diagnostic is declared in. If this trigger is not being run, the message will not be displayed.
@@ -67,15 +67,20 @@ namespace AutomationTool
 	class Graph
 	{
 		/// <summary>
-		/// List of nodes grouped by the agent they will be executed on
+		/// List of agents containing nodes to execute
 		/// </summary>
-		public List<AgentGroup> Groups = new List<AgentGroup>();
+		public List<Agent> Agents = new List<Agent>();
 
 		/// <summary>
 		/// All manual triggers that are part of this graph
 		/// </summary>
 		public Dictionary<string, ManualTrigger> NameToTrigger = new Dictionary<string, ManualTrigger>(StringComparer.InvariantCultureIgnoreCase);
 
+		/// <summary>
+		/// Mapping from name to agent
+		/// </summary>
+		public Dictionary<string, Agent> NameToAgent = new Dictionary<string, Agent>(StringComparer.InvariantCultureIgnoreCase);
+
 		/// <summary>
 		/// Mapping of names to the corresponding node.
 		/// </summary>
@@ -217,13 +222,13 @@ namespace AutomationTool
 			}
 
 			// Remove all the nodes which are not marked to be kept
-			foreach(AgentGroup Group in Groups)
+			foreach(Agent Agent in Agents)
 			{
-				Group.Nodes = Group.Nodes.Where(x => RetainNodes.Contains(x)).ToList();
+				Agent.Nodes = Agent.Nodes.Where(x => RetainNodes.Contains(x)).ToList();
 			}
 
-			// Remove all the empty groups
-			Groups.RemoveAll(x => x.Nodes.Count == 0);
+			// Remove all the empty agents
+			Agents.RemoveAll(x => x.Nodes.Count == 0);
 
 			// Remove all the order dependencies which are no longer part of the graph. Since we don't need to build them, we don't need to wait for them
 			foreach(Node Node in RetainNodes)
@@ -238,7 +243,53 @@ namespace AutomationTool
 			AggregateNameToNodes = AggregateNameToNodes.Where(x => x.Value.All(y => RetainNodes.Contains(y))).ToDictionary(Pair => Pair.Key, Pair => Pair.Value);
 
 			// Remove any diagnostics which are no longer part of the graph
-			Diagnostics.RemoveAll(x => (x.EnclosingNode != null && !RetainNodes.Contains(x.EnclosingNode)) || (x.EnclosingGroup != null && !Groups.Contains(x.EnclosingGroup)));
+			Diagnostics.RemoveAll(x => (x.EnclosingNode != null && !RetainNodes.Contains(x.EnclosingNode)) || (x.EnclosingAgent != null && !Agents.Contains(x.EnclosingAgent)));
+		}
+
+		/// <summary>
+		/// Writes a preprocessed build graph to a script file
+		/// </summary>
+		/// <param name="File">The file to load</param>
+		public void Write(FileReference File, FileReference SchemaFile)
+		{
+			XmlWriterSettings Settings = new XmlWriterSettings();
+			Settings.Indent = true;
+			Settings.IndentChars = "\t";
+
+			using (XmlWriter Writer = XmlWriter.Create(File.FullName, Settings))
+			{
+				Writer.WriteStartElement("BuildGraph", "http://www.epicgames.com/BuildGraph");
+
+				if (SchemaFile != null)
+				{
+					Writer.WriteAttributeString("schemaLocation", "http://www.w3.org/2001/XMLSchema-instance", "http://www.epicgames.com/BuildGraph " + SchemaFile.MakeRelativeTo(File.Directory));
+				}
+
+				foreach (Agent Agent in Agents)
+				{
+					Agent.Write(Writer, null);
+				}
+
+				foreach (ManualTrigger ControllingTrigger in Agents.SelectMany(x => x.Nodes).Where(x => x.ControllingTrigger != null).Select(x => x.ControllingTrigger).Distinct())
+				{
+					Writer.WriteStartElement("Trigger");
+					Writer.WriteAttributeString("Name", ControllingTrigger.QualifiedName);
+					foreach (Agent Agent in Agents)
+					{
+						Agent.Write(Writer, ControllingTrigger);
+					}
+					Writer.WriteEndElement();
+				}
+
+				foreach (KeyValuePair<string, Node[]> Aggregate in AggregateNameToNodes)
+				{
+					Writer.WriteStartElement("Aggregate");
+					Writer.WriteAttributeString("Name", Aggregate.Key);
+					Writer.WriteAttributeString("Requires", String.Join(";", Aggregate.Value.Select(x => x.Name)));
+					Writer.WriteEndElement();
+				}
+				Writer.WriteEndElement();
+			}
 		}
 
 		/// <summary>
@@ -251,7 +302,7 @@ namespace AutomationTool
 		{
 			// Find all the nodes which we're actually going to execute. We'll use this to filter the graph.
 			HashSet<Node> NodesToExecute = new HashSet<Node>();
-			foreach(Node Node in Groups.SelectMany(x => x.Nodes))
+			foreach(Node Node in Agents.SelectMany(x => x.Nodes))
 			{
 				if(!CompletedNodes.Contains(Node))
 				{
@@ -267,17 +318,17 @@ namespace AutomationTool
 			{
 				JsonWriter.WriteObjectStart();
 
-				// Write all the agent groups
+				// Write all the agents
 				JsonWriter.WriteArrayStart("Groups");
-				foreach(AgentGroup Group in Groups)
+				foreach(Agent Agent in Agents)
 				{
-					Node[] Nodes = Group.Nodes.Where(x => NodesToExecute.Contains(x)).ToArray();
+					Node[] Nodes = Agent.Nodes.Where(x => NodesToExecute.Contains(x)).ToArray();
 					if(Nodes.Length > 0)
 					{
 						JsonWriter.WriteObjectStart();
-						JsonWriter.WriteValue("Name", Group.Name);
+						JsonWriter.WriteValue("Name", Agent.Name);
 						JsonWriter.WriteArrayStart("Agent Types");
-						foreach(string AgentType in Group.PossibleTypes)
+						foreach(string AgentType in Agent.PossibleTypes)
 						{
 							JsonWriter.WriteValue(AgentType);
 						}
@@ -309,7 +360,7 @@ namespace AutomationTool
 					{
 						// Find all the nodes that this trigger is dependent on
 						HashSet<Node> Dependencies = new HashSet<Node>();
-						foreach(Node Node in Groups.SelectMany(x => x.Nodes))
+						foreach(Node Node in Agents.SelectMany(x => x.Nodes))
 						{
 							for(ManualTrigger ControllingTrigger = Node.ControllingTrigger; ControllingTrigger != null; ControllingTrigger = ControllingTrigger.Parent)
 							{
@@ -331,7 +382,7 @@ namespace AutomationTool
 						// Write out the object
 						JsonWriter.WriteObjectStart();
 						JsonWriter.WriteValue("Name", Trigger.Name);
-						JsonWriter.WriteValue("AllDependencies", String.Join(";", Groups.SelectMany(x => x.Nodes).Where(x => Dependencies.Contains(x)).Select(x => x.Name)));
+						JsonWriter.WriteValue("AllDependencies", String.Join(";", Agents.SelectMany(x => x.Nodes).Where(x => Dependencies.Contains(x)).Select(x => x.Name)));
 						JsonWriter.WriteValue("DirectDependencies", String.Join(";", Dependencies.Where(x => DirectDependencies.Contains(x)).Select(x => x.Name)));
 						JsonWriter.WriteValue("Notify", String.Join(";", Trigger.NotifyUsers));
 						JsonWriter.WriteObjectEnd();
@@ -360,18 +411,18 @@ namespace AutomationTool
 			foreach(ManualTrigger Trigger in AllTriggers)
 			{
 				// Filter everything by this trigger
-				Dictionary<AgentGroup, Node[]> FilteredGroupToNodes = new Dictionary<AgentGroup,Node[]>();
-				foreach(AgentGroup Group in Groups)
+				Dictionary<Agent, Node[]> FilteredAgentToNodes = new Dictionary<Agent,Node[]>();
+				foreach(Agent Agent in Agents)
 				{
-					Node[] Nodes = Group.Nodes.Where(x => x.ControllingTrigger == Trigger).ToArray();
+					Node[] Nodes = Agent.Nodes.Where(x => x.ControllingTrigger == Trigger).ToArray();
 					if(Nodes.Length > 0)
 					{
-						FilteredGroupToNodes[Group] = Nodes;
+						FilteredAgentToNodes[Agent] = Nodes;
 					}
 				}
 
 				// Skip this trigger if there's nothing to display
-				if(FilteredGroupToNodes.Count == 0)
+				if(FilteredAgentToNodes.Count == 0)
 				{
 					continue;
 				}
@@ -386,13 +437,13 @@ namespace AutomationTool
 					}
 				}
 
-				// Output all the groups for this trigger
-				foreach(AgentGroup Group in Groups)
+				// Output all the agents for this trigger
+				foreach(Agent Agent in Agents)
 				{
 					Node[] Nodes;
-					if(FilteredGroupToNodes.TryGetValue(Group, out Nodes))
+					if(FilteredAgentToNodes.TryGetValue(Agent, out Nodes))
 					{
-						CommandUtils.Log("        Group: {0} ({1})", Group.Name, String.Join(";", Group.PossibleTypes));
+						CommandUtils.Log("        Agent: {0} ({1})", Agent.Name, String.Join(";", Agent.PossibleTypes));
 						foreach(Node Node in Nodes)
 						{
 							CommandUtils.Log("            Node: {0}{1}", Node.Name, CompletedNodes.Contains(Node)? " (completed)" : "");
diff --git a/Engine/Source/Programs/AutomationTool/BuildGraph/Node.cs b/Engine/Source/Programs/AutomationTool/BuildGraph/Node.cs
--- a/Engine/Source/Programs/AutomationTool/BuildGraph/Node.cs
+++ b/Engine/Source/Programs/AutomationTool/BuildGraph/Node.cs
@@ -5,6 +5,7 @@ using System.Text;
 using System.Threading.Tasks;
 using UnrealBuildTool;
 using AutomationTool;
+using System.Xml;
 
 namespace AutomationTool
 {
@@ -229,6 +230,46 @@ namespace AutomationTool
 			return Trigger == null || ControllingTrigger == Trigger || (ControllingTrigger != null && ControllingTrigger.IsDownstreamFrom(Trigger));
 		}
 
+		/// <summary>
+		/// Write this node to an XML writer
+		/// </summary>
+		/// <param name="Writer">The writer to output the node to</param>
+		public void Write(XmlWriter Writer)
+		{
+			Writer.WriteStartElement("Node");
+			Writer.WriteAttributeString("Name", Name);
+
+			Node[] RequireNodes = Inputs.Where(x => x.Name == x.ProducingNode.Name).Select(x => x.ProducingNode).ToArray();
+			string[] RequireNames = RequireNodes.Select(x => x.Name).Union(Inputs.Where(x => !RequireNodes.Contains(x.ProducingNode)).Select(x => "#" + x.Name)).ToArray();
+			if (RequireNames.Length > 0)
+			{
+				Writer.WriteAttributeString("Requires", String.Join(";", RequireNames));
+			}
+
+			string[] ProducesNames = Outputs.Where(x => x.Name != Name).Select(x => "#" + x.Name).ToArray();
+			if (ProducesNames.Length > 0)
+			{
+				Writer.WriteAttributeString("Produces", String.Join(";", ProducesNames));
+			}
+
+			string[] AfterNames = GetDirectOrderDependencies().Except(InputDependencies).Select(x => x.Name).ToArray();
+			if (AfterNames.Length > 0)
+			{
+				Writer.WriteAttributeString("After", String.Join(";", AfterNames));
+			}
+
+			if (!bNotifyOnWarnings)
+			{
+				Writer.WriteAttributeString("NotifyOnWarnings", bNotifyOnWarnings.ToString());
+			}
+
+			foreach (CustomTask Task in Tasks)
+			{
+				Task.Write(Writer);
+			}
+			Writer.WriteEndElement();
+		}
+
 		/// <summary>
 		/// Returns the name of this node
 		/// </summary>
diff --git a/Engine/Source/Programs/AutomationTool/BuildGraph/Schema.cs b/Engine/Source/Programs/AutomationTool/BuildGraph/Schema.cs
--- a/Engine/Source/Programs/AutomationTool/BuildGraph/Schema.cs
+++ b/Engine/Source/Programs/AutomationTool/BuildGraph/Schema.cs
@@ -431,22 +431,22 @@ namespace AutomationTool
 		/// <returns>Type definition for an agent</returns>
 		static XmlSchemaType CreateAgentType()
 		{
-			XmlSchemaChoice GroupChoice = new XmlSchemaChoice();
-			GroupChoice.MinOccurs = 0;
-			GroupChoice.MaxOccursString = "unbounded";
-			GroupChoice.Items.Add(CreateSchemaElement("Property", ScriptSchemaStandardType.Property));
-			GroupChoice.Items.Add(CreateSchemaElement("Local", ScriptSchemaStandardType.Local));
-			GroupChoice.Items.Add(CreateSchemaElement("Node", ScriptSchemaStandardType.Node));
-			GroupChoice.Items.Add(CreateSchemaElement("Warning", ScriptSchemaStandardType.Warning));
-			GroupChoice.Items.Add(CreateSchemaElement("Error", ScriptSchemaStandardType.Error));
-
-			XmlSchemaComplexType GroupType = new XmlSchemaComplexType();
-			GroupType.Name = GetTypeName(ScriptSchemaStandardType.Agent);
-			GroupType.Particle = GroupChoice;
-			GroupType.Attributes.Add(CreateSchemaAttribute("Name", StringTypeName, XmlSchemaUse.Required));
-			GroupType.Attributes.Add(CreateSchemaAttribute("Type", ScriptSchemaStandardType.NameList, XmlSchemaUse.Optional));
-			GroupType.Attributes.Add(CreateSchemaAttribute("If", ScriptSchemaStandardType.BalancedString, XmlSchemaUse.Optional));
-			return GroupType;
+			XmlSchemaChoice AgentChoice = new XmlSchemaChoice();
+			AgentChoice.MinOccurs = 0;
+			AgentChoice.MaxOccursString = "unbounded";
+			AgentChoice.Items.Add(CreateSchemaElement("Property", ScriptSchemaStandardType.Property));
+			AgentChoice.Items.Add(CreateSchemaElement("Local", ScriptSchemaStandardType.Local));
+			AgentChoice.Items.Add(CreateSchemaElement("Node", ScriptSchemaStandardType.Node));
+			AgentChoice.Items.Add(CreateSchemaElement("Warning", ScriptSchemaStandardType.Warning));
+			AgentChoice.Items.Add(CreateSchemaElement("Error", ScriptSchemaStandardType.Error));
+
+			XmlSchemaComplexType AgentType = new XmlSchemaComplexType();
+			AgentType.Name = GetTypeName(ScriptSchemaStandardType.Agent);
+			AgentType.Particle = AgentChoice;
+			AgentType.Attributes.Add(CreateSchemaAttribute("Name", StringTypeName, XmlSchemaUse.Required));
+			AgentType.Attributes.Add(CreateSchemaAttribute("Type", ScriptSchemaStandardType.NameList, XmlSchemaUse.Optional));
+			AgentType.Attributes.Add(CreateSchemaAttribute("If", ScriptSchemaStandardType.BalancedString, XmlSchemaUse.Optional));
+			return AgentType;
 		}
 
 		/// <summary>
diff --git a/Engine/Source/Programs/AutomationTool/BuildGraph/Script.cs b/Engine/Source/Programs/AutomationTool/BuildGraph/Script.cs
--- a/Engine/Source/Programs/AutomationTool/BuildGraph/Script.cs
+++ b/Engine/Source/Programs/AutomationTool/BuildGraph/Script.cs
@@ -60,12 +60,12 @@ namespace AutomationTool
 		{
 			ScriptDocument Document = new ScriptDocument(File);
 
-            XmlReaderSettings Settings = new XmlReaderSettings();
-            Settings.Schemas.Add(Schema.CompiledSchema);
-            Settings.ValidationType = ValidationType.Schema;
+			XmlReaderSettings Settings = new XmlReaderSettings();
+			Settings.Schemas.Add(Schema.CompiledSchema);
+			Settings.ValidationType = ValidationType.Schema;
 			Settings.ValidationEventHandler += Document.ValidationEvent;
-	
-			using(XmlReader Reader = XmlReader.Create(File.FullName, Settings))
+
+			using (XmlReader Reader = XmlReader.Create(File.FullName, Settings))
 			{
 				// Read the document
 				Document.LineInfo = (IXmlLineInfo)Reader;
@@ -73,9 +73,9 @@ namespace AutomationTool
 				{
 					Document.Load(Reader);
 				}
-				catch(XmlException Ex)
+				catch (XmlException Ex)
 				{
-					if(!Document.bHasErrors)
+					if (!Document.bHasErrors)
 					{
 						CommandUtils.LogError("{0}", Ex.Message);
 						Document.bHasErrors = true;
@@ -83,20 +83,20 @@ namespace AutomationTool
 				}
 
 				// If we hit any errors while parsing
-				if(Document.bHasErrors)
+				if (Document.bHasErrors)
 				{
 					OutDocument = null;
 					return false;
 				}
 
 				// Check that the root element is valid. If not, we didn't actually validate against the schema.
-				if(Document.DocumentElement.Name != ScriptSchema.RootElementName)
+				if (Document.DocumentElement.Name != ScriptSchema.RootElementName)
 				{
 					CommandUtils.LogError("Script does not have a root element called '{0}'", ScriptSchema.RootElementName);
 					OutDocument = null;
 					return false;
 				}
-				if(Document.DocumentElement.NamespaceURI != ScriptSchema.NamespaceURI)
+				if (Document.DocumentElement.NamespaceURI != ScriptSchema.NamespaceURI)
 				{
 					CommandUtils.LogError("Script root element is not in the '{0}' namespace (add the xmlns=\"{0}\" attribute)", ScriptSchema.NamespaceURI);
 					OutDocument = null;
@@ -115,7 +115,7 @@ namespace AutomationTool
 		/// <param name="Args">Standard argument for ValidationEventHandler</param>
 		void ValidationEvent(object Sender, ValidationEventArgs Args)
 		{
-			if(Args.Severity == XmlSeverityType.Warning)
+			if (Args.Severity == XmlSeverityType.Warning)
 			{
 				CommandUtils.LogWarning("{0}({1}): {2}", File.FullName, Args.Exception.LineNumber, Args.Message);
 			}
@@ -145,7 +145,7 @@ namespace AutomationTool
 		/// <summary>
 		/// Constructor
 		/// </summary>
-		public ScriptElement(FileReference InFile, int InLineNumber, string Prefix, string LocalName, string NamespaceUri, ScriptDocument Document) 
+		public ScriptElement(FileReference InFile, int InLineNumber, string Prefix, string LocalName, string NamespaceUri, ScriptDocument Document)
 			: base(Prefix, LocalName, NamespaceUri, Document)
 		{
 			File = InFile;
@@ -163,21 +163,16 @@ namespace AutomationTool
 		/// </summary>
 		Graph Graph = new Graph();
 
-		/// <summary>
-		/// Mapping from name to group
-		/// </summary>
-		Dictionary<string, AgentGroup> NameToGroup = new Dictionary<string,AgentGroup>(StringComparer.InvariantCultureIgnoreCase);
-
 		/// <summary>
 		/// Mapping of global property name to values.
 		/// </summary>
-		Dictionary<string, string> GlobalProperties = new Dictionary<string,string>();
+		Dictionary<string, string> GlobalProperties = new Dictionary<string, string>();
 
 		/// <summary>
-		/// List of property name to value lookups. Modifications to properties are scoped to nodes and groups. EnterScope() pushes an empty dictionary onto the end of this list, and LeaveScope() removes one. 
+		/// List of property name to value lookups. Modifications to properties are scoped to nodes and agents. EnterScope() pushes an empty dictionary onto the end of this list, and LeaveScope() removes one. 
 		/// ExpandProperties() searches from last to first lookup when trying to resolve a property name, and takes the first it finds.
 		/// </summary>
-		List<Dictionary<string, string>> ScopedProperties = new List<Dictionary<string,string>>();
+		List<Dictionary<string, string>> ScopedProperties = new List<Dictionary<string, string>>();
 
 		/// <summary>
 		/// Schema for the script
@@ -196,7 +191,7 @@ namespace AutomationTool
 		/// <param name="InSchema">Schema for the script</param>
 		private ScriptReader(IDictionary<string, string> Properties, ScriptSchema InSchema)
 		{
-			GlobalProperties = new Dictionary<string,string>(Properties, StringComparer.InvariantCultureIgnoreCase);
+			GlobalProperties = new Dictionary<string, string>(Properties, StringComparer.InvariantCultureIgnoreCase);
 			ScopedProperties.Add(GlobalProperties);
 			Schema = InSchema;
 		}
@@ -212,7 +207,7 @@ namespace AutomationTool
 		public static bool TryRead(FileReference File, IDictionary<string, string> DefaultProperties, ScriptSchema Schema, out Graph Graph)
 		{
 			// Check the file exists before doing anything.
-			if(!File.Exists())
+			if (!File.Exists())
 			{
 				CommandUtils.LogError("Cannot open '{0}'", File.FullName);
 				Graph = null;
@@ -221,7 +216,7 @@ namespace AutomationTool
 
 			// Read the file and build the graph
 			ScriptReader Reader = new ScriptReader(DefaultProperties, Schema);
-			if(Reader.TryRead(File) && Reader.NumErrors == 0)
+			if (Reader.TryRead(File) && Reader.NumErrors == 0)
 			{
 				Graph = Reader.Graph;
 				return true;
@@ -241,7 +236,7 @@ namespace AutomationTool
 		{
 			// Read the document and validate it against the schema
 			ScriptDocument Document;
-			if(!ScriptDocument.TryRead(File, Schema, out Document))
+			if (!ScriptDocument.TryRead(File, Schema, out Document))
 			{
 				NumErrors++;
 				return false;
@@ -249,9 +244,9 @@ namespace AutomationTool
 
 			// Read the root BuildGraph element
 			EnterScope();
-			foreach(ScriptElement Element in Document.DocumentElement.ChildNodes.OfType<ScriptElement>())
+			foreach (ScriptElement Element in Document.DocumentElement.ChildNodes.OfType<ScriptElement>())
 			{
-				switch(Element.Name)
+				switch (Element.Name)
 				{
 					case "Include":
 						ReadInclude(Element, File.Directory);
@@ -296,7 +291,7 @@ namespace AutomationTool
 		/// <param name="Args">Event arguments</param>
 		void ValidationHandler(object Sender, ValidationEventArgs Args)
 		{
-			if(Args.Severity == XmlSeverityType.Warning)
+			if (Args.Severity == XmlSeverityType.Warning)
 			{
 				CommandUtils.LogWarning("Script: {0}", Args.Message);
 			}
@@ -312,7 +307,7 @@ namespace AutomationTool
 		/// </summary>
 		void EnterScope()
 		{
-			ScopedProperties.Add(new Dictionary<string,string>(StringComparer.InvariantCultureIgnoreCase));
+			ScopedProperties.Add(new Dictionary<string, string>(StringComparer.InvariantCultureIgnoreCase));
 		}
 
 		/// <summary>
@@ -324,25 +319,25 @@ namespace AutomationTool
 		}
 
 		/// <summary>
-		/// Reads the definition for an agent group.
+		/// Reads the definition for a trigger.
 		/// </summary>
 		/// <param name="Element">Xml element to read the definition from</param>
 		void ReadTrigger(ScriptElement Element)
 		{
 			string[] QualifiedName;
-			if(EvaluateCondition(Element) && TryReadQualifiedObjectName(Element, out QualifiedName))
+			if (EvaluateCondition(Element) && TryReadQualifiedObjectName(Element, out QualifiedName))
 			{
 				// Validate all the parent triggers
 				ManualTrigger ParentTrigger = null;
-				for(int Idx = 0; Idx < QualifiedName.Length - 1; Idx++)
+				for (int Idx = 0; Idx < QualifiedName.Length - 1; Idx++)
 				{
 					ManualTrigger NextTrigger;
-					if(!Graph.NameToTrigger.TryGetValue(QualifiedName[Idx], out NextTrigger))
+					if (!Graph.NameToTrigger.TryGetValue(QualifiedName[Idx], out NextTrigger))
 					{
 						LogError(Element, "Unknown trigger '{0}'", QualifiedName[Idx]);
 						return;
 					}
-					if(NextTrigger.Parent != ParentTrigger)
+					if (NextTrigger.Parent != ParentTrigger)
 					{
 						LogError(Element, "Qualified name of trigger '{0}' is '{1}'", NextTrigger.Name, NextTrigger.QualifiedName);
 						return;
@@ -355,12 +350,12 @@ namespace AutomationTool
 
 				// Create the new trigger
 				ManualTrigger Trigger;
-				if(!Graph.NameToTrigger.TryGetValue(Name, out Trigger))
+				if (!Graph.NameToTrigger.TryGetValue(Name, out Trigger))
 				{
 					Trigger = new ManualTrigger(ParentTrigger, Name);
 					Graph.NameToTrigger.Add(Name, Trigger);
 				}
-				else if(Trigger.Parent != ParentTrigger)
+				else if (Trigger.Parent != ParentTrigger)
 				{
 					LogError(Element, "Conflicting parent for '{0}' - previously declared as '{1}', now '{2}'", Name, Trigger.QualifiedName, new ManualTrigger(ParentTrigger, Name).QualifiedName);
 					return;
@@ -368,9 +363,9 @@ namespace AutomationTool
 
 				// Read the root BuildGraph element
 				EnterScope();
-				foreach(ScriptElement ChildElement in Element.ChildNodes.OfType<ScriptElement>())
+				foreach (ScriptElement ChildElement in Element.ChildNodes.OfType<ScriptElement>())
 				{
-					switch(ChildElement.Name)
+					switch (ChildElement.Name)
 					{
 						case "Property":
 							ReadProperty(ChildElement);
@@ -409,10 +404,10 @@ namespace AutomationTool
 		/// <param name="BaseDir">Base directory to resolve relative include paths from </param>
 		void ReadInclude(ScriptElement Element, DirectoryReference BaseDir)
 		{
-			if(EvaluateCondition(Element))
+			if (EvaluateCondition(Element))
 			{
 				FileReference Script = FileReference.Combine(BaseDir, Element.GetAttribute("Script"));
-				if(Script.Exists())
+				if (Script.Exists())
 				{
 					TryRead(Script);
 				}
@@ -429,10 +424,10 @@ namespace AutomationTool
 		/// <param name="Element">Xml element to read the definition from</param>
 		void ReadProperty(ScriptElement Element)
 		{
-			if(EvaluateCondition(Element))
+			if (EvaluateCondition(Element))
 			{
 				string Name = ReadAttribute(Element, "Name");
-				if(ValidateName(Element, Name))
+				if (ValidateName(Element, Name))
 				{
 					GlobalProperties[Name] = ReadAttribute(Element, "Value");
 				}
@@ -445,10 +440,10 @@ namespace AutomationTool
 		/// <param name="Element">Xml element to read the definition from</param>
 		void ReadLocalProperty(ScriptElement Element)
 		{
-			if(EvaluateCondition(Element))
+			if (EvaluateCondition(Element))
 			{
 				string Name = ReadAttribute(Element, "Name");
-				if(ValidateName(Element, Name))
+				if (ValidateName(Element, Name))
 				{
 					ScopedProperties[ScopedProperties.Count - 1][Name] = ReadAttribute(Element, "Value");
 				}
@@ -456,48 +451,48 @@ namespace AutomationTool
 		}
 
 		/// <summary>
-		/// Reads the definition for an agent group.
+		/// Reads the definition for an agent.
 		/// </summary>
 		/// <param name="Element">Xml element to read the definition from</param>
-		/// <param name="Trigger">The controlling trigger for nodes in this group</param>
+		/// <param name="Trigger">The controlling trigger for nodes in this agent</param>
 		void ReadAgent(ScriptElement Element, ManualTrigger Trigger)
 		{
 			string Name;
-			if(EvaluateCondition(Element) && TryReadObjectName(Element, out Name))
+			if (EvaluateCondition(Element) && TryReadObjectName(Element, out Name))
 			{
-				// Read the valid agent types. This may be omitted if we're continuing an existing group.
+				// Read the valid agent types. This may be omitted if we're continuing an existing agent.
 				string[] Types = ReadListAttribute(Element, "Type");
 
-				// Create the group object, or continue an existing one
-				AgentGroup Group;
-				if(NameToGroup.TryGetValue(Name, out Group))
+				// Create the agent object, or continue an existing one
+				Agent Agent;
+				if (Graph.NameToAgent.TryGetValue(Name, out Agent))
 				{
-					if(Types.Length > 0 && Group.PossibleTypes.Length > 0)
+					if (Types.Length > 0 && Agent.PossibleTypes.Length > 0)
 					{
-						string[] NewTypes = Group.PossibleTypes.Intersect(Types, StringComparer.InvariantCultureIgnoreCase).ToArray();
-						if(NewTypes.Length == 0)
+						string[] NewTypes = Agent.PossibleTypes.Intersect(Types, StringComparer.InvariantCultureIgnoreCase).ToArray();
+						if (NewTypes.Length == 0)
 						{
 							LogError(Element, "No common agent types with previous agent definition");
 						}
-						Group.PossibleTypes = NewTypes;
+						Agent.PossibleTypes = NewTypes;
 					}
 				}
 				else
 				{
-					if(Types.Length == 0)
+					if (Types.Length == 0)
 					{
-						LogError(Element, "Missing agent type for group '{0}'", Name);
+						LogError(Element, "Missing type for agent '{0}'", Name);
 					}
-					Group = new AgentGroup(Name, Types);
-					NameToGroup.Add(Name, Group);
-					Graph.Groups.Add(Group);
+					Agent = new Agent(Name, Types);
+					Graph.NameToAgent.Add(Name, Agent);
+					Graph.Agents.Add(Agent);
 				}
 
 				// Process all the child elements.
 				EnterScope();
-				foreach(ScriptElement ChildElement in Element.ChildNodes.OfType<ScriptElement>())
+				foreach (ScriptElement ChildElement in Element.ChildNodes.OfType<ScriptElement>())
 				{
-					switch(ChildElement.Name)
+					switch (ChildElement.Name)
 					{
 						case "Property":
 							ReadProperty(ChildElement);
@@ -506,16 +501,16 @@ namespace AutomationTool
 							ReadLocalProperty(ChildElement);
 							break;
 						case "Node":
-							ReadNode(ChildElement, Group, Trigger);
+							ReadNode(ChildElement, Agent, Trigger);
 							break;
 						case "Aggregate":
 							ReadAggregate(ChildElement);
 							break;
 						case "Warning":
-							ReadDiagnostic(ChildElement, LogEventType.Warning, null, Group, Trigger);
+							ReadDiagnostic(ChildElement, LogEventType.Warning, null, Agent, Trigger);
 							break;
 						case "Error":
-							ReadDiagnostic(ChildElement, LogEventType.Error, null, Group, Trigger);
+							ReadDiagnostic(ChildElement, LogEventType.Error, null, Agent, Trigger);
 							break;
 						default:
 							LogError(ChildElement, "Unexpected element type '{0}'", ChildElement.Name);
@@ -527,13 +522,13 @@ namespace AutomationTool
 		}
 
 		/// <summary>
-		/// Reads the definition for an aggregate, and adds it to the given agent group
+		/// Reads the definition for an aggregate
 		/// </summary>
 		/// <param name="Element">Xml element to read the definition from</param>
 		void ReadAggregate(ScriptElement Element)
 		{
 			string Name;
-			if(EvaluateCondition(Element) && TryReadObjectName(Element, out Name) && CheckNameIsUnique(Element, Name))
+			if (EvaluateCondition(Element) && TryReadObjectName(Element, out Name) && CheckNameIsUnique(Element, Name))
 			{
 				string[] RequiredNames = ReadListAttribute(Element, "Requires");
 				Graph.AggregateNameToNodes.Add(Name, ResolveReferences(Element, RequiredNames).ToArray());
@@ -541,15 +536,15 @@ namespace AutomationTool
 		}
 
 		/// <summary>
-		/// Reads the definition for a node, and adds it to the given agent group
+		/// Reads the definition for a node, and adds it to the given agent
 		/// </summary>
 		/// <param name="Element">Xml element to read the definition from</param>
-		/// <param name="Group">Group for the node to be added to</param>
+		/// <param name="ParentAgent">Agent for the node to be added to</param>
 		/// <param name="ControllingTrigger">The controlling trigger for this node</param>
-		void ReadNode(ScriptElement Element, AgentGroup Group, ManualTrigger ControllingTrigger)
+		void ReadNode(ScriptElement Element, Agent ParentAgent, ManualTrigger ControllingTrigger)
 		{
 			string Name;
-			if(EvaluateCondition(Element) && TryReadObjectName(Element, out Name))
+			if (EvaluateCondition(Element) && TryReadObjectName(Element, out Name))
 			{
 				string[] RequiresNames = ReadListAttribute(Element, "Requires");
 				string[] ProducesNames = ReadListAttribute(Element, "Produces");
@@ -561,9 +556,9 @@ namespace AutomationTool
 
 				// Gather up all the input dependencies, and check they're all upstream of the current node
 				HashSet<Node> InputDependencies = new HashSet<Node>();
-				foreach(Node InputDependency in Inputs.Select(x => x.ProducingNode).Distinct())
+				foreach (Node InputDependency in Inputs.Select(x => x.ProducingNode).Distinct())
 				{
-					if(InputDependency.ControllingTrigger != null && InputDependency.ControllingTrigger != ControllingTrigger && !InputDependency.ControllingTrigger.IsUpstreamFrom(ControllingTrigger))
+					if (InputDependency.ControllingTrigger != null && InputDependency.ControllingTrigger != ControllingTrigger && !InputDependency.ControllingTrigger.IsUpstreamFrom(ControllingTrigger))
 					{
 						LogError(Element, "'{0}' is dependent on '{1}', which is behind a different controlling trigger ({2})", Name, InputDependency.Name, InputDependency.ControllingTrigger.QualifiedName);
 					}
@@ -574,16 +569,16 @@ namespace AutomationTool
 				}
 
 				// Recursively include all their dependencies too
-				foreach(Node InputDependency in InputDependencies.ToArray())
+				foreach (Node InputDependency in InputDependencies.ToArray())
 				{
 					InputDependencies.UnionWith(InputDependency.InputDependencies);
 				}
 
 				// Add the name of the node itself to the list of outputs.
 				List<string> OutputNames = new List<string>();
-				foreach(string ProducesName in ProducesNames)
+				foreach (string ProducesName in ProducesNames)
 				{
-					if(ProducesName.StartsWith("#"))
+					if (ProducesName.StartsWith("#"))
 					{
 						OutputNames.Add(ProducesName.Substring(1));
 					}
@@ -599,23 +594,23 @@ namespace AutomationTool
 				OrderDependencies.UnionWith(ResolveReferences(Element, AfterNames));
 
 				// Recursively include all their order dependencies too
-				foreach(Node OrderDependency in OrderDependencies.ToArray())
+				foreach (Node OrderDependency in OrderDependencies.ToArray())
 				{
 					OrderDependencies.UnionWith(OrderDependency.OrderDependencies);
 				}
 
-				// Check that we're not dependent on anything completing that is declared after the initial declaration of this group.
-				int GroupIdx = Graph.Groups.IndexOf(Group);
-				for(int Idx = GroupIdx + 1; Idx < Graph.Groups.Count; Idx++)
+				// Check that we're not dependent on anything completing that is declared after the initial declaration of this agent.
+				int AgentIdx = Graph.Agents.IndexOf(ParentAgent);
+				for (int Idx = AgentIdx + 1; Idx < Graph.Agents.Count; Idx++)
 				{
-					foreach(Node Node in Graph.Groups[Idx].Nodes.Where(x => OrderDependencies.Contains(x)))
+					foreach (Node Node in Graph.Agents[Idx].Nodes.Where(x => OrderDependencies.Contains(x)))
 					{
-						LogError(Element, "Node '{0}' has a dependency on '{1}', which was declared after the initial definition of '{2}'.", Name, Node.Name, Group.Name);
+						LogError(Element, "Node '{0}' has a dependency on '{1}', which was declared after the initial definition of '{2}'.", Name, Node.Name, ParentAgent.Name);
 					}
 				}
 
 				// Construct and register the node
-				if(CheckNameIsUnique(Element, Name))
+				if (CheckNameIsUnique(Element, Name))
 				{
 					// Add it to the node lookup
 					Node NewNode = new Node(Name, Inputs.ToArray(), OutputNames.ToArray(), InputDependencies.ToArray(), OrderDependencies.ToArray(), ControllingTrigger);
@@ -623,9 +618,9 @@ namespace AutomationTool
 					Graph.NameToNode.Add(Name, NewNode);
 
 					// Register each of the outputs as a reference to this node
-					foreach(NodeOutput Output in NewNode.Outputs)
+					foreach (NodeOutput Output in NewNode.Outputs)
 					{
-						if(Output.Name == Name || CheckNameIsUnique(Element, Output.Name))
+						if (Output.Name == Name || CheckNameIsUnique(Element, Output.Name))
 						{
 							Graph.NameToNodeOutput.Add(Output.Name, Output);
 						}
@@ -633,9 +628,9 @@ namespace AutomationTool
 
 					// Add all the tasks
 					EnterScope();
-					foreach(ScriptElement ChildElement in Element.ChildNodes.OfType<ScriptElement>())
+					foreach (ScriptElement ChildElement in Element.ChildNodes.OfType<ScriptElement>())
 					{
-						switch(ChildElement.Name)
+						switch (ChildElement.Name)
 						{
 							case "Property":
 								ReadProperty(ChildElement);
@@ -644,10 +639,10 @@ namespace AutomationTool
 								ReadLocalProperty(ChildElement);
 								break;
 							case "Warning":
-								ReadDiagnostic(ChildElement, LogEventType.Warning, NewNode, Group, ControllingTrigger);
+								ReadDiagnostic(ChildElement, LogEventType.Warning, NewNode, ParentAgent, ControllingTrigger);
 								break;
 							case "Error":
-								ReadDiagnostic(ChildElement, LogEventType.Error, NewNode, Group, ControllingTrigger);
+								ReadDiagnostic(ChildElement, LogEventType.Error, NewNode, ParentAgent, ControllingTrigger);
 								break;
 							default:
 								ReadTask(ChildElement, NewNode.Tasks);
@@ -656,8 +651,8 @@ namespace AutomationTool
 					}
 					LeaveScope();
 
-					// Add it to the current agent group
-					Group.Nodes.Add(NewNode);
+					// Add it to the current agent
+					ParentAgent.Nodes.Add(NewNode);
 				}
 			}
 		}
@@ -669,11 +664,11 @@ namespace AutomationTool
 		/// <param name="Tasks">List of tasks to add to</param>
 		void ReadTask(ScriptElement Element, List<CustomTask> Tasks)
 		{
-			if(EvaluateCondition(Element))
+			if (EvaluateCondition(Element))
 			{
 				// Get the reflection info for this element
 				ScriptTask Task;
-				if(!Schema.TryGetTask(Element.Name, out Task))
+				if (!Schema.TryGetTask(Element.Name, out Task))
 				{
 					LogError(Element, "Unknown task '{0}'", Element.Name);
 					return;
@@ -681,9 +676,9 @@ namespace AutomationTool
 
 				// Check all the required parameters are present
 				bool bHasRequiredAttributes = true;
-				foreach(ScriptTaskParameter Parameter in Task.NameToParameter.Values)
+				foreach (ScriptTaskParameter Parameter in Task.NameToParameter.Values)
 				{
-					if(!Parameter.bOptional && !Element.HasAttribute(Parameter.Name))
+					if (!Parameter.bOptional && !Element.HasAttribute(Parameter.Name))
 					{
 						LogError(Element, "Missing required attribute - {0}", Parameter.Name);
 						bHasRequiredAttributes = false;
@@ -692,13 +687,13 @@ namespace AutomationTool
 
 				// Read all the attributes into a parameters object for this task
 				object ParametersObject = Activator.CreateInstance(Task.ParametersClass);
-				foreach(XmlAttribute Attribute in Element.Attributes)
+				foreach (XmlAttribute Attribute in Element.Attributes)
 				{
-					if(String.Compare(Attribute.Name, "If", StringComparison.InvariantCultureIgnoreCase) != 0)
+					if (String.Compare(Attribute.Name, "If", StringComparison.InvariantCultureIgnoreCase) != 0)
 					{
 						// Get the field that this attribute should be written to in the parameters object
 						ScriptTaskParameter Parameter;
-						if(!Task.NameToParameter.TryGetValue(Attribute.Name, out Parameter))
+						if (!Task.NameToParameter.TryGetValue(Attribute.Name, out Parameter))
 						{
 							LogError(Element, "Unknown attribute '{0}'", Attribute.Name);
 							continue;
@@ -709,11 +704,11 @@ namespace AutomationTool
 
 						// Parse it and assign it to the parameters object
 						object Value;
-						if(Parameter.FieldInfo.FieldType.IsEnum)
+						if (Parameter.FieldInfo.FieldType.IsEnum)
 						{
 							Value = Enum.Parse(Parameter.FieldInfo.FieldType, ExpandedValue);
 						}
-						else if(Parameter.FieldInfo.FieldType == typeof(Boolean))
+						else if (Parameter.FieldInfo.FieldType == typeof(Boolean))
 						{
 							Value = Condition.Evaluate(ExpandedValue);
 						}
@@ -726,7 +721,7 @@ namespace AutomationTool
 				}
 
 				// Construct the task
-				if(bHasRequiredAttributes)
+				if (bHasRequiredAttributes)
 				{
 					Tasks.Add((CustomTask)Activator.CreateInstance(Task.TaskClass, ParametersObject));
 				}
@@ -739,7 +734,7 @@ namespace AutomationTool
 		/// <param name="Element">Xml element to read the definition from</param>
 		void ReadNotifier(ScriptElement Element)
 		{
-			if(EvaluateCondition(Element))
+			if (EvaluateCondition(Element))
 			{
 				string[] TargetNames = ReadListAttribute(Element, "Targets");
 				string[] ExceptNames = ReadListAttribute(Element, "Except");
@@ -747,14 +742,14 @@ namespace AutomationTool
 				string[] TriggerNames = ReadListAttribute(Element, "Triggers");
 				string[] Users = ReadListAttribute(Element, "Users");
 				string[] Submitters = ReadListAttribute(Element, "Submitters");
-				bool? bWarnings = Element.HasAttribute("Warnings")? (bool?)ReadBooleanAttribute(Element, "Warnings", true) : null;
+				bool? bWarnings = Element.HasAttribute("Warnings") ? (bool?)ReadBooleanAttribute(Element, "Warnings", true) : null;
 
 				// Find the list of targets which are included, and recurse through all their dependencies
 				HashSet<Node> Nodes = new HashSet<Node>();
-				if(TargetNames != null)
+				if (TargetNames != null)
 				{
 					HashSet<Node> TargetNodes = ResolveReferences(Element, TargetNames);
-					foreach(Node Node in TargetNodes)
+					foreach (Node Node in TargetNodes)
 					{
 						Nodes.Add(Node);
 						Nodes.UnionWith(Node.InputDependencies);
@@ -762,43 +757,43 @@ namespace AutomationTool
 				}
 
 				// Add all the individually referenced nodes
-				if(IndividualNodeNames != null)
+				if (IndividualNodeNames != null)
 				{
 					HashSet<Node> IndividualNodes = ResolveReferences(Element, IndividualNodeNames);
 					Nodes.UnionWith(IndividualNodes);
 				}
 
 				// Exclude all the exceptions
-				if(ExceptNames != null)
+				if (ExceptNames != null)
 				{
 					HashSet<Node> ExceptNodes = ResolveReferences(Element, ExceptNames);
 					Nodes.ExceptWith(ExceptNodes);
 				}
 
 				// Update all the referenced nodes with the settings
-				foreach(Node Node in Nodes)
+				foreach (Node Node in Nodes)
 				{
-					if(Users != null)
+					if (Users != null)
 					{
 						Node.NotifyUsers.UnionWith(Users);
 					}
-					if(Submitters != null)
+					if (Submitters != null)
 					{
 						Node.NotifySubmitters.UnionWith(Submitters);
 					}
-					if(bWarnings.HasValue)
+					if (bWarnings.HasValue)
 					{
 						Node.bNotifyOnWarnings = bWarnings.Value;
 					}
 				}
 
 				// Add the users to the list of triggers
-				if(TriggerNames != null)
+				if (TriggerNames != null)
 				{
-					foreach(string TriggerName in TriggerNames)
+					foreach (string TriggerName in TriggerNames)
 					{
 						ManualTrigger Trigger;
-						if(Graph.NameToTrigger.TryGetValue(TriggerName, out Trigger))
+						if (Graph.NameToTrigger.TryGetValue(TriggerName, out Trigger))
 						{
 							Trigger.NotifyUsers.UnionWith(Users);
 						}
@@ -816,10 +811,12 @@ namespace AutomationTool
 		/// </summary>
 		/// <param name="Element">Xml element to read the definition from</param>
 		/// <param name="EventType">The diagnostic event type</param>
-		/// <param name="EnclosingObject">The enclosing object instance</param>
-		void ReadDiagnostic(ScriptElement Element, LogEventType EventType, Node EnclosingNode, AgentGroup EnclosingGroup, ManualTrigger EnclosingTrigger)
+		/// <param name="EnclosingNode">The node that this diagnostic is declared in, or null</param>
+		/// <param name="EnclosingAgent">The agent that this diagnostic is declared in, or null</param>
+		/// <param name="EnclosingTrigger">The trigger that this diagnostic is declared in, or null</param>
+		void ReadDiagnostic(ScriptElement Element, LogEventType EventType, Node EnclosingNode, Agent EnclosingAgent, ManualTrigger EnclosingTrigger)
 		{
-			if(EvaluateCondition(Element))
+			if (EvaluateCondition(Element))
 			{
 				string Message = ReadAttribute(Element, "Message");
 
@@ -827,7 +824,7 @@ namespace AutomationTool
 				Diagnostic.EventType = EventType;
 				Diagnostic.Message = String.Format("{0}({1}): {2}", Element.File.FullName, Element.LineNumber, Message);
 				Diagnostic.EnclosingNode = EnclosingNode;
-				Diagnostic.EnclosingGroup = EnclosingGroup;
+				Diagnostic.EnclosingAgent = EnclosingAgent;
 				Diagnostic.EnclosingTrigger = EnclosingTrigger;
 				Graph.Diagnostics.Add(Diagnostic);
 			}
@@ -843,7 +840,7 @@ namespace AutomationTool
 		bool CheckNameIsUnique(ScriptElement Element, string Name)
 		{
 			// Get the nodes that it maps to
-			if(Graph.ContainsName(Name))
+			if (Graph.ContainsName(Name))
 			{
 				LogError(Element, "'{0}' is already defined; cannot add a second time", Name);
 				return false;
@@ -860,14 +857,14 @@ namespace AutomationTool
 		HashSet<Node> ResolveReferences(ScriptElement Element, IEnumerable<string> ReferenceNames)
 		{
 			HashSet<Node> Nodes = new HashSet<Node>();
-			foreach(string ReferenceName in ReferenceNames)
+			foreach (string ReferenceName in ReferenceNames)
 			{
 				Node[] OtherNodes;
-				if(Graph.TryResolveReference(ReferenceName, out OtherNodes))
+				if (Graph.TryResolveReference(ReferenceName, out OtherNodes))
 				{
 					Nodes.UnionWith(OtherNodes);
 				}
-				else if(!ReferenceName.StartsWith("#") && Graph.NameToNodeOutput.ContainsKey(ReferenceName))
+				else if (!ReferenceName.StartsWith("#") && Graph.NameToNodeOutput.ContainsKey(ReferenceName))
 				{
 					LogError(Element, "Reference to '{0}' cannot be resolved; did you mean '#{0}'?", ReferenceName);
 				}
@@ -888,14 +885,14 @@ namespace AutomationTool
 		HashSet<NodeOutput> ResolveInputReferences(ScriptElement Element, IEnumerable<string> ReferenceNames)
 		{
 			HashSet<NodeOutput> Inputs = new HashSet<NodeOutput>();
-			foreach(string ReferenceName in ReferenceNames)
+			foreach (string ReferenceName in ReferenceNames)
 			{
 				NodeOutput[] ReferenceInputs;
-				if(Graph.TryResolveInputReference(ReferenceName, out ReferenceInputs))
+				if (Graph.TryResolveInputReference(ReferenceName, out ReferenceInputs))
 				{
 					Inputs.UnionWith(ReferenceInputs);
 				}
-				else if(!ReferenceName.StartsWith("#") && Graph.NameToNodeOutput.ContainsKey(ReferenceName))
+				else if (!ReferenceName.StartsWith("#") && Graph.NameToNodeOutput.ContainsKey(ReferenceName))
 				{
 					LogError(Element, "Reference to '{0}' cannot be resolved; did you mean '#{0}'?", ReferenceName);
 				}
@@ -916,7 +913,7 @@ namespace AutomationTool
 		bool TryReadObjectName(ScriptElement Element, out string Name)
 		{
 			// Check the name attribute is present
-			if(!Element.HasAttribute("Name"))
+			if (!Element.HasAttribute("Name"))
 			{
 				LogError(Element, "Missing 'Name' attribute");
 				Name = null;
@@ -925,7 +922,7 @@ namespace AutomationTool
 
 			// Get the value of it, strip any leading or trailing whitespace, and make sure it's not empty
 			string Value = ReadAttribute(Element, "Name");
-			if(!ValidateName(Element, Value))
+			if (!ValidateName(Element, Value))
 			{
 				Name = null;
 				return false;
@@ -945,7 +942,7 @@ namespace AutomationTool
 		bool TryReadQualifiedObjectName(ScriptElement Element, out string[] QualifiedName)
 		{
 			// Check the name attribute is present
-			if(!Element.HasAttribute("Name"))
+			if (!Element.HasAttribute("Name"))
 			{
 				LogError(Element, "Missing 'Name' attribute");
 				QualifiedName = null;
@@ -954,9 +951,9 @@ namespace AutomationTool
 
 			// Get the value of it, strip any leading or trailing whitespace, and make sure it's not empty
 			string[] Values = ReadAttribute(Element, "Name").Split('.');
-			foreach(string Value in Values)
+			foreach (string Value in Values)
 			{
-				if(!ValidateName(Element, Value))
+				if (!ValidateName(Element, Value))
 				{
 					QualifiedName = null;
 					return false;
@@ -977,21 +974,21 @@ namespace AutomationTool
 		bool ValidateName(ScriptElement Element, string Name)
 		{
 			// Check it's not empty
-			if(Name.Length == 0)
+			if (Name.Length == 0)
 			{
 				LogError(Element, "Name is empty");
 				return false;
 			}
 
 			// Check there are no invalid characters
-			for(int Idx = 0; Idx < Name.Length; Idx++)
+			for (int Idx = 0; Idx < Name.Length; Idx++)
 			{
-				if(Idx > 0 && Name[Idx] == ' ' && Name[Idx - 1] == ' ')
+				if (Idx > 0 && Name[Idx] == ' ' && Name[Idx - 1] == ' ')
 				{
 					LogError(Element, "Consecutive spaces in object name");
 					return false;
 				}
-				if(!Char.IsLetterOrDigit(Name[Idx]) && Name[Idx] != '_' && Name[Idx] != ' ')
+				if (!Char.IsLetterOrDigit(Name[Idx]) && Name[Idx] != '_' && Name[Idx] != ' ')
 				{
 					LogError(Element, "Invalid character in object name - '{0}'", Name[Idx]);
 					return false;
@@ -1020,7 +1017,7 @@ namespace AutomationTool
 		string[] ReadListAttribute(ScriptElement Element, string Name)
 		{
 			string Value = ReadAttribute(Element, Name);
-			return Value.Split(new char[]{ ';' }).Select(x => x.Trim()).Where(x => x.Length > 0).ToArray();
+			return Value.Split(new char[] { ';' }).Select(x => x.Trim()).Where(x => x.Length > 0).ToArray();
 		}
 
 		/// <summary>
@@ -1033,14 +1030,14 @@ namespace AutomationTool
 		bool ReadBooleanAttribute(ScriptElement Element, string Name, bool bDefaultValue)
 		{
 			bool bResult = bDefaultValue;
-			if(Element.HasAttribute(Name))
+			if (Element.HasAttribute(Name))
 			{
 				string Value = ReadAttribute(Element, Name).Trim();
-				if(Value.Equals("true", StringComparison.InvariantCultureIgnoreCase))
+				if (Value.Equals("true", StringComparison.InvariantCultureIgnoreCase))
 				{
 					bResult = true;
 				}
-				else if(Value.Equals("false", StringComparison.InvariantCultureIgnoreCase))
+				else if (Value.Equals("false", StringComparison.InvariantCultureIgnoreCase))
 				{
 					bResult = false;
 				}
@@ -1063,12 +1060,12 @@ namespace AutomationTool
 		T ReadEnumAttribute<T>(ScriptElement Element, string Name, T DefaultValue) where T : struct
 		{
 			T Result = DefaultValue;
-			if(Element.HasAttribute(Name))
+			if (Element.HasAttribute(Name))
 			{
 				string Value = ReadAttribute(Element, Name).Trim();
 
 				T EnumValue;
-				if(Enum.TryParse(Value, true, out EnumValue))
+				if (Enum.TryParse(Value, true, out EnumValue))
 				{
 					Result = EnumValue;
 				}
@@ -1101,7 +1098,7 @@ namespace AutomationTool
 		{
 			// Check if the element has a conditional attribute
 			const string AttributeName = "If";
-			if(!Element.HasAttribute(AttributeName))
+			if (!Element.HasAttribute(AttributeName))
 			{
 				return true;
 			}
@@ -1112,7 +1109,7 @@ namespace AutomationTool
 				string Text = ExpandProperties(Element.GetAttribute("If"));
 				return Condition.Evaluate(Text);
 			}
-			catch(ConditionException Ex)
+			catch (ConditionException Ex)
 			{
 				LogError(Element, "Error in condition: {0}", Ex.Message);
 				return false;
@@ -1141,9 +1138,9 @@ namespace AutomationTool
 
 				// Find the value for it, either from the dictionary or the environment block
 				string Value = null;
-				for(int ScopeIdx = ScopedProperties.Count - 1; ScopeIdx >= 0; ScopeIdx--)
+				for (int ScopeIdx = ScopedProperties.Count - 1; ScopeIdx >= 0; ScopeIdx--)
 				{
-					if(ScopedProperties[ScopeIdx].TryGetValue(Name, out Value))
+					if (ScopedProperties[ScopeIdx].TryGetValue(Name, out Value))
 					{
 						break;
 					}
diff --git a/Engine/Source/Programs/AutomationTool/BuildGraph/Task.cs b/Engine/Source/Programs/AutomationTool/BuildGraph/Task.cs
--- a/Engine/Source/Programs/AutomationTool/BuildGraph/Task.cs
+++ b/Engine/Source/Programs/AutomationTool/BuildGraph/Task.cs
@@ -100,7 +100,7 @@ namespace AutomationTool
 		public TaskElementAttribute(string InName, Type InParametersType)
 		{
 			Name = InName;
-			ParametersType =  InParametersType;
+			ParametersType = InParametersType;
 		}
 	}
 
@@ -126,6 +126,44 @@ namespace AutomationTool
 		/// <returns>Whether the task succeeded or not. Exiting with an exception will be caught and treated as a failure.</returns>
 		public abstract bool Execute(JobContext Job, HashSet<FileReference> BuildProducts, Dictionary<string, HashSet<FileReference>> TagNameToFileSet);
 
+		/// <summary>
+		/// Output this task out to an XML writer.
+		/// </summary>
+		public abstract void Write(XmlWriter Writer);
+
+		/// <summary>
+		/// Writes this task to an XML writer, using the given parameters object.
+		/// </summary>
+		/// <param name="Writer"></param>
+		/// <param name="Task"></param>
+		protected void Write(XmlWriter Writer, object Parameters)
+		{
+			TaskElementAttribute Element = GetType().GetCustomAttribute<TaskElementAttribute>();
+			Writer.WriteStartElement(Element.Name);
+
+			foreach (FieldInfo Field in Parameters.GetType().GetFields())
+			{
+				if (Field.MemberType == MemberTypes.Field)
+				{
+					TaskParameterAttribute ParameterAttribute = Field.GetCustomAttribute<TaskParameterAttribute>();
+					if (ParameterAttribute != null)
+					{
+						object Value = Field.GetValue(Parameters);
+						if (Value != null && Field.FieldType == typeof(bool) && (bool)Value == false)
+						{
+							Value = null;
+						}
+						if (Value != null)
+						{
+							Writer.WriteAttributeString(Field.Name, Value.ToString());
+						}
+					}
+				}
+			}
+
+			Writer.WriteEndElement();
+		}
+
 		/// <summary>
 		/// Resolves a single name to a file reference, resolving relative paths to the root of the current path.
 		/// </summary>
diff --git a/Engine/Source/Programs/AutomationTool/BuildGraph/Tasks/ChunkTask.cs b/Engine/Source/Programs/AutomationTool/BuildGraph/Tasks/ChunkTask.cs
--- a/Engine/Source/Programs/AutomationTool/BuildGraph/Tasks/ChunkTask.cs
+++ b/Engine/Source/Programs/AutomationTool/BuildGraph/Tasks/ChunkTask.cs
@@ -5,6 +5,7 @@ using System.IO;
 using System.Linq;
 using System.Text;
 using System.Threading.Tasks;
+using System.Xml;
 using UnrealBuildTool;
 
 namespace AutomationTool.Tasks
@@ -75,7 +76,7 @@ namespace AutomationTool.Tasks
 		ChunkTaskParameters Parameters;
 
 		/// <summary>
-		/// Construct a new CommandTask.
+		/// Construct a new ChunkTask.
 		/// </summary>
 		/// <param name="InParameters">Parameters for this task</param>
 		public ChunkTask(ChunkTaskParameters InParameters)
@@ -132,5 +133,13 @@ namespace AutomationTool.Tasks
 			BuildPatchToolBase.Get().Execute(Options);
 			return true;
 		}
+
+		/// <summary>
+		/// Output this task out to an XML writer.
+		/// </summary>
+		public override void Write(XmlWriter Writer)
+		{
+			Write(Writer, Parameters);
+		}
 	}
 }
diff --git a/Engine/Source/Programs/AutomationTool/BuildGraph/Tasks/CommandTask.cs b/Engine/Source/Programs/AutomationTool/BuildGraph/Tasks/CommandTask.cs
--- a/Engine/Source/Programs/AutomationTool/BuildGraph/Tasks/CommandTask.cs
+++ b/Engine/Source/Programs/AutomationTool/BuildGraph/Tasks/CommandTask.cs
@@ -5,6 +5,7 @@ using System.Text;
 using System.Threading.Tasks;
 using AutomationTool;
 using UnrealBuildTool;
+using System.Xml;
 
 namespace BuildGraph.Tasks
 {
@@ -58,5 +59,13 @@ namespace BuildGraph.Tasks
 			CommandUtils.RunUAT(CommandUtils.CmdEnv, String.Format("{0} {1}", Parameters.Name, Parameters.Arguments ?? ""));
 			return true;
 		}
+
+		/// <summary>
+		/// Output this task out to an XML writer.
+		/// </summary>
+		public override void Write(XmlWriter Writer)
+		{
+			Write(Writer, Parameters);
+		}
 	}
 }
diff --git a/Engine/Source/Programs/AutomationTool/BuildGraph/Tasks/CompileTask.cs b/Engine/Source/Programs/AutomationTool/BuildGraph/Tasks/CompileTask.cs
--- a/Engine/Source/Programs/AutomationTool/BuildGraph/Tasks/CompileTask.cs
+++ b/Engine/Source/Programs/AutomationTool/BuildGraph/Tasks/CompileTask.cs
@@ -5,6 +5,7 @@ using System.Text;
 using System.Threading.Tasks;
 using UnrealBuildTool;
 using AutomationTool;
+using System.Xml;
 
 namespace AutomationTool
 {
@@ -133,5 +134,28 @@ namespace AutomationTool
 			BuildProducts.UnionWith(Builder.LibraryBuildProductFiles.Select(x => new FileReference(x)));
 			return true;
 		}
+
+		/// <summary>
+		/// Output this task out to an XML writer.
+		/// </summary>
+		public override void Write(XmlWriter Writer)
+		{
+			foreach (UE4Build.BuildTarget Target in Targets)
+			{
+				CompileTaskParameters Parameters = new CompileTaskParameters();
+				Parameters.Target = Target.TargetName;
+				Parameters.Platform = Target.Platform;
+				Parameters.Configuration = Target.Config;
+				Parameters.Arguments = Target.UBTArgs;
+
+				string TagName;
+				if (TargetToTagName.TryGetValue(Target, out TagName))
+				{
+					Parameters.Tag = TagName;
+				}
+
+				Write(Writer, Parameters);
+			}
+		}
 	}
 }
diff --git a/Engine/Source/Programs/AutomationTool/BuildGraph/Tasks/CookTask.cs b/Engine/Source/Programs/AutomationTool/BuildGraph/Tasks/CookTask.cs
--- a/Engine/Source/Programs/AutomationTool/BuildGraph/Tasks/CookTask.cs
+++ b/Engine/Source/Programs/AutomationTool/BuildGraph/Tasks/CookTask.cs
@@ -4,6 +4,7 @@ using System.Collections.Generic;
 using System.Linq;
 using System.Text;
 using System.Threading.Tasks;
+using System.Xml;
 using UnrealBuildTool;
 
 namespace BuildGraph.Tasks
@@ -122,5 +123,13 @@ namespace BuildGraph.Tasks
 			BuildProducts.UnionWith(CookedFiles);
 			return true;
 		}
+
+		/// <summary>
+		/// Output this task out to an XML writer.
+		/// </summary>
+		public override void Write(XmlWriter Writer)
+		{
+			Write(Writer, Parameters);
+		}
 	}
 }
diff --git a/Engine/Source/Programs/AutomationTool/BuildGraph/Tasks/CopyTask.cs b/Engine/Source/Programs/AutomationTool/BuildGraph/Tasks/CopyTask.cs
--- a/Engine/Source/Programs/AutomationTool/BuildGraph/Tasks/CopyTask.cs
+++ b/Engine/Source/Programs/AutomationTool/BuildGraph/Tasks/CopyTask.cs
@@ -4,6 +4,7 @@ using System.Collections.Generic;
 using System.Linq;
 using System.Text;
 using System.Threading.Tasks;
+using System.Xml;
 using UnrealBuildTool;
 
 namespace BuildGraph.Tasks
@@ -85,5 +86,13 @@ namespace BuildGraph.Tasks
 			}
 			return true;
 		}
+
+		/// <summary>
+		/// Output this task out to an XML writer.
+		/// </summary>
+		public override void Write(XmlWriter Writer)
+		{
+			Write(Writer, Parameters);
+		}
 	}
 }
diff --git a/Engine/Source/Programs/AutomationTool/BuildGraph/Tasks/DeleteTask.cs b/Engine/Source/Programs/AutomationTool/BuildGraph/Tasks/DeleteTask.cs
--- a/Engine/Source/Programs/AutomationTool/BuildGraph/Tasks/DeleteTask.cs
+++ b/Engine/Source/Programs/AutomationTool/BuildGraph/Tasks/DeleteTask.cs
@@ -5,6 +5,7 @@ using System.IO;
 using System.Linq;
 using System.Text;
 using System.Threading.Tasks;
+using System.Xml;
 using UnrealBuildTool;
 
 namespace BuildGraph.Tasks
@@ -113,5 +114,13 @@ namespace BuildGraph.Tasks
 				return false;
 			}
 		}
+
+		/// <summary>
+		/// Output this task out to an XML writer.
+		/// </summary>
+		public override void Write(XmlWriter Writer)
+		{
+			Write(Writer, Parameters);
+		}
 	}
 }
diff --git a/Engine/Source/Programs/AutomationTool/BuildGraph/Tasks/LabelBuildTask.cs b/Engine/Source/Programs/AutomationTool/BuildGraph/Tasks/LabelBuildTask.cs
--- a/Engine/Source/Programs/AutomationTool/BuildGraph/Tasks/LabelBuildTask.cs
+++ b/Engine/Source/Programs/AutomationTool/BuildGraph/Tasks/LabelBuildTask.cs
@@ -5,6 +5,7 @@ using System.Collections.Generic;
 using System.Linq;
 using System.Text;
 using System.Threading.Tasks;
+using System.Xml;
 using UnrealBuildTool;
 
 namespace AutomationTool.Tasks
@@ -79,5 +80,13 @@ namespace AutomationTool.Tasks
 			BuildInfoPublisherBase.Get().LabelBuild(StagingInfo, LabelWithPlatform, Parameters.McpConfig);
 			return true;
 		}
+
+		/// <summary>
+		/// Output this task out to an XML writer.
+		/// </summary>
+		public override void Write(XmlWriter Writer)
+		{
+			Write(Writer, Parameters);
+		}
 	}
 }
diff --git a/Engine/Source/Programs/AutomationTool/BuildGraph/Tasks/LogTask.cs b/Engine/Source/Programs/AutomationTool/BuildGraph/Tasks/LogTask.cs
--- a/Engine/Source/Programs/AutomationTool/BuildGraph/Tasks/LogTask.cs
+++ b/Engine/Source/Programs/AutomationTool/BuildGraph/Tasks/LogTask.cs
@@ -5,6 +5,7 @@ using System.IO;
 using System.Linq;
 using System.Text;
 using System.Threading.Tasks;
+using System.Xml;
 using UnrealBuildTool;
 
 namespace BuildGraph.Tasks
@@ -87,5 +88,13 @@ namespace BuildGraph.Tasks
 
 			return true;
 		}
+
+		/// <summary>
+		/// Output this task out to an XML writer.
+		/// </summary>
+		public override void Write(XmlWriter Writer)
+		{
+			Write(Writer, Parameters);
+		}
 	}
 }
diff --git a/Engine/Source/Programs/AutomationTool/BuildGraph/Tasks/MergeTask.cs b/Engine/Source/Programs/AutomationTool/BuildGraph/Tasks/MergeTask.cs
--- a/Engine/Source/Programs/AutomationTool/BuildGraph/Tasks/MergeTask.cs
+++ b/Engine/Source/Programs/AutomationTool/BuildGraph/Tasks/MergeTask.cs
@@ -4,6 +4,7 @@ using System.Collections.Generic;
 using System.Linq;
 using System.Text;
 using System.Threading.Tasks;
+using System.Xml;
 using UnrealBuildTool;
 
 namespace AutomationTool.Tasks
@@ -93,5 +94,13 @@ namespace AutomationTool.Tasks
 			BuildPatchToolBase.Get().Execute(Options);
 			return true;
 		}
+
+		/// <summary>
+		/// Output this task out to an XML writer.
+		/// </summary>
+		public override void Write(XmlWriter Writer)
+		{
+			Write(Writer, Parameters);
+		}
 	}
 }
diff --git a/Engine/Source/Programs/AutomationTool/BuildGraph/Tasks/MsBuildTask.cs b/Engine/Source/Programs/AutomationTool/BuildGraph/Tasks/MsBuildTask.cs
--- a/Engine/Source/Programs/AutomationTool/BuildGraph/Tasks/MsBuildTask.cs
+++ b/Engine/Source/Programs/AutomationTool/BuildGraph/Tasks/MsBuildTask.cs
@@ -151,6 +151,14 @@ namespace BuildGraph.Tasks
 			return true;
 		}
 
+		/// <summary>
+		/// Output this task out to an XML writer.
+		/// </summary>
+		public override void Write(XmlWriter Writer)
+		{
+			Write(Writer, Parameters);
+		}
+
 		/// <summary>
 		/// Find all the build products created by compiling the given project file
 		/// </summary>
diff --git a/Engine/Source/Programs/AutomationTool/BuildGraph/Tasks/PakFileTask.cs b/Engine/Source/Programs/AutomationTool/BuildGraph/Tasks/PakFileTask.cs
--- a/Engine/Source/Programs/AutomationTool/BuildGraph/Tasks/PakFileTask.cs
+++ b/Engine/Source/Programs/AutomationTool/BuildGraph/Tasks/PakFileTask.cs
@@ -5,6 +5,7 @@ using System.IO;
 using System.Linq;
 using System.Text;
 using System.Threading.Tasks;
+using System.Xml;
 using UnrealBuildTool;
 
 namespace BuildGraph.Tasks
@@ -191,5 +192,13 @@ namespace BuildGraph.Tasks
 			}
 			return RelativePath;
 		}
+
+		/// <summary>
+		/// Output this task out to an XML writer.
+		/// </summary>
+		public override void Write(XmlWriter Writer)
+		{
+			Write(Writer, Parameters);
+		}
 	}
 }
diff --git a/Engine/Source/Programs/AutomationTool/BuildGraph/Tasks/PostBuildTask.cs b/Engine/Source/Programs/AutomationTool/BuildGraph/Tasks/PostBuildTask.cs
--- a/Engine/Source/Programs/AutomationTool/BuildGraph/Tasks/PostBuildTask.cs
+++ b/Engine/Source/Programs/AutomationTool/BuildGraph/Tasks/PostBuildTask.cs
@@ -5,6 +5,7 @@ using System.Collections.Generic;
 using System.Linq;
 using System.Text;
 using System.Threading.Tasks;
+using System.Xml;
 using UnrealBuildTool;
 
 namespace BuildGraph.Tasks
@@ -79,5 +80,13 @@ namespace BuildGraph.Tasks
 			BuildInfoPublisherBase.Get().PostBuildInfo(StagingInfo, Parameters.McpConfig);
 			return true;
 		}
+
+		/// <summary>
+		/// Output this task out to an XML writer.
+		/// </summary>
+		public override void Write(XmlWriter Writer)
+		{
+			Write(Writer, Parameters);
+		}
 	}
 }
diff --git a/Engine/Source/Programs/AutomationTool/BuildGraph/Tasks/RenameTask.cs b/Engine/Source/Programs/AutomationTool/BuildGraph/Tasks/RenameTask.cs
--- a/Engine/Source/Programs/AutomationTool/BuildGraph/Tasks/RenameTask.cs
+++ b/Engine/Source/Programs/AutomationTool/BuildGraph/Tasks/RenameTask.cs
@@ -6,6 +6,7 @@ using System.Linq;
 using System.Text;
 using System.Text.RegularExpressions;
 using System.Threading.Tasks;
+using System.Xml;
 using UnrealBuildTool;
 
 namespace BuildGraph.Tasks
@@ -144,5 +145,13 @@ namespace BuildGraph.Tasks
 			}
 			return true;
 		}
+
+		/// <summary>
+		/// Output this task out to an XML writer.
+		/// </summary>
+		public override void Write(XmlWriter Writer)
+		{
+			Write(Writer, Parameters);
+		}
 	}
 }
diff --git a/Engine/Source/Programs/AutomationTool/BuildGraph/Tasks/SetVersionTask.cs b/Engine/Source/Programs/AutomationTool/BuildGraph/Tasks/SetVersionTask.cs
--- a/Engine/Source/Programs/AutomationTool/BuildGraph/Tasks/SetVersionTask.cs
+++ b/Engine/Source/Programs/AutomationTool/BuildGraph/Tasks/SetVersionTask.cs
@@ -3,6 +3,7 @@ using System.Collections.Generic;
 using System.Linq;
 using System.Text;
 using System.Threading.Tasks;
+using System.Xml;
 using UnrealBuildTool;
 
 namespace AutomationTool.Tasks
@@ -86,5 +87,13 @@ namespace AutomationTool.Tasks
 			BuildProducts.UnionWith(VersionFiles);
 			return true;
 		}
+
+		/// <summary>
+		/// Output this task out to an XML writer.
+		/// </summary>
+		public override void Write(XmlWriter Writer)
+		{
+			Write(Writer, Parameters);
+		}
 	}
 }
diff --git a/Engine/Source/Programs/AutomationTool/BuildGraph/Tasks/SignTask.cs b/Engine/Source/Programs/AutomationTool/BuildGraph/Tasks/SignTask.cs
--- a/Engine/Source/Programs/AutomationTool/BuildGraph/Tasks/SignTask.cs
+++ b/Engine/Source/Programs/AutomationTool/BuildGraph/Tasks/SignTask.cs
@@ -4,6 +4,7 @@ using System.Collections.Generic;
 using System.Linq;
 using System.Text;
 using System.Threading.Tasks;
+using System.Xml;
 using UnrealBuildTool;
 
 namespace BuildGraph.Tasks
@@ -71,5 +72,13 @@ namespace BuildGraph.Tasks
 			BuildProducts.UnionWith(Files);
 			return true;
 		}
+
+		/// <summary>
+		/// Output this task out to an XML writer.
+		/// </summary>
+		public override void Write(XmlWriter Writer)
+		{
+			Write(Writer, Parameters);
+		}
 	}
 }
diff --git a/Engine/Source/Programs/AutomationTool/BuildGraph/Tasks/SpawnTask.cs b/Engine/Source/Programs/AutomationTool/BuildGraph/Tasks/SpawnTask.cs
--- a/Engine/Source/Programs/AutomationTool/BuildGraph/Tasks/SpawnTask.cs
+++ b/Engine/Source/Programs/AutomationTool/BuildGraph/Tasks/SpawnTask.cs
@@ -4,6 +4,7 @@ using System.Collections.Generic;
 using System.Linq;
 using System.Text;
 using System.Threading.Tasks;
+using System.Xml;
 using UnrealBuildTool;
 
 namespace BuildGraph.Tasks
@@ -58,5 +59,13 @@ namespace BuildGraph.Tasks
 			CommandUtils.Run(Parameters.Exe, Parameters.Arguments);
 			return true;
 		}
+
+		/// <summary>
+		/// Output this task out to an XML writer.
+		/// </summary>
+		public override void Write(XmlWriter Writer)
+		{
+			Write(Writer, Parameters);
+		}
 	}
 }
diff --git a/Engine/Source/Programs/AutomationTool/BuildGraph/Tasks/StageTask.cs b/Engine/Source/Programs/AutomationTool/BuildGraph/Tasks/StageTask.cs
--- a/Engine/Source/Programs/AutomationTool/BuildGraph/Tasks/StageTask.cs
+++ b/Engine/Source/Programs/AutomationTool/BuildGraph/Tasks/StageTask.cs
@@ -4,6 +4,7 @@ using System.Collections.Generic;
 using System.Linq;
 using System.Text;
 using System.Threading.Tasks;
+using System.Xml;
 using UnrealBuildTool;
 
 namespace BuildGraph.Tasks
@@ -179,5 +180,13 @@ namespace BuildGraph.Tasks
 			BuildProducts.UnionWith(TargetFiles);
 			return true;
 		}
+
+		/// <summary>
+		/// Output this task out to an XML writer.
+		/// </summary>
+		public override void Write(XmlWriter Writer)
+		{
+			Write(Writer, Parameters);
+		}
 	}
 }
diff --git a/Engine/Source/Programs/AutomationTool/BuildGraph/Tasks/StripTask.cs b/Engine/Source/Programs/AutomationTool/BuildGraph/Tasks/StripTask.cs
--- a/Engine/Source/Programs/AutomationTool/BuildGraph/Tasks/StripTask.cs
+++ b/Engine/Source/Programs/AutomationTool/BuildGraph/Tasks/StripTask.cs
@@ -4,6 +4,7 @@ using System.Collections.Generic;
 using System.Linq;
 using System.Text;
 using System.Threading.Tasks;
+using System.Xml;
 using UnrealBuildTool;
 
 namespace BuildGraph.Tasks
@@ -112,5 +113,13 @@ namespace BuildGraph.Tasks
 			BuildProducts.UnionWith(TargetFiles);
 			return true;
 		}
+
+		/// <summary>
+		/// Output this task out to an XML writer.
+		/// </summary>
+		public override void Write(XmlWriter Writer)
+		{
+			Write(Writer, Parameters);
+		}
 	}
 }
diff --git a/Engine/Source/Programs/AutomationTool/BuildGraph/Tasks/TagTask.cs b/Engine/Source/Programs/AutomationTool/BuildGraph/Tasks/TagTask.cs
--- a/Engine/Source/Programs/AutomationTool/BuildGraph/Tasks/TagTask.cs
+++ b/Engine/Source/Programs/AutomationTool/BuildGraph/Tasks/TagTask.cs
@@ -4,6 +4,7 @@ using System.IO;
 using System.Linq;
 using System.Text;
 using System.Threading.Tasks;
+using System.Xml;
 using UnrealBuildTool;
 
 namespace AutomationTool.Tasks
@@ -147,5 +148,13 @@ namespace AutomationTool.Tasks
 			}
 			return Rules;
 		}
+
+		/// <summary>
+		/// Output this task out to an XML writer.
+		/// </summary>
+		public override void Write(XmlWriter Writer)
+		{
+			Write(Writer, Parameters);
+		}
 	}
 }
diff --git a/Engine/Source/Programs/AutomationTool/BuildGraph/Tasks/ZipTask.cs b/Engine/Source/Programs/AutomationTool/BuildGraph/Tasks/ZipTask.cs
--- a/Engine/Source/Programs/AutomationTool/BuildGraph/Tasks/ZipTask.cs
+++ b/Engine/Source/Programs/AutomationTool/BuildGraph/Tasks/ZipTask.cs
@@ -4,6 +4,7 @@ using System.Collections.Generic;
 using System.Linq;
 using System.Text;
 using System.Threading.Tasks;
+using System.Xml;
 using UnrealBuildTool;
 
 namespace BuildGraph.Tasks
@@ -94,5 +95,13 @@ namespace BuildGraph.Tasks
 			BuildProducts.Add(ArchiveFile);
 			return true;
 		}
+
+		/// <summary>
+		/// Output this task out to an XML writer.
+		/// </summary>
+		public override void Write(XmlWriter Writer)
+		{
+			Write(Writer, Parameters);
+		}
 	}
 }
diff --git a/Engine/Source/Programs/AutomationTool/GUBP/GUBP.cs b/Engine/Source/Programs/AutomationTool/GUBP/GUBP.cs
--- a/Engine/Source/Programs/AutomationTool/GUBP/GUBP.cs
+++ b/Engine/Source/Programs/AutomationTool/GUBP/GUBP.cs
@@ -97,7 +97,7 @@ public partial class GUBP : BuildCommand
         bool bSaveSharedTempStorage = P4Enabled && (IsBuildMachine || GlobalCommandLine.UseLocalBuildStorage);
 
         // Get the branch hacker options.
-        string BranchName = P4Enabled ? P4Env.BuildRootP4 : ParseParamValue("BranchName", "");
+        string BranchName = ParseParamValue("BranchName", P4Enabled ? P4Env.BuildRootP4 : "");
         GUBPBranchHacker.BranchOptions BranchOptions = GetBranchOptions(BranchName);
 
 		foreach (var Kind in BranchOptions.MonolithicsToRemove)
diff --git a/Engine/Source/Programs/AutomationTool/Scripts/AutomationScripts.Automation.csproj b/Engine/Source/Programs/AutomationTool/Scripts/AutomationScripts.Automation.csproj
--- a/Engine/Source/Programs/AutomationTool/Scripts/AutomationScripts.Automation.csproj
+++ b/Engine/Source/Programs/AutomationTool/Scripts/AutomationScripts.Automation.csproj
@@ -64,6 +64,7 @@
     <Compile Include="BuildCommonTools.Automation.cs" />
     <Compile Include="BuildCookRun.Automation.cs" />
     <Compile Include="BuildDocumentation.Automation.cs" />
+    <Compile Include="BuildForUGS.Automation.cs" />
     <Compile Include="BuildPluginCommand.Automation.cs" />
     <Compile Include="BuildProjectCommand.Automation.cs" />
     <Compile Include="BuildTargetCommand.Automation.cs" />
diff --git a/Engine/Source/Programs/AutomationTool/Scripts/BuildForUGS.Automation.cs b/Engine/Source/Programs/AutomationTool/Scripts/BuildForUGS.Automation.cs
new file mode 100644
--- /dev/null
+++ b/Engine/Source/Programs/AutomationTool/Scripts/BuildForUGS.Automation.cs
@@ -0,0 +1,149 @@
+﻿// Copyright 1998-2016 Epic Games, Inc. All Rights Reserved.
+
+using System;
+using System.Collections.Generic;
+using System.IO;
+using System.Linq;
+using System.Text;
+using System.Threading.Tasks;
+using UnrealBuildTool;
+
+namespace AutomationTool
+{
+	[RequireP4]
+	class BuildForUGS : BuildCommand
+	{
+		public override void ExecuteBuild()
+		{
+			// Parse the target list
+			string[] Targets = ParseParamValues("Target");
+			if(Targets.Length == 0)
+			{
+				throw new AutomationException("No targets specified (eg. -Target=\"UE4Editor Win64 Development\")");
+			}
+
+			// Parse the archive path
+			string ArchivePath = ParseParamValue("Archive");
+			if(ArchivePath != null && (!ArchivePath.StartsWith("//") || ArchivePath.Sum(x => (x == '/')? 1 : 0) < 4))
+			{
+				throw new AutomationException("Archive path is not a valid depot filename");
+			}
+
+			// Prepare the build agenda
+			UE4Build.BuildAgenda Agenda = new UE4Build.BuildAgenda();
+			foreach(string Target in Targets)
+			{
+				string[] Tokens = Target.Split(new char[]{ ' ' }, StringSplitOptions.RemoveEmptyEntries);
+
+				UnrealTargetPlatform Platform;
+				UnrealTargetConfiguration Configuration;
+				if(Tokens.Length < 3 || !Enum.TryParse(Tokens[1], true, out Platform) || !Enum.TryParse(Tokens[2], true, out Configuration))
+				{
+					throw new AutomationException("Invalid target '{0}' - expected <TargetName> <Platform> <Configuration>");
+				}
+
+				Agenda.AddTarget(Tokens[0], Platform, Configuration, InAddArgs: String.Join(" ", Tokens.Skip(3)));
+			}
+
+			// Build everything
+			UE4Build Builder = new UE4Build(this);
+			Builder.Build(Agenda, InDeleteBuildProducts: true, InUpdateVersionFiles: true, InForceNoXGE: true);
+
+			// Include the build products for UAT and UBT if required
+			if(ParseParam("WithUAT"))
+			{
+				Builder.AddUATFilesToBuildProducts();
+			}
+			if(ParseParam("WithUBT"))
+			{
+				Builder.AddUBTFilesToBuildProducts();
+			}
+
+			// Archive the build products
+			if(ArchivePath != null)
+			{
+				// Create an output folder
+				string OutputFolder = Path.Combine(CommandUtils.CmdEnv.LocalRoot, "ArchiveForUGS");
+				Directory.CreateDirectory(OutputFolder);
+
+				// Create a temp folder for storing stripped PDB files
+				string SymbolsFolder = Path.Combine(OutputFolder, "Symbols");
+				Directory.CreateDirectory(SymbolsFolder);
+
+				// Get the Windows toolchain
+				UEToolChain WindowsToolChain = UEBuildPlatform.GetBuildPlatform(UnrealTargetPlatform.Win64).CreateContext(null).CreateToolChain(CPPTargetPlatform.Win64);
+
+				// Figure out all the files for the archive
+				Ionic.Zip.ZipFile Zip = new Ionic.Zip.ZipFile();
+				Zip.UseZip64WhenSaving = Ionic.Zip.Zip64Option.Always;
+				foreach(string BuildProduct in Builder.BuildProductFiles)
+				{
+					if(!File.Exists(BuildProduct))
+					{
+						throw new AutomationException("Missing build product: {0}", BuildProduct);
+					}
+					if(BuildProduct.EndsWith(".pdb", StringComparison.InvariantCultureIgnoreCase))
+					{
+						string StrippedFileName = CommandUtils.MakeRerootedFilePath(BuildProduct, CommandUtils.CmdEnv.LocalRoot, SymbolsFolder);
+						Directory.CreateDirectory(Path.GetDirectoryName(StrippedFileName));
+						WindowsToolChain.StripSymbols(BuildProduct, StrippedFileName);
+						Zip.AddFile(StrippedFileName, Path.GetDirectoryName(CommandUtils.StripBaseDirectory(StrippedFileName, SymbolsFolder)));
+					}
+					else
+					{
+						Zip.AddFile(BuildProduct, Path.GetDirectoryName(CommandUtils.StripBaseDirectory(BuildProduct, CommandUtils.CmdEnv.LocalRoot)));
+					}
+				}
+				// Create the zip file
+				string ZipFileName = Path.Combine(OutputFolder, "Archive.zip");
+				Console.WriteLine("Writing {0}...", ZipFileName);
+				Zip.Save(ZipFileName);
+
+				// Submit it to Perforce if required
+				if(CommandUtils.AllowSubmit)
+				{
+					// Delete any existing clientspec for submitting
+					string ClientName = Environment.MachineName + "_BuildForUGS";
+
+					// Create a brand new one
+					P4ClientInfo Client = new P4ClientInfo();
+					Client.Owner = Environment.UserName;
+					Client.Host = Environment.MachineName;
+					Client.Stream = ArchivePath.Substring(0, ArchivePath.IndexOf('/', ArchivePath.IndexOf('/', 2) + 1));
+					Client.RootPath = Path.Combine(OutputFolder, "Perforce");
+					Client.Name = ClientName;
+					Client.Options = P4ClientOption.NoAllWrite | P4ClientOption.NoClobber | P4ClientOption.NoCompress | P4ClientOption.Unlocked | P4ClientOption.NoModTime | P4ClientOption.RmDir;
+					Client.LineEnd = P4LineEnd.Local;
+					P4.CreateClient(Client, AllowSpew: false);
+
+					// Create a new P4 connection for this workspace
+					P4Connection SubmitP4 = new P4Connection(Client.Owner, Client.Name);
+					SubmitP4.Revert("-k //...");
+
+					// Figure out where the zip file has to go in Perforce
+					P4WhereRecord WhereZipFile = SubmitP4.Where(ArchivePath, false).FirstOrDefault(x => !x.bUnmap && x.Path != null);
+					if(WhereZipFile == null)
+					{
+						throw new AutomationException("Couldn't locate {0} in this workspace");
+					}
+
+					// Get the latest version of it
+					int NewCL = SubmitP4.CreateChange(Description: String.Format("[CL {0}] Updated binaries", P4Env.Changelist));
+					SubmitP4.Sync(String.Format("-k \"{0}\"", ArchivePath), AllowSpew:false);
+					CommandUtils.CopyFile(ZipFileName, WhereZipFile.Path);
+					SubmitP4.Add(NewCL, String.Format("\"{0}\"", ArchivePath));
+					SubmitP4.Edit(NewCL, String.Format("\"{0}\"", ArchivePath));
+
+					// Submit it
+					int SubmittedCL;
+					SubmitP4.Submit(NewCL, out SubmittedCL);
+					if(SubmittedCL <= 0)
+					{
+						throw new AutomationException("Submit failed.");
+					}
+					Console.WriteLine("Submitted in changelist {0}", SubmittedCL);
+				}
+			}
+		}
+	}
+}
diff --git a/Engine/Source/Programs/AutomationTool/Scripts/RocketBuild.Automation.cs b/Engine/Source/Programs/AutomationTool/Scripts/RocketBuild.Automation.cs
--- a/Engine/Source/Programs/AutomationTool/Scripts/RocketBuild.Automation.cs
+++ b/Engine/Source/Programs/AutomationTool/Scripts/RocketBuild.Automation.cs
@@ -128,12 +128,28 @@ namespace Rocket
 				BranchConfig.AddNode(new GatherRocketNode(BranchConfig, HostPlatform, TargetPlatforms, LocalOutputDir));
 
 				// Add a node for GitHub promotions
-				if(HostPlatform == UnrealTargetPlatform.Win64)
+				if(HostPlatform == UnrealTargetPlatform.Win64 && !BranchConfig.JobInfo.IsPreflight)
 				{
 					string GitConfigRelativePath = "Engine/Build/Git/UnrealBot.ini";
 					if(CommandUtils.FileExists(CommandUtils.CombinePaths(CommandUtils.CmdEnv.LocalRoot, GitConfigRelativePath)))
 					{
-						BranchConfig.AddNode(new LabelGitHubPromotion(HostPlatform, BranchConfig.HostPlatforms));
+						string BranchName = null;
+						if (BranchConfig.BranchName == "//UE4/Main")
+						{
+							BranchName = "Main";
+						}
+						else if (BranchConfig.BranchName.StartsWith("//UE4/Release-4."))
+						{
+							int MinorVersion;
+							if(int.TryParse(BranchConfig.BranchName.Substring(BranchConfig.BranchName.IndexOf('.') + 1), out MinorVersion))
+							{
+								BranchName = "4." + MinorVersion.ToString();
+							}
+						}
+						if (BranchName != null)
+						{
+							BranchConfig.AddNode(new RunGithubPromotion(HostPlatform, BranchConfig.HostPlatforms, BranchName, CommandUtils.P4Env.Changelist));
+						}
 					}
 				}
 
@@ -326,10 +342,15 @@ namespace Rocket
         }
     }
 
-	public class LabelGitHubPromotion : GUBP.HostPlatformNode
+	public class RunGithubPromotion : GUBP.HostPlatformNode
 	{
-		public LabelGitHubPromotion(UnrealTargetPlatform HostPlatform, List<UnrealTargetPlatform> ForHostPlatforms) : base(HostPlatform)
+		string BranchName;
+		int Changelist;
+
+		public RunGithubPromotion(UnrealTargetPlatform HostPlatform, List<UnrealTargetPlatform> ForHostPlatforms, string InBranchName, int InChangelist) : base(HostPlatform)
 		{
+			BranchName = InBranchName;
+			Changelist = InChangelist;
 			foreach(UnrealTargetPlatform ForHostPlatform in ForHostPlatforms)
 			{
 				AddPseudodependency(GUBP.RootEditorNode.StaticGetFullName(ForHostPlatform));
@@ -340,7 +361,7 @@ namespace Rocket
 
 		public static string StaticGetFullName(UnrealTargetPlatform HostPlatform)
 		{
-			return "LabelGitHubPromotion" + StaticGetHostPlatformSuffix(HostPlatform);
+			return "RunGithubPromotion" + StaticGetHostPlatformSuffix(HostPlatform);
 		}
 
 		public override string GetFullName()
@@ -356,9 +377,9 @@ namespace Rocket
 		public override void DoBuild(GUBP bp)
 		{
 			// Label everything in the branch at this changelist
-			if(CommandUtils.AllowSubmit)
+			if(CommandUtils.IsBuildMachine)
 			{
-				CommandUtils.P4.MakeDownstreamLabel(CommandUtils.P4Env, "GitHub-Promotion", null);
+				CommandUtils.Run("ectool.exe", String.Format("runProcedure GitHub --procedureName \"Run Promotion\" --actualParameter \"Branch={0}\" --actualParameter \"CL={1}\"", BranchName, Changelist), Options: CommandUtils.ERunOptions.None);
 			}
 
 			// Create a dummy build product
@@ -1192,8 +1213,9 @@ namespace Rocket
 				IniLines.Add("");
 			}
 			
-			// Write information about platforms installed in a Rocket build
-			IniLines.Add("[InstalledPlatforms]");
+			// Create list of platform configurations installed in a Rocket build
+			List<InstalledPlatformInfo.InstalledPlatformConfiguration> InstalledConfigs = new List<InstalledPlatformInfo.InstalledPlatformConfiguration>();
+
 			foreach (UnrealTargetPlatform CodeTargetPlatform in CodeTargetPlatforms)
 			{
 				// Bit of a hack to mark these platforms as available in any type of project
@@ -1234,16 +1256,24 @@ namespace Rocket
 					}
 					string ReceiptFileName = TargetReceipt.GetDefaultPath(OutputEnginePath, "UE4Game", CodeTargetPlatform, EngineConfiguration, Architecture);
 
+					// Hack to write Android architecture until we retrieve it from build process -
+					// has to be after getting receipt path as one is shared for all architectures
+					if (string.IsNullOrEmpty(Architecture) && CodeTargetPlatform == UnrealTargetPlatform.Android)
+					{
+						Architecture = "-armv7-es2";
+					}
+
 					if (File.Exists(ReceiptFileName))
 					{
 						// Strip the output folder so that this can be used on any machine
 						ReceiptFileName = new FileReference(ReceiptFileName).MakeRelativeTo(new DirectoryReference(OutputDir));
-						IniLines.Add(string.Format("+InstalledPlatformConfigurations=(PlatformName=\"{0}\", Configuration=\"{1}\", RequiredFile=\"{2}\", ProjectType=\"{3}\", bCanBeDisplayed={4})",
-													CodeTargetPlatform.ToString(), CodeTargetConfiguration.ToString(), ReceiptFileName, ProjectType.ToString(), bCanBeDisplayed.ToString()));
+						InstalledConfigs.Add(new InstalledPlatformInfo.InstalledPlatformConfiguration(CodeTargetConfiguration, CodeTargetPlatform, TargetRules.TargetType.Game, Architecture, ReceiptFileName, ProjectType, bCanBeDisplayed));
 					}
 				}
 			}
 
+			UnrealBuildTool.InstalledPlatformInfo.WriteConfigFileEntries(InstalledConfigs, ref IniLines);
+
 			// Write Rocket specific Analytics settings
 			IniLines.Add("");
 			IniLines.Add("[Analytics]");
diff --git a/Engine/Source/Programs/UnrealBuildTool/Configuration/UEBuildTarget.cs b/Engine/Source/Programs/UnrealBuildTool/Configuration/UEBuildTarget.cs
--- a/Engine/Source/Programs/UnrealBuildTool/Configuration/UEBuildTarget.cs
+++ b/Engine/Source/Programs/UnrealBuildTool/Configuration/UEBuildTarget.cs
@@ -1682,17 +1682,13 @@ namespace UnrealBuildTool
 		{
 			string FileListPath = "../Intermediate/Build/ExternalFiles.xml";
 
-			// Find all the external modules
-			HashSet<string> ModuleNames = new HashSet<string>();
+			// Find all the modules we depend on
+			HashSet<UEBuildModule> Modules = new HashSet<UEBuildModule>();
 			foreach (UEBuildBinary Binary in AppBinaries)
 			{
 				foreach (UEBuildModule Module in Binary.GetAllDependencyModules(bIncludeDynamicallyLoaded: false, bForceCircular: false))
 				{
-					UEBuildExternalModule ExternalModule = Module as UEBuildExternalModule;
-					if (ExternalModule != null)
-					{
-						ModuleNames.Add(ExternalModule.Name);
-					}
+					Modules.Add(Module);
 				}
 			}
 
@@ -1702,74 +1698,84 @@ namespace UnrealBuildTool
 			// Get the platform we're building for
 			UEBuildPlatform BuildPlatform = UEBuildPlatform.GetBuildPlatform(Platform);
 
-			// Add all their include paths
-			foreach (string ModuleName in ModuleNames)
+			foreach (UEBuildModule Module in Modules)
 			{
 				// Create the module rules
 				FileReference ModuleRulesFileName;
-				ModuleRules Rules = CreateModuleRulesAndSetDefaults(ModuleName, out ModuleRulesFileName);
+				ModuleRules Rules = CreateModuleRulesAndSetDefaults(Module.Name, out ModuleRulesFileName);
 
-				// Add the rules file itself
-				FileNames.Add(ModuleRulesFileName.FullName);
+				// Add Additional Bundle Resources for all modules
+				foreach (UEBuildBundleResource Resource in Rules.AdditionalBundleResources)
+				{
+					FileNames.Add(Resource.ResourcePath);
+				}
 
-				// Get a list of all the library paths
-				List<string> LibraryPaths = new List<string>();
-				LibraryPaths.Add(Directory.GetCurrentDirectory());
-				LibraryPaths.AddRange(Rules.PublicLibraryPaths.Where(x => !x.StartsWith("$(")).Select(x => Path.GetFullPath(x.Replace('/', '\\'))));
+				// Add all the include paths etc. for external modules
+				UEBuildExternalModule ExternalModule = Module as UEBuildExternalModule;
+				if (ExternalModule != null)
+				{
+					// Add the rules file itself
+					FileNames.Add(ModuleRulesFileName.FullName);
 
-				// Get all the extensions to look for
-				List<string> LibraryExtensions = new List<string>();
-				LibraryExtensions.Add(BuildPlatform.GetBinaryExtension(UEBuildBinaryType.StaticLibrary));
-				LibraryExtensions.Add(BuildPlatform.GetBinaryExtension(UEBuildBinaryType.DynamicLinkLibrary));
+					// Get a list of all the library paths
+					List<string> LibraryPaths = new List<string>();
+					LibraryPaths.Add(Directory.GetCurrentDirectory());
+					LibraryPaths.AddRange(Rules.PublicLibraryPaths.Where(x => !x.StartsWith("$(")).Select(x => Path.GetFullPath(x.Replace('/', '\\'))));
 
-				// Add all the libraries
-				foreach (string LibraryExtension in LibraryExtensions)
-				{
-					foreach (string LibraryName in Rules.PublicAdditionalLibraries)
+					// Get all the extensions to look for
+					List<string> LibraryExtensions = new List<string>();
+					LibraryExtensions.Add(BuildPlatform.GetBinaryExtension(UEBuildBinaryType.StaticLibrary));
+					LibraryExtensions.Add(BuildPlatform.GetBinaryExtension(UEBuildBinaryType.DynamicLinkLibrary));
+
+					// Add all the libraries
+					foreach (string LibraryExtension in LibraryExtensions)
 					{
-						foreach (string LibraryPath in LibraryPaths)
+						foreach (string LibraryName in Rules.PublicAdditionalLibraries)
 						{
-							string LibraryFileName = Path.Combine(LibraryPath, LibraryName);
-							if (File.Exists(LibraryFileName))
+							foreach (string LibraryPath in LibraryPaths)
 							{
-								FileNames.Add(LibraryFileName);
-							}
+								string LibraryFileName = Path.Combine(LibraryPath, LibraryName);
+								if (File.Exists(LibraryFileName))
+								{
+									FileNames.Add(LibraryFileName);
+								}
 
-							string UnixLibraryFileName = Path.Combine(LibraryPath, "lib" + LibraryName + LibraryExtension);
-							if (File.Exists(UnixLibraryFileName))
-							{
-								FileNames.Add(UnixLibraryFileName);
+								string UnixLibraryFileName = Path.Combine(LibraryPath, "lib" + LibraryName + LibraryExtension);
+								if (File.Exists(UnixLibraryFileName))
+								{
+									FileNames.Add(UnixLibraryFileName);
+								}
 							}
 						}
 					}
-				}
 
-				// Add all the additional shadow files
-				foreach (string AdditionalShadowFile in Rules.PublicAdditionalShadowFiles)
-				{
-					string ShadowFileName = Path.GetFullPath(AdditionalShadowFile);
-					if (File.Exists(ShadowFileName))
+					// Add all the additional shadow files
+					foreach (string AdditionalShadowFile in Rules.PublicAdditionalShadowFiles)
 					{
-						FileNames.Add(ShadowFileName);
+						string ShadowFileName = Path.GetFullPath(AdditionalShadowFile);
+						if (File.Exists(ShadowFileName))
+						{
+							FileNames.Add(ShadowFileName);
+						}
 					}
-				}
 
-				// Find all the include paths
-				List<string> AllIncludePaths = new List<string>();
-				AllIncludePaths.AddRange(Rules.PublicIncludePaths);
-				AllIncludePaths.AddRange(Rules.PublicSystemIncludePaths);
+					// Find all the include paths
+					List<string> AllIncludePaths = new List<string>();
+					AllIncludePaths.AddRange(Rules.PublicIncludePaths);
+					AllIncludePaths.AddRange(Rules.PublicSystemIncludePaths);
 
-				// Add all the include paths
-				foreach (string IncludePath in AllIncludePaths.Where(x => !x.StartsWith("$(")))
-				{
-					if (Directory.Exists(IncludePath))
+					// Add all the include paths
+					foreach (string IncludePath in AllIncludePaths.Where(x => !x.StartsWith("$(")))
 					{
-						foreach (string IncludeFileName in Directory.EnumerateFiles(IncludePath, "*", SearchOption.AllDirectories))
+						if (Directory.Exists(IncludePath))
 						{
-							string Extension = Path.GetExtension(IncludeFileName).ToLower();
-							if (Extension == ".h" || Extension == ".inl")
+							foreach (string IncludeFileName in Directory.EnumerateFiles(IncludePath, "*", SearchOption.AllDirectories))
 							{
-								FileNames.Add(IncludeFileName);
+								string Extension = Path.GetExtension(IncludeFileName).ToLower();
+								if (Extension == ".h" || Extension == ".inl")
+								{
+									FileNames.Add(IncludeFileName);
+								}
 							}
 						}
 					}
diff --git a/Engine/Source/Programs/UnrealBuildTool/System/InstalledPlatformInfo.cs b/Engine/Source/Programs/UnrealBuildTool/System/InstalledPlatformInfo.cs
--- a/Engine/Source/Programs/UnrealBuildTool/System/InstalledPlatformInfo.cs
+++ b/Engine/Source/Programs/UnrealBuildTool/System/InstalledPlatformInfo.cs
@@ -22,7 +22,7 @@ namespace UnrealBuildTool
 		/// <summary>
 		/// Information about a single installed platform configuration
 		/// </summary>
-		private struct InstalledPlatformConfiguration
+		public struct InstalledPlatformConfiguration
 		{
 			/// Build Configuration of this combination
 			public UnrealTargetConfiguration Configuration;
@@ -30,6 +30,12 @@ namespace UnrealBuildTool
 			/// Platform for this combination
 			public UnrealTargetPlatform Platform;
 
+			/// Type of Platform for this combination
+			public TargetRules.TargetType PlatformType;
+
+			/// Architecture for this combination
+			public string Architecture;
+
 			/// Location of a file that must exist for this combination to be valid (optional)
 			public string RequiredFile;
 
@@ -39,10 +45,12 @@ namespace UnrealBuildTool
 			/// Whether to display this platform as an option even if it is not valid
 			public bool bCanBeDisplayed;
 
-			public InstalledPlatformConfiguration(UnrealTargetConfiguration InConfiguration, UnrealTargetPlatform InPlatform, string InRequiredFile, EProjectType InProjectType, bool bInCanBeDisplayed)
+			public InstalledPlatformConfiguration(UnrealTargetConfiguration InConfiguration, UnrealTargetPlatform InPlatform, TargetRules.TargetType InPlatformType, string InArchitecture, string InRequiredFile, EProjectType InProjectType, bool bInCanBeDisplayed)
 			{
 				Configuration = InConfiguration;
 				Platform = InPlatform;
+				PlatformType = InPlatformType;
+				Architecture = InArchitecture;
 				RequiredFile = InRequiredFile;
 				ProjectType = InProjectType;
 				bCanBeDisplayed = bInCanBeDisplayed;
@@ -113,6 +121,25 @@ namespace UnrealBuildTool
 				bCanCreateEntry = false;
 			}
 
+			string PlatformTypeName;
+			TargetRules.TargetType PlatformType = TargetRules.TargetType.Game;
+			if (ParseSubValue(PlatformConfiguration, "PlatformTypeName=", out PlatformTypeName))
+			{
+				if (!Enum.TryParse(PlatformTypeName, out PlatformType))
+				{
+					Log.TraceWarning("Unable to read Platform Type from {0}, defaulting to Game", PlatformConfiguration);
+					PlatformType = TargetRules.TargetType.Game;
+				}
+			}
+			if (PlatformType == TargetRules.TargetType.Program)
+			{
+				Log.TraceWarning("Program is not a valid PlatformType for an Installed Platform, defaulting to Game");
+				PlatformType = TargetRules.TargetType.Game;
+			}
+
+			string Architecture;
+			ParseSubValue(PlatformConfiguration, "Architecture=", out Architecture);
+
 			string RequiredFile;
 			if (ParseSubValue(PlatformConfiguration, "RequiredFile=", out RequiredFile))
 			{
@@ -140,7 +167,7 @@ namespace UnrealBuildTool
 
 			if (bCanCreateEntry)
 			{
-				InstalledPlatformConfigurations.Add(new InstalledPlatformConfiguration(Configuration, Platform, RequiredFile, ProjectType,bCanBeDisplayed));
+				InstalledPlatformConfigurations.Add(new InstalledPlatformConfiguration(Configuration, Platform, PlatformType, Architecture, RequiredFile, ProjectType, bCanBeDisplayed));
 			}
 		}
 
@@ -176,72 +203,52 @@ namespace UnrealBuildTool
 
 		public bool IsValidConfiguration(UnrealTargetConfiguration Configuration, EProjectType ProjectType = EProjectType.Any)
 		{
-			if (UnrealBuildTool.IsEngineInstalled())
-			{
-				foreach (InstalledPlatformConfiguration PlatformConfiguration in InstalledPlatformConfigurations)
+			return ContainsValidConfiguration(
+				(InstalledPlatformConfiguration CurConfig) =>
 				{
-					if (PlatformConfiguration.Configuration == Configuration)
-					{
-						if (ProjectType == EProjectType.Any || PlatformConfiguration.ProjectType == EProjectType.Any
-						 || PlatformConfiguration.ProjectType == ProjectType)
-						{
-							if (string.IsNullOrEmpty(PlatformConfiguration.RequiredFile)
-								|| File.Exists(PlatformConfiguration.RequiredFile))
-							{
-								return true;
-							}
-						}
-					}
+					return CurConfig.Configuration == Configuration
+						&& (ProjectType == EProjectType.Any || CurConfig.ProjectType == EProjectType.Any
+						|| CurConfig.ProjectType == ProjectType);
 				}
-
-				return false;
-			}
-			return true;
+			);
 		}
 
 		public bool IsValidPlatform(UnrealTargetPlatform Platform, EProjectType ProjectType = EProjectType.Any)
 		{
-			if (UnrealBuildTool.IsEngineInstalled())
-			{
-				foreach (InstalledPlatformConfiguration PlatformConfiguration in InstalledPlatformConfigurations)
+			return ContainsValidConfiguration(
+				(InstalledPlatformConfiguration CurConfig) =>
 				{
-					if (PlatformConfiguration.Platform == Platform)
-					{
-						if (ProjectType == EProjectType.Any || PlatformConfiguration.ProjectType == EProjectType.Any
-						 || PlatformConfiguration.ProjectType == ProjectType)
-						{
-							if (string.IsNullOrEmpty(PlatformConfiguration.RequiredFile)
-								|| File.Exists(PlatformConfiguration.RequiredFile))
-							{
-								return true;
-							}
-						}
-					}
+					return CurConfig.Platform == Platform
+						&& (ProjectType == EProjectType.Any || CurConfig.ProjectType == EProjectType.Any
+						|| CurConfig.ProjectType == ProjectType);
 				}
-
-				return false;
-			}
-			return true;
+			);
 		}
 
 		public bool IsValidPlatformAndConfiguration(UnrealTargetConfiguration Configuration, UnrealTargetPlatform Platform, EProjectType ProjectType = EProjectType.Any)
+		{
+			return ContainsValidConfiguration(
+				(InstalledPlatformConfiguration CurConfig) =>
+				{
+					return CurConfig.Configuration == Configuration && CurConfig.Platform == Platform
+						&& (ProjectType == EProjectType.Any || CurConfig.ProjectType == EProjectType.Any
+						|| CurConfig.ProjectType == ProjectType);
+				}
+			);
+		}
+
+		private bool ContainsValidConfiguration(Predicate<InstalledPlatformConfiguration> ConfigFilter)
 		{
 			if (UnrealBuildTool.IsEngineInstalled())
 			{
 				foreach (InstalledPlatformConfiguration PlatformConfiguration in InstalledPlatformConfigurations)
 				{
-					if (PlatformConfiguration.Configuration == Configuration
-					 && PlatformConfiguration.Platform == Platform)
+					// Check whether filter accepts this configuration and it has required file
+					if (ConfigFilter(PlatformConfiguration)
+					&& (string.IsNullOrEmpty(PlatformConfiguration.RequiredFile)
+					|| File.Exists(PlatformConfiguration.RequiredFile)))
 					{
-						if (ProjectType == EProjectType.Any || PlatformConfiguration.ProjectType == EProjectType.Any
-						 || PlatformConfiguration.ProjectType == ProjectType)
-						{
-							if (string.IsNullOrEmpty(PlatformConfiguration.RequiredFile)
-							 || File.Exists(PlatformConfiguration.RequiredFile))
-							{
-								return true;
-							}
-						}
+						return true;
 					}
 				}
 
@@ -249,5 +256,48 @@ namespace UnrealBuildTool
 			}
 			return true;
 		}
+
+		public static void WriteConfigFileEntries(List<InstalledPlatformConfiguration> Configs, ref List<String> OutEntries)
+		{
+			// Write config section header
+			OutEntries.Add("[InstalledPlatforms]");
+
+			foreach (InstalledPlatformConfiguration Config in Configs)
+			{
+				WriteConfigFileEntry(Config, ref OutEntries);
+			}
+		}
+
+		private static void WriteConfigFileEntry(InstalledPlatformConfiguration Config, ref List<String> OutEntries)
+		{
+			string ConfigDescription = "+InstalledPlatformConfigurations=(";
+			if (Config.Platform != UnrealTargetPlatform.Unknown)
+			{
+				ConfigDescription += string.Format("PlatformName=\"{0}\", ", Config.Platform.ToString());
+			}
+			if (Config.Configuration != UnrealTargetConfiguration.Unknown)
+			{
+				ConfigDescription += string.Format("Configuration=\"{0}\", ", Config.Configuration.ToString());
+			}
+			if (Config.PlatformType != TargetRules.TargetType.Program)
+			{
+				ConfigDescription += string.Format("PlatformType=\"{0}\", ", Config.PlatformType.ToString());
+			}
+			if (!string.IsNullOrEmpty(Config.Architecture))
+			{
+				ConfigDescription += string.Format("Architecture=\"{0}\", ", Config.Architecture);
+			}
+			if (!string.IsNullOrEmpty(Config.RequiredFile))
+			{
+				ConfigDescription += string.Format("RequiredFile=\"{0}\", ", Config.RequiredFile);
+			}
+			if (Config.ProjectType != EProjectType.Unknown)
+			{
+				ConfigDescription += string.Format("ProjectType=\"{0}\", ", Config.ProjectType.ToString());
+			}
+			ConfigDescription += string.Format("bCanBeDisplayed={0})", Config.bCanBeDisplayed.ToString());
+
+			OutEntries.Add(ConfigDescription);
+		}
 	}
 }
diff --git a/Engine/Source/Programs/UnrealBuildTool/System/LocalExecutor.cs b/Engine/Source/Programs/UnrealBuildTool/System/LocalExecutor.cs
--- a/Engine/Source/Programs/UnrealBuildTool/System/LocalExecutor.cs
+++ b/Engine/Source/Programs/UnrealBuildTool/System/LocalExecutor.cs
@@ -562,7 +562,7 @@ namespace UnrealBuildTool
 				TotalThreadSeconds += ThreadSeconds;
 			}
 
-			Log.TraceInformation("-------- End Detailed Actions Stats -----------------------------------------------------------");
+			Log.WriteLineIf(BuildConfiguration.bLogDetailedActionStats, LogEventType.Console, "-------- End Detailed Actions Stats -----------------------------------------------------------");
 
 			// Log total CPU seconds and numbers of processors involved in tasks.
 			Log.WriteLineIf(BuildConfiguration.bLogDetailedActionStats || BuildConfiguration.bPrintDebugInfo,
diff --git a/Engine/Source/Programs/UnrealBuildTool/System/ProjectFileGenerator.cs b/Engine/Source/Programs/UnrealBuildTool/System/ProjectFileGenerator.cs
--- a/Engine/Source/Programs/UnrealBuildTool/System/ProjectFileGenerator.cs
+++ b/Engine/Source/Programs/UnrealBuildTool/System/ProjectFileGenerator.cs
@@ -772,6 +772,11 @@ namespace UnrealBuildTool
 
 			if( bGeneratingGameProjectFiles || UnrealBuildTool.IsEngineInstalled() )
 			{
+				if (OnlyGameProject == null)
+				{
+					throw new BuildException("A game project path was not specified, which is required when generating project files using an installed build or passing -game on the command line");
+				}
+
 				GameProjectName = OnlyGameProject.GetFileNameWithoutExtension();
 				if (String.IsNullOrEmpty(GameProjectName))
 				{
diff --git a/Engine/Source/Programs/UnrealBuildTool/Windows/VCToolChain.cs b/Engine/Source/Programs/UnrealBuildTool/Windows/VCToolChain.cs
--- a/Engine/Source/Programs/UnrealBuildTool/Windows/VCToolChain.cs
+++ b/Engine/Source/Programs/UnrealBuildTool/Windows/VCToolChain.cs
@@ -1833,7 +1833,13 @@ namespace UnrealBuildTool
 		public override void StripSymbols(string SourceFileName, string TargetFileName)
 		{
 			ProcessStartInfo StartInfo = new ProcessStartInfo();
-			StartInfo.FileName = Path.Combine(Environment.GetFolderPath(Environment.SpecialFolder.ProgramFilesX86), "MSBuild", "Microsoft", "VisualStudio", "v12.0", "AppxPackage", "PDBCopy.exe");
+			string PDBCopyPath = Path.Combine(Environment.GetFolderPath(Environment.SpecialFolder.ProgramFilesX86), "MSBuild", "Microsoft", "VisualStudio", "v14.0", "AppxPackage", "PDBCopy.exe");
+			if (!File.Exists(PDBCopyPath))
+			{
+				// Fall back on VS2013 version
+				PDBCopyPath = Path.Combine(Environment.GetFolderPath(Environment.SpecialFolder.ProgramFilesX86), "MSBuild", "Microsoft", "VisualStudio", "v12.0", "AppxPackage", "PDBCopy.exe");
+			}
+			StartInfo.FileName = PDBCopyPath;
 			StartInfo.Arguments = String.Format("\"{0}\" \"{1}\" -p", SourceFileName, TargetFileName);
 			StartInfo.UseShellExecute = false;
 			StartInfo.CreateNoWindow = true;
diff --git a/Engine/Source/Runtime/Core/Public/UObject/ReleaseObjectVersion.h b/Engine/Source/Runtime/Core/Public/UObject/ReleaseObjectVersion.h
--- a/Engine/Source/Runtime/Core/Public/UObject/ReleaseObjectVersion.h
+++ b/Engine/Source/Runtime/Core/Public/UObject/ReleaseObjectVersion.h
@@ -1,4 +1,4 @@
-// Copyright 1998-2015 Epic Games, Inc. All Rights Reserved.
+// Copyright 1998-2016 Epic Games, Inc. All Rights Reserved.
 #pragma once
 
 // Custom serialization version for changes made in Dev-Core stream
diff --git a/Engine/Source/Runtime/Slate/Private/Framework/Text/RunUtils.cpp b/Engine/Source/Runtime/Slate/Private/Framework/Text/RunUtils.cpp
--- a/Engine/Source/Runtime/Slate/Private/Framework/Text/RunUtils.cpp
+++ b/Engine/Source/Runtime/Slate/Private/Framework/Text/RunUtils.cpp
@@ -1,4 +1,4 @@
-// Copyright 1998-2015 Epic Games, Inc. All Rights Reserved.
+// Copyright 1998-2016 Epic Games, Inc. All Rights Reserved.
 
 #include "SlatePrivatePCH.h"
 #include "RunUtils.h"
diff --git a/Engine/Source/Runtime/Slate/Public/Framework/Text/RunUtils.h b/Engine/Source/Runtime/Slate/Public/Framework/Text/RunUtils.h
--- a/Engine/Source/Runtime/Slate/Public/Framework/Text/RunUtils.h
+++ b/Engine/Source/Runtime/Slate/Public/Framework/Text/RunUtils.h
@@ -1,4 +1,4 @@
-// Copyright 1998-2015 Epic Games, Inc. All Rights Reserved.
+// Copyright 1998-2016 Epic Games, Inc. All Rights Reserved.
 #pragma once
 
 namespace RunUtils
diff --git a/Engine/Source/Runtime/Slate/Public/Framework/Text/ShapedTextCacheFwd.h b/Engine/Source/Runtime/Slate/Public/Framework/Text/ShapedTextCacheFwd.h
--- a/Engine/Source/Runtime/Slate/Public/Framework/Text/ShapedTextCacheFwd.h
+++ b/Engine/Source/Runtime/Slate/Public/Framework/Text/ShapedTextCacheFwd.h
@@ -1,4 +1,4 @@
-// Copyright 1998-2015 Epic Games, Inc. All Rights Reserved.
+// Copyright 1998-2016 Epic Games, Inc. All Rights Reserved.
 
 #pragma once
 
