fix for trying to load /Engine/Maps/Entry instead of GameDefaultMap

diff --git a/Engine/Source/Programs/AutomationTool/Scripts/CopyBuildToStagingDirectory.Automation.cs b/Engine/Source/Programs/AutomationTool/Scripts/CopyBuildToStagingDirectory.Automation.cs
--- a/Engine/Source/Programs/AutomationTool/Scripts/CopyBuildToStagingDirectory.Automation.cs
+++ b/Engine/Source/Programs/AutomationTool/Scripts/CopyBuildToStagingDirectory.Automation.cs
@@ -784,8 +784,7 @@ public partial class Project : CommandUtils
 				int ConfigRootIdx = SrcDirectoryPath.LastIndexOf(ConfigDirectory);
 				string SubpathUnderConfig = (ConfigRootIdx != -1) ? SrcDirectoryPath.Substring(ConfigRootIdx + ConfigDirectory.Length) : "Unknown";
 
-				string FileName = SrcDirectoryPath.GetHashCode().ToString() + "_DefaultEngine.ini";
-				string NewIniFilename = CombinePaths(SC.ProjectRoot, "Saved", "Temp", SC.PlatformDir, SubpathUnderConfig, FileName);
+				string NewIniFilename = CombinePaths(SC.ProjectRoot, "Saved", "Temp", SC.PlatformDir, SubpathUnderConfig, "DefaultEngine.ini");
 				InternalUtils.SafeCreateDirectory(Path.GetDirectoryName(NewIniFilename), true);
 				InternalUtils.SafeCopyFile(Src, NewIniFilename, bFilterSpecialLinesFromIniFiles:true);
 				Src = NewIniFilename;
