Copying //UE4/Dev-Build to //UE4/Dev-Main (Source: //UE4/Dev-Build @ 2959429)

==========================
MAJOR FEATURES + CHANGES
==========================

Change 2945275 on 2016/04/15 by [REDACTED_PERSON]

	BuildCommonTools: Stop forcing the DeleteBuildProducts flag to true; just respect the -Clean parameter from the command line.

Change 2946668 on 2016/04/18 by [REDACTED_PERSON]

	EC: Prevent lookbehind for UBT error strings on Mac.

Change 2952657 on 2016/04/22 by [REDACTED_PERSON]

	UGS: Require the user to explicitly choose to show *.uprojectdirs files, to discourage people from selecting the first thing they see in a synced branch. The uprojectdirs workflow is only used by [REDACTED_PERSON], but catches a lot of people out.

Change 2954256 on 2016/04/25 by [REDACTED_PERSON]

	EC: Fix lines starting with error: and warning: being swallowed by the postprocessor. Also remove confusing 4 line look-behind on generic error and warning messages.

Change 2954449 on 2016/04/25 by [REDACTED_PERSON]

	Use the original application name for log files (and for the prefix in stdout), rather than the application name after the host platform has modified it. Prevents UAT/UBT calls showing up with a "mono: " prefix on Mac, rather than "AutomationTool:" or "UnrealBuildTool:".

Change 2955885 on 2016/04/26 by [REDACTED_PERSON]

	BuildGraph: Allow passing -Clean on the command line to propagate to UE4Build, impacting how targets are compiled as well as clearing the cached BuildGraph state. Add a second parameter, -ClearHistory, to just wipe the history of completed nodes.

Change 2955919 on 2016/04/26 by [REDACTED_PERSON]

	Fixed timestamp resolution to check for two seconds instead of two ticks.  This was causing mac builders to throw false positives on file changes

Change 2956118 on 2016/04/26 by [REDACTED_PERSON]

	BuildGraph: Add support for conditional blocks in BuildGraph scripts, either with a simple condition, or picking from a list of options. Two new elements can be added anywhere in scripts:

	<Do If="...">
	    <!-- Inner elements -->
	</Do>

	<Choose>
	    <Option If="...">
	        <!-- Inner elements -->
	    </Option>
	    <Option If="...">
	        <!-- Inner elements -->
	    </Option>
	    <Otherwise>
	        <!-- Inner elements -->
	    </Otherwise>
	</Choose>

Change 2956792 on 2016/04/26 by [REDACTED_PERSON]

	EC: Prevent scheduled builds being queued up, and starting at times other than the times they're scheduled for. Prevents builds which have just been added to the stream settings from starting immediately, and prevents full builds starting during the day (as soon as the first change is made).

Change 2957131 on 2016/04/26 by [REDACTED_PERSON]

	EC: Increase the precedence of the stack trace matcher.

Change 2957419 on 2016/04/27 by [REDACTED_PERSON]

	EC: Skip the "end: stack for UAT" line in postp.

Change 2957588 on 2016/04/27 by [REDACTED_PERSON]

	Core: Change formatting for callstacks for crashes and ensures so that EC can parse them from logs more easily.

Change 2958047 on 2016/04/27 by [REDACTED_PERSON]

	BuildGraph: Feature to generate reports as part of build graph scripts. Reports operate similarly to triggers, but just provide a summary of completed jobsteps without offering to run a downstream job. Syntax is similar to declaring aggregates: <Report Name="Summary" Requires="Node1;Node2"/>

Change 2958188 on 2016/04/27 by [REDACTED_PERSON]

	BuildGraph: Automatically generate a report when a preflight completes.

Change 2959053 on 2016/04/28 by [REDACTED_PERSON]

	BuildGraph: Move the CleanTempStorage commandlet into BuildGraph, and add support for cleaning out new-style temp storage directories (which do not contain TempManifest files).

Change 2959429 on 2016/04/28 by [REDACTED_PERSON]

	UAT: Add a script to describe a stream being copied up to its parent. To use, just run the UAT command "StreamCopyDescription -Stream=//UE4/Dev-Build". Optionally specify -Changes=//UE4/OtherStream/Engine/...

diff --git a/Engine/Build/BuildFarm/PostpExtensions.pl b/Engine/Build/BuildFarm/PostpExtensions.pl
--- a/Engine/Build/BuildFarm/PostpExtensions.pl
+++ b/Engine/Build/BuildFarm/PostpExtensions.pl
@@ -65,6 +65,11 @@ $::gDontCheck .= ",javacNote";
 # NOTE: also be aware that the look ahead lines are considered processed
 
 unshift @::gMatchers, (
+	{
+		id => "UATErrStack",
+		pattern => q{begin: stack for UAT},
+		action => q{incValue("errors"); my $line_count = forwardTo(q{end: stack for UAT}); my $last_line = $::gCurrentLine + $line_count; diagnostic("AutomationTool", "error", 1, $line_count - 1); $::gCurrentLine = $last_line; }
+	},
     {
         id =>               "clErrorMultiline",
         pattern =>          q{([^(]+)(\([\d,]+\))? ?: (fatal )?error [a-zA-Z]+[\d]+},
@@ -75,6 +80,16 @@ unshift @::gMatchers, (
         pattern =>          q{([^(]+)\([\d,]+\) ?: warning },
         action =>           q{incValue("warnings"); my ($file_only) = ($1 =~ /([^\\\\]+)$/); diagnostic($file_only || $1, "warning", 0, forwardWhile("^    "))},
     },
+    {
+        id =>               "clangError",
+        pattern =>          q{([^:]+):[\d:]+ error:},
+        action =>           q{incValue("errors"); diagnostic($1, "error", backWhile(": In (member )?function"), forwardWhile("^   "))},
+    },
+    {
+        id =>               "clangWarning",
+        pattern =>          q{([^:]+):[\d:]+ warning:},
+        action =>           q{incValue("warnings"); diagnostic($1, "warning", backWhile(": In function"), 0)},
+    },
     {
         id =>               "ubtFailedToProduceItem",
         pattern =>          q{(ERROR: )?UBT ERROR: Failed to produce item: },
@@ -119,7 +134,17 @@ unshift @::gMatchers, (
 		id =>				"ubtWarning",
 		pattern =>			q{^WARNING:},
 		action =>			q{incValue("warnings"); diagnostic("AutomationTool", "warning")}
-	}
+	},
+    {
+        id =>               "genericError",
+        pattern =>          q{^(.* )?(ERROR|[Ee]rror)}.q{( (\([^)]+\)|\[[^\]]+\]))?: },
+        action =>           q{incValue("errors"); diagnostic("", "error", 0, forwardWhile("^   "))}
+    },
+    {
+        id =>               "genericWarning",
+        pattern =>          q{WARNING:|[Ww]arning:},
+        action =>           q{incValue("warnings"); diagnostic("", "warning", 0)},
+    },
 );
 
 push @::gMatchers, (
@@ -188,11 +213,6 @@ push @::gMatchers, (
 		pattern => q{AutomationTool: Stack:},
 		action => q{incValue("errors"); diagnostic("AutomationTool", "error", 0, forwardTo(q{AutomationTool: \n}));}
 	},
-	{
-		id => "UATErrStack",
-		pattern => q{begin: stack for UAT},
-		action => q{incValue("errors"); diagnostic("AutomationTool", "error", 0, forwardTo(q{end: stack for UAT}));}
-	},
 	{
 		id => "BuildFailed",
 		pattern => q{BuildCommand.Execute: ERROR:},
diff --git a/Engine/Build/BuildFarm/PostpFilter.pl b/Engine/Build/BuildFarm/PostpFilter.pl
--- a/Engine/Build/BuildFarm/PostpFilter.pl
+++ b/Engine/Build/BuildFarm/PostpFilter.pl
@@ -20,7 +20,7 @@ while(<>)
 	s/^AutomationTool: ([A-Za-z0-9-]+: )/$1/;
 	
 	# remove the gubp function or prefix
-	s/^([A-Za-z0-9.-]+): // if !/^AutomationTool\\.AutomationException:/;
+	s/^([A-Za-z0-9.-]+): // if !/^AutomationTool\\.AutomationException:/ && !/^error:/ && !/^warning:/;
 
 	# if it was the editor, also strip any additional timestamp
 	if($1 && $1 =~ /^UE4Editor/)
diff --git a/Engine/Build/EngineBuild.xml b/Engine/Build/EngineBuild.xml
new file mode 100644
--- /dev/null
+++ b/Engine/Build/EngineBuild.xml
@@ -0,0 +1,58 @@
+<?xml version='1.0' ?>
+<BuildGraph xmlns="http://www.epicgames.com/BuildGraph" xmlns:xsi="http://www.w3.org/2001/XMLSchema-instance" xsi:schemaLocation="http://www.epicgames.com/BuildGraph ../../Engine/Build/Graph/Schema.xsd" >
+
+	<!-- Base editor and tools for compile -->
+	<Agent Name="Editor Win64" Type="CompileWin64">
+		<Node Name="Update Version Files">
+			<SetVersion Change="$(Change)" Branch="$(EscapedBranch)" If="$(IsBuildMachine)"/>
+		</Node>
+		<Node Name="Compile UnrealHeaderTool Win64" Requires="Update Version Files">
+			<Compile Target="UnrealHeaderTool" Platform="Win64" Configuration="Development"/>
+		</Node>
+		<Node Name="Compile UE4Editor Win64" Requires="Compile UnrealHeaderTool Win64">
+			<Compile Target="UE4Editor" Platform="Win64" Configuration="Development"/>
+		</Node>
+	</Agent>
+
+	<!-- Base editor and tools for compile on Mac -->
+	<Agent Name="Editor Mac" Type="CompileMac">
+		<Node Name="Compile UnrealHeaderTool Mac" Requires="Update Version Files">
+			<Compile Target="UnrealHeaderTool" Platform="Mac" Configuration="Development"/>
+		</Node>
+		<Node Name="Compile UE4Editor Mac" Requires="Compile UnrealHeaderTool Mac">
+			<Compile Target="UE4Editor" Platform="Mac" Configuration="Development"/>
+		</Node>
+	</Agent>
+	
+	<!-- Downstream monolithics and tools -->
+	<Agent Name="Monolithics Win64" Type="CompileWin64">
+		<Node Name="Compile Tools Win64" Requires="Compile UnrealHeaderTool Win64;Update Version Files">
+			<Compile Target="CrashReportClient" Platform="Win64" Configuration="Shipping"/>
+			<Compile Target="ShaderCompileWorker" Platform="Win64" Configuration="Development"/>
+			<Compile Target="UnrealLightmass" Platform="Win64" Configuration="Development"/>
+			<Compile Target="UnrealPak" Platform="Win64" Configuration="Development"/>
+		</Node>
+		<Node Name="Compile UE4Game Win64" Requires="Compile UnrealHeaderTool Win64">
+			<Compile Target="UE4Game" Platform="Win64" Configuration="Development"/>
+			<Compile Target="UE4Game" Platform="Win64" Configuration="Shipping"/>
+		</Node>
+	</Agent>
+
+	<!-- Downstream monolithics and tools -->
+	<Agent Name="Monolithics Mac" Type="CompileMac">
+		<Node Name="Compile Tools Mac" Requires="Compile UnrealHeaderTool Mac;Update Version Files">
+			<Compile Target="CrashReportClient" Platform="Mac" Configuration="Shipping"/>
+			<Compile Target="ShaderCompileWorker" Platform="Mac" Configuration="Development"/>
+			<Compile Target="UnrealLightmass" Platform="Mac" Configuration="Development"/>
+			<Compile Target="UnrealPak" Platform="Mac" Configuration="Development"/>
+		</Node>
+		<Node Name="Compile UE4Game Mac" Requires="Compile UnrealHeaderTool Mac">
+			<Compile Target="UE4Game" Platform="Mac" Configuration="Development"/>
+			<Compile Target="UE4Game" Platform="Mac" Configuration="Shipping"/>
+		</Node>
+	</Agent>
+	
+	<!-- Email settings -->
+	<Local Name="AllNodes" Value="Compile UE4Editor Win64;Compile UE4Editor Mac;Compile UE4Game Win64;Compile UE4Game Mac"/>
+	<Notify Targets="$(AllNodes)" Submitters=".../Build/...;.../Source/..." />
+</BuildGraph>
diff --git a/Engine/Source/Programs/AutomationTool/AutomationUtils/ProcessUtils.cs b/Engine/Source/Programs/AutomationTool/AutomationUtils/ProcessUtils.cs
--- a/Engine/Source/Programs/AutomationTool/AutomationUtils/ProcessUtils.cs
+++ b/Engine/Source/Programs/AutomationTool/AutomationUtils/ProcessUtils.cs
@@ -684,6 +684,10 @@ namespace AutomationTool
 		public static ProcessResult Run(string App, string CommandLine = null, string Input = null, ERunOptions Options = ERunOptions.Default, Dictionary<string, string> Env = null)
 		{
 			App = ConvertSeparators(PathSeparator.Default, App);
+
+			// Get the log name before allowing the platform to modify the app/command-line. We want mono apps to be written to a log named after the application and appear with the application prefix rather than "mono:".
+			string LogName = Path.GetFileNameWithoutExtension(App);
+
 			HostPlatform.Current.SetupOptionsForRun(ref App, ref Options, ref CommandLine);
 			if (App == "ectool" || App == "zip" || App == "xcodebuild")
 			{
@@ -701,7 +705,7 @@ namespace AutomationTool
             {
                 LogWithVerbosity(SpewVerbosity,"Run: " + App + " " + (String.IsNullOrEmpty(CommandLine) ? "" : CommandLine));
             }
-			ProcessResult Result = ProcessManager.CreateProcess(App, Options.HasFlag(ERunOptions.AllowSpew), Path.GetFileNameWithoutExtension(App), Env, SpewVerbosity:SpewVerbosity);
+			ProcessResult Result = ProcessManager.CreateProcess(App, Options.HasFlag(ERunOptions.AllowSpew), LogName, Env, SpewVerbosity:SpewVerbosity);
 			Process Proc = Result.ProcessObject;
 
 			bool bRedirectStdOut = (Options & ERunOptions.NoStdOutRedirect) != ERunOptions.NoStdOutRedirect;			
diff --git a/Engine/Source/Programs/AutomationTool/BuildGraph/BuildGraph.Automation.csproj b/Engine/Source/Programs/AutomationTool/BuildGraph/BuildGraph.Automation.csproj
--- a/Engine/Source/Programs/AutomationTool/BuildGraph/BuildGraph.Automation.csproj
+++ b/Engine/Source/Programs/AutomationTool/BuildGraph/BuildGraph.Automation.csproj
@@ -65,6 +65,7 @@
     <Compile Include="BuildGraph.cs" />
     <Compile Include="Graph.cs" />
     <Compile Include="Node.cs" />
+    <Compile Include="Report.cs" />
     <Compile Include="Task.cs" />
     <Compile Include="JobContext.cs" />
     <Compile Include="ManualTrigger.cs" />
diff --git a/Engine/Source/Programs/AutomationTool/BuildGraph/BuildGraph.cs b/Engine/Source/Programs/AutomationTool/BuildGraph/BuildGraph.cs
--- a/Engine/Source/Programs/AutomationTool/BuildGraph/BuildGraph.cs
+++ b/Engine/Source/Programs/AutomationTool/BuildGraph/BuildGraph.cs
@@ -92,10 +92,11 @@ namespace AutomationTool
 			string SingleNodeName = ParseParamValue("SingleNode", null);
 			string[] TriggerNames = ParseParamValue("Trigger", "").Split(new char[]{ '+' }, StringSplitOptions.RemoveEmptyEntries).ToArray();
 			bool bSkipTriggers = ParseParam("SkipTriggers");
-			bool bClean = ParseParam("Clean");
+			bool bClearHistory = ParseParam("Clean") || ParseParam("ClearHistory");
 			bool bListOnly = ParseParam("ListOnly");
 			bool bWriteToSharedStorage = ParseParam("WriteToSharedStorage") || CommandUtils.IsBuildMachine;
 			bool bPublicTasksOnly = ParseParam("PublicTasksOnly");
+			string ReportName = ParseParamValue("ReportName", null); 
 
 			GraphPrintOptions PrintOptions = 0;
 			if(ParseParam("ShowDeps"))
@@ -174,7 +175,7 @@ namespace AutomationTool
 			// Create the temp storage handler
 			DirectoryReference RootDir = new DirectoryReference(CommandUtils.CmdEnv.LocalRoot);
 			TempStorage Storage = new TempStorage(RootDir, DirectoryReference.Combine(RootDir, "Engine", "Saved", "BuildGraph"), (SharedStorageDir == null)? null : new DirectoryReference(SharedStorageDir), bWriteToSharedStorage);
-			if(bClean)
+			if(bClearHistory)
 			{
 				Storage.CleanLocal();
 			}
@@ -199,6 +200,14 @@ namespace AutomationTool
 			// Cull the graph to include only those nodes
 			Graph.Select(TargetNodes);
 
+			// If a report for the whole build was requested, insert it into the graph
+			if (ReportName != null)
+			{
+				Report NewReport = new Report(ReportName);
+				NewReport.Nodes.UnionWith(Graph.Agents.SelectMany(x => x.Nodes));
+				Graph.NameToReport.Add(ReportName, NewReport);
+			}
+
 			// Write out the preprocessed script
 			if (PreprocessedFileName != null)
 			{
diff --git a/Engine/Source/Programs/AutomationTool/BuildGraph/Graph.cs b/Engine/Source/Programs/AutomationTool/BuildGraph/Graph.cs
--- a/Engine/Source/Programs/AutomationTool/BuildGraph/Graph.cs
+++ b/Engine/Source/Programs/AutomationTool/BuildGraph/Graph.cs
@@ -86,6 +86,11 @@ namespace AutomationTool
 		/// </summary>
 		public Dictionary<string, Node> NameToNode = new Dictionary<string,Node>(StringComparer.InvariantCultureIgnoreCase);
 
+		/// <summary>
+		/// Mapping of names to the corresponding report.
+		/// </summary>
+		public Dictionary<string, Report> NameToReport = new Dictionary<string, Report>(StringComparer.InvariantCultureIgnoreCase);
+
 		/// <summary>
 		/// Mapping of names to their corresponding node output.
 		/// </summary>
@@ -119,7 +124,7 @@ namespace AutomationTool
 			string TrimName = Name.StartsWith("#")? Name.Substring(1) : Name;
 
 			// Check it doesn't match anything in the graph
-			return NameToNode.ContainsKey(TrimName) || NameToNodeOutput.ContainsKey(TrimName) || AggregateNameToNodes.ContainsKey(TrimName);
+			return NameToNode.ContainsKey(TrimName) || NameToNodeOutput.ContainsKey(TrimName) || NameToReport.ContainsKey(TrimName) || AggregateNameToNodes.ContainsKey(TrimName);
 		}
 
 		/// <summary>
@@ -230,6 +235,15 @@ namespace AutomationTool
 			// Remove all the empty agents
 			Agents.RemoveAll(x => x.Nodes.Count == 0);
 
+			// Trim down the list of nodes for each report to the ones that are being built
+			foreach (Report Report in NameToReport.Values)
+			{
+				Report.Nodes.RemoveWhere(x => !RetainNodes.Contains(x));
+			}
+
+			// Remove all the empty reports
+			NameToReport = NameToReport.Where(x => x.Value.Nodes.Count > 0).ToDictionary(Pair => Pair.Key, Pair => Pair.Value);
+
 			// Remove all the order dependencies which are no longer part of the graph. Since we don't need to build them, we don't need to wait for them
 			foreach(Node Node in RetainNodes)
 			{
@@ -288,6 +302,15 @@ namespace AutomationTool
 					Writer.WriteAttributeString("Requires", String.Join(";", Aggregate.Value.Select(x => x.Name)));
 					Writer.WriteEndElement();
 				}
+
+				foreach (Report Report in NameToReport.Values)
+				{
+					Writer.WriteStartElement("Report");
+					Writer.WriteAttributeString("Name", Report.Name);
+					Writer.WriteAttributeString("Requires", String.Join(";", Report.Nodes.Select(x => x.Name)));
+					Writer.WriteEndElement();
+				}
+
 				Writer.WriteEndElement();
 			}
 		}
@@ -352,8 +375,29 @@ namespace AutomationTool
 				}
 				JsonWriter.WriteArrayEnd();
 
-				// Write all the triggers
-				JsonWriter.WriteArrayStart("Triggers");
+				// Write all the triggers and reports. 
+				JsonWriter.WriteArrayStart("Reports");
+				foreach (Report Report in NameToReport.Values)
+				{
+					Node[] Dependencies = Report.Nodes.Where(x => NodesToExecute.Contains(x)).ToArray();
+					if (Dependencies.Length > 0)
+					{
+						// Reduce that list to the smallest subset of direct dependencies
+						HashSet<Node> DirectDependencies = new HashSet<Node>(Dependencies);
+						foreach (Node Dependency in Dependencies)
+						{
+							DirectDependencies.ExceptWith(Dependency.OrderDependencies);
+						}
+
+						JsonWriter.WriteObjectStart();
+						JsonWriter.WriteValue("Name", Report.Name);
+						JsonWriter.WriteValue("AllDependencies", String.Join(";", Agents.SelectMany(x => x.Nodes).Where(x => Dependencies.Contains(x)).Select(x => x.Name)));
+						JsonWriter.WriteValue("DirectDependencies", String.Join(";", DirectDependencies.Select(x => x.Name)));
+						JsonWriter.WriteValue("Notify", String.Join(";", Report.NotifyUsers));
+						JsonWriter.WriteValue("IsTrigger", false);
+						JsonWriter.WriteObjectEnd();
+					}
+				}
 				foreach (ManualTrigger Trigger in NameToTrigger.Values)
 				{
 					if(!ActivatedTriggers.Contains(Trigger) && NodesToExecute.Any(x => x.ControllingTrigger == Trigger.Parent))
@@ -385,10 +429,12 @@ namespace AutomationTool
 						JsonWriter.WriteValue("AllDependencies", String.Join(";", Agents.SelectMany(x => x.Nodes).Where(x => Dependencies.Contains(x)).Select(x => x.Name)));
 						JsonWriter.WriteValue("DirectDependencies", String.Join(";", Dependencies.Where(x => DirectDependencies.Contains(x)).Select(x => x.Name)));
 						JsonWriter.WriteValue("Notify", String.Join(";", Trigger.NotifyUsers));
+						JsonWriter.WriteValue("IsTrigger", true);
 						JsonWriter.WriteObjectEnd();
 					}
 				}
 				JsonWriter.WriteArrayEnd();
+
 				JsonWriter.WriteObjectEnd();
 			}
 		}
diff --git a/Engine/Source/Programs/AutomationTool/BuildGraph/Report.cs b/Engine/Source/Programs/AutomationTool/BuildGraph/Report.cs
new file mode 100644
--- /dev/null
+++ b/Engine/Source/Programs/AutomationTool/BuildGraph/Report.cs
@@ -0,0 +1,48 @@
+﻿using System;
+using System.Collections.Generic;
+using System.Diagnostics;
+using System.Linq;
+using System.Text;
+using System.Threading.Tasks;
+
+namespace AutomationTool
+{
+	/// <summary>
+	/// Defines a report to be generated as part of the build.
+	/// </summary>
+	class Report
+	{
+		/// <summary>
+		/// Name of this trigger
+		/// </summary>
+		public readonly string Name;
+
+		/// <summary>
+		/// Set of nodes to include in the report
+		/// </summary>
+		public HashSet<Node> Nodes = new HashSet<Node>();
+
+		/// <summary>
+		/// List of users to notify with this report
+		/// </summary>
+		public HashSet<string> NotifyUsers = new HashSet<string>(StringComparer.InvariantCultureIgnoreCase);
+
+		/// <summary>
+		/// Constructor
+		/// </summary>
+		/// <param name="InName">Name of this report</param>
+		public Report(string InName)
+		{
+			Name = InName;
+		}
+
+		/// <summary>
+		/// Get the name of this report
+		/// </summary>
+		/// <returns>The name of this report</returns>
+		public override string ToString()
+		{
+			return Name;
+		}
+	}
+}
diff --git a/Engine/Source/Programs/AutomationTool/BuildGraph/Schema.cs b/Engine/Source/Programs/AutomationTool/BuildGraph/Schema.cs
--- a/Engine/Source/Programs/AutomationTool/BuildGraph/Schema.cs
+++ b/Engine/Source/Programs/AutomationTool/BuildGraph/Schema.cs
@@ -109,9 +109,13 @@ namespace AutomationTool
 	{
 		Graph,
 		Trigger,
+		TriggerBody,
 		Agent,
+		AgentBody,
 		Node,
+		NodeBody,
 		Aggregate,
+		Report,
 		Notify,
 		Include,
 		Property,
@@ -267,9 +271,13 @@ namespace AutomationTool
 			NewSchema.Items.Add(CreateSchemaElement(RootElementName, ScriptSchemaStandardType.Graph));
 			NewSchema.Items.Add(CreateGraphType());
 			NewSchema.Items.Add(CreateTriggerType());
+			NewSchema.Items.Add(CreateTriggerBodyType());
 			NewSchema.Items.Add(CreateAgentType());
-			NewSchema.Items.Add(CreateNodeType(TaskNameToType));
+			NewSchema.Items.Add(CreateAgentBodyType());
+			NewSchema.Items.Add(CreateNodeType());
+			NewSchema.Items.Add(CreateNodeBodyType(TaskNameToType));
 			NewSchema.Items.Add(CreateAggregateType());
+			NewSchema.Items.Add(CreateReportType());
 			NewSchema.Items.Add(CreateNotifyType());
 			NewSchema.Items.Add(CreateIncludeType());
 			NewSchema.Items.Add(CreatePropertyType(ScriptSchemaStandardType.Property));
@@ -392,20 +400,44 @@ namespace AutomationTool
 			GraphChoice.Items.Add(CreateSchemaElement("Agent", ScriptSchemaStandardType.Agent));
 			GraphChoice.Items.Add(CreateSchemaElement("Trigger", ScriptSchemaStandardType.Trigger));
 			GraphChoice.Items.Add(CreateSchemaElement("Aggregate", ScriptSchemaStandardType.Aggregate));
+			GraphChoice.Items.Add(CreateSchemaElement("Report", ScriptSchemaStandardType.Report));
 			GraphChoice.Items.Add(CreateSchemaElement("Notify", ScriptSchemaStandardType.Notify));
 			GraphChoice.Items.Add(CreateSchemaElement("Warning", ScriptSchemaStandardType.Warning));
 			GraphChoice.Items.Add(CreateSchemaElement("Error", ScriptSchemaStandardType.Error));
+			GraphChoice.Items.Add(CreateDoElement(ScriptSchemaStandardType.Graph));
+			GraphChoice.Items.Add(CreateChooseElement(ScriptSchemaStandardType.Graph));
 
 			XmlSchemaComplexType GraphType = new XmlSchemaComplexType();
 			GraphType.Name = GetTypeName(ScriptSchemaStandardType.Graph);
 			GraphType.Particle = GraphChoice;
 			return GraphType;
 		}
+
 		/// <summary>
 		/// Creates the schema type representing the trigger type
 		/// </summary>
 		/// <returns>Type definition for a trigger</returns>
 		static XmlSchemaType CreateTriggerType()
+		{
+			XmlSchemaComplexContentExtension Extension = new XmlSchemaComplexContentExtension();
+			Extension.BaseTypeName = GetQualifiedTypeName(ScriptSchemaStandardType.TriggerBody);
+			Extension.Attributes.Add(CreateSchemaAttribute("Name", ScriptSchemaStandardType.QualifiedName, XmlSchemaUse.Required));
+			Extension.Attributes.Add(CreateSchemaAttribute("If", ScriptSchemaStandardType.BalancedString, XmlSchemaUse.Optional));
+
+			XmlSchemaComplexContent ContentModel = new XmlSchemaComplexContent();
+			ContentModel.Content = Extension;
+
+			XmlSchemaComplexType ComplexType = new XmlSchemaComplexType();
+			ComplexType.Name = GetTypeName(ScriptSchemaStandardType.Trigger);
+			ComplexType.ContentModel = ContentModel;
+			return ComplexType;
+		}
+
+		/// <summary>
+		/// Creates the schema type representing the contents of a trigger type
+		/// </summary>
+		/// <returns>Type definition for an agent</returns>
+		static XmlSchemaType CreateTriggerBodyType()
 		{
 			XmlSchemaChoice TriggerChoice = new XmlSchemaChoice();
 			TriggerChoice.MinOccurs = 0;
@@ -416,12 +448,12 @@ namespace AutomationTool
 			TriggerChoice.Items.Add(CreateSchemaElement("Aggregate", ScriptSchemaStandardType.Aggregate));
 			TriggerChoice.Items.Add(CreateSchemaElement("Warning", ScriptSchemaStandardType.Warning));
 			TriggerChoice.Items.Add(CreateSchemaElement("Error", ScriptSchemaStandardType.Error));
+			TriggerChoice.Items.Add(CreateDoElement(ScriptSchemaStandardType.TriggerBody));
+			TriggerChoice.Items.Add(CreateChooseElement(ScriptSchemaStandardType.TriggerBody));
 
 			XmlSchemaComplexType TriggerType = new XmlSchemaComplexType();
-			TriggerType.Name = GetTypeName(ScriptSchemaStandardType.Trigger);
+			TriggerType.Name = GetTypeName(ScriptSchemaStandardType.TriggerBody);
 			TriggerType.Particle = TriggerChoice;
-			TriggerType.Attributes.Add(CreateSchemaAttribute("Name", ScriptSchemaStandardType.QualifiedName, XmlSchemaUse.Required));
-			TriggerType.Attributes.Add(CreateSchemaAttribute("If", ScriptSchemaStandardType.BalancedString, XmlSchemaUse.Optional));
 			return TriggerType;
 		}
 
@@ -430,6 +462,27 @@ namespace AutomationTool
 		/// </summary>
 		/// <returns>Type definition for an agent</returns>
 		static XmlSchemaType CreateAgentType()
+		{
+			XmlSchemaComplexContentExtension Extension = new XmlSchemaComplexContentExtension();
+			Extension.BaseTypeName = GetQualifiedTypeName(ScriptSchemaStandardType.AgentBody);
+			Extension.Attributes.Add(CreateSchemaAttribute("Name", StringTypeName, XmlSchemaUse.Required));
+			Extension.Attributes.Add(CreateSchemaAttribute("Type", ScriptSchemaStandardType.NameList, XmlSchemaUse.Optional));
+			Extension.Attributes.Add(CreateSchemaAttribute("If", ScriptSchemaStandardType.BalancedString, XmlSchemaUse.Optional));
+
+			XmlSchemaComplexContent ContentModel = new XmlSchemaComplexContent();
+			ContentModel.Content = Extension;
+
+			XmlSchemaComplexType ComplexType = new XmlSchemaComplexType();
+			ComplexType.Name = GetTypeName(ScriptSchemaStandardType.Agent);
+			ComplexType.ContentModel = ContentModel;
+			return ComplexType;
+		}
+
+		/// <summary>
+		/// Creates the schema type representing the contents of agent type
+		/// </summary>
+		/// <returns>Type definition for an agent</returns>
+		static XmlSchemaType CreateAgentBodyType()
 		{
 			XmlSchemaChoice AgentChoice = new XmlSchemaChoice();
 			AgentChoice.MinOccurs = 0;
@@ -439,13 +492,12 @@ namespace AutomationTool
 			AgentChoice.Items.Add(CreateSchemaElement("Node", ScriptSchemaStandardType.Node));
 			AgentChoice.Items.Add(CreateSchemaElement("Warning", ScriptSchemaStandardType.Warning));
 			AgentChoice.Items.Add(CreateSchemaElement("Error", ScriptSchemaStandardType.Error));
+			AgentChoice.Items.Add(CreateDoElement(ScriptSchemaStandardType.AgentBody));
+			AgentChoice.Items.Add(CreateChooseElement(ScriptSchemaStandardType.AgentBody));
 
 			XmlSchemaComplexType AgentType = new XmlSchemaComplexType();
-			AgentType.Name = GetTypeName(ScriptSchemaStandardType.Agent);
+			AgentType.Name = GetTypeName(ScriptSchemaStandardType.AgentBody);
 			AgentType.Particle = AgentChoice;
-			AgentType.Attributes.Add(CreateSchemaAttribute("Name", StringTypeName, XmlSchemaUse.Required));
-			AgentType.Attributes.Add(CreateSchemaAttribute("Type", ScriptSchemaStandardType.NameList, XmlSchemaUse.Optional));
-			AgentType.Attributes.Add(CreateSchemaAttribute("If", ScriptSchemaStandardType.BalancedString, XmlSchemaUse.Optional));
 			return AgentType;
 		}
 
@@ -453,7 +505,31 @@ namespace AutomationTool
 		/// Creates the schema type representing the node type
 		/// </summary>
 		/// <returns>Type definition for a node</returns>
-		static XmlSchemaType CreateNodeType(Dictionary<string, XmlSchemaComplexType> TaskNameToType)
+		static XmlSchemaType CreateNodeType()
+		{
+			XmlSchemaComplexContentExtension Extension = new XmlSchemaComplexContentExtension();
+			Extension.BaseTypeName = GetQualifiedTypeName(ScriptSchemaStandardType.NodeBody);
+			Extension.Attributes.Add(CreateSchemaAttribute("Name", ScriptSchemaStandardType.Name, XmlSchemaUse.Required));
+			Extension.Attributes.Add(CreateSchemaAttribute("Requires", ScriptSchemaStandardType.NameOrTagList, XmlSchemaUse.Optional));
+			Extension.Attributes.Add(CreateSchemaAttribute("Produces", ScriptSchemaStandardType.TagList, XmlSchemaUse.Optional));
+			Extension.Attributes.Add(CreateSchemaAttribute("After", ScriptSchemaStandardType.NameOrTagList, XmlSchemaUse.Optional));
+			Extension.Attributes.Add(CreateSchemaAttribute("If", ScriptSchemaStandardType.BalancedString, XmlSchemaUse.Optional));
+			Extension.Attributes.Add(CreateSchemaAttribute("NotifyOnWarnings", ScriptSchemaStandardType.Boolean, XmlSchemaUse.Optional));
+
+			XmlSchemaComplexContent ContentModel = new XmlSchemaComplexContent();
+			ContentModel.Content = Extension;
+
+			XmlSchemaComplexType ComplexType = new XmlSchemaComplexType();
+			ComplexType.Name = GetTypeName(ScriptSchemaStandardType.Node);
+			ComplexType.ContentModel = ContentModel;
+			return ComplexType;
+		}
+
+		/// <summary>
+		/// Creates the schema type representing the body of the node type
+		/// </summary>
+		/// <returns>Type definition for a node</returns>
+		static XmlSchemaType CreateNodeBodyType(Dictionary<string, XmlSchemaComplexType> TaskNameToType)
 		{
 			XmlSchemaChoice NodeChoice = new XmlSchemaChoice();
 			NodeChoice.MinOccurs = 0;
@@ -462,20 +538,16 @@ namespace AutomationTool
 			NodeChoice.Items.Add(CreateSchemaElement("Local", ScriptSchemaStandardType.Local));
 			NodeChoice.Items.Add(CreateSchemaElement("Warning", ScriptSchemaStandardType.Warning));
 			NodeChoice.Items.Add(CreateSchemaElement("Error", ScriptSchemaStandardType.Error));
-			foreach(KeyValuePair<string, XmlSchemaComplexType> Pair in TaskNameToType.OrderBy(x => x.Key))
+			NodeChoice.Items.Add(CreateDoElement(ScriptSchemaStandardType.NodeBody));
+			NodeChoice.Items.Add(CreateChooseElement(ScriptSchemaStandardType.NodeBody));
+			foreach (KeyValuePair<string, XmlSchemaComplexType> Pair in TaskNameToType.OrderBy(x => x.Key))
 			{
 				NodeChoice.Items.Add(CreateSchemaElement(Pair.Key, new XmlQualifiedName(Pair.Value.Name, NamespaceURI)));
 			}
 
 			XmlSchemaComplexType NodeType = new XmlSchemaComplexType();
-			NodeType.Name = GetTypeName(ScriptSchemaStandardType.Node);
+			NodeType.Name = GetTypeName(ScriptSchemaStandardType.NodeBody);
 			NodeType.Particle = NodeChoice;
-			NodeType.Attributes.Add(CreateSchemaAttribute("Name", ScriptSchemaStandardType.Name, XmlSchemaUse.Required));
-			NodeType.Attributes.Add(CreateSchemaAttribute("Requires", ScriptSchemaStandardType.NameOrTagList, XmlSchemaUse.Optional));
-			NodeType.Attributes.Add(CreateSchemaAttribute("Produces", ScriptSchemaStandardType.TagList, XmlSchemaUse.Optional));
-			NodeType.Attributes.Add(CreateSchemaAttribute("After", ScriptSchemaStandardType.NameOrTagList, XmlSchemaUse.Optional));
-			NodeType.Attributes.Add(CreateSchemaAttribute("If", ScriptSchemaStandardType.BalancedString, XmlSchemaUse.Optional));
-			NodeType.Attributes.Add(CreateSchemaAttribute("NotifyOnWarnings", ScriptSchemaStandardType.Boolean, XmlSchemaUse.Optional));
 			return NodeType;
 		}
 
@@ -493,6 +565,20 @@ namespace AutomationTool
 			return AggregateType;
 		}
 
+		/// <summary>
+		/// Creates the schema type representing the report type
+		/// </summary>
+		/// <returns>Type definition for a report</returns>
+		static XmlSchemaType CreateReportType()
+		{
+			XmlSchemaComplexType AggregateType = new XmlSchemaComplexType();
+			AggregateType.Name = GetTypeName(ScriptSchemaStandardType.Report);
+			AggregateType.Attributes.Add(CreateSchemaAttribute("Name", ScriptSchemaStandardType.Name, XmlSchemaUse.Required));
+			AggregateType.Attributes.Add(CreateSchemaAttribute("Requires", ScriptSchemaStandardType.NameOrTagList, XmlSchemaUse.Required));
+			AggregateType.Attributes.Add(CreateSchemaAttribute("If", ScriptSchemaStandardType.BalancedString, XmlSchemaUse.Optional));
+			return AggregateType;
+		}
+
 		/// <summary>
 		/// Creates the schema type representing a notifier
 		/// </summary>
@@ -505,6 +591,7 @@ namespace AutomationTool
 			AggregateType.Attributes.Add(CreateSchemaAttribute("Except", ScriptSchemaStandardType.NameOrTagList, XmlSchemaUse.Optional));
 			AggregateType.Attributes.Add(CreateSchemaAttribute("Nodes", ScriptSchemaStandardType.NameOrTagList, XmlSchemaUse.Optional));
 			AggregateType.Attributes.Add(CreateSchemaAttribute("Triggers", ScriptSchemaStandardType.NameList, XmlSchemaUse.Optional));
+			AggregateType.Attributes.Add(CreateSchemaAttribute("Reports", ScriptSchemaStandardType.NameList, XmlSchemaUse.Optional));
 			AggregateType.Attributes.Add(CreateSchemaAttribute("Users", ScriptSchemaStandardType.BalancedString, XmlSchemaUse.Optional));
 			AggregateType.Attributes.Add(CreateSchemaAttribute("Submitters", ScriptSchemaStandardType.BalancedString, XmlSchemaUse.Optional));
 			AggregateType.Attributes.Add(CreateSchemaAttribute("Warnings", ScriptSchemaStandardType.Boolean, XmlSchemaUse.Optional));
@@ -552,6 +639,74 @@ namespace AutomationTool
 			return PropertyType;
 		}
 
+		/// <summary>
+		/// Creates an element representing a conditional "Do" block, which recursively contains another type
+		/// </summary>
+		/// <param name="InnerType">The base type for the do block to contain</param>
+		/// <returns>New schema element for the block</returns>
+		static XmlSchemaElement CreateDoElement(ScriptSchemaStandardType InnerType)
+		{
+			XmlSchemaComplexContentExtension Extension = new XmlSchemaComplexContentExtension();
+			Extension.BaseTypeName = GetQualifiedTypeName(InnerType);
+			Extension.Attributes.Add(CreateSchemaAttribute("If", ScriptSchemaStandardType.BalancedString, XmlSchemaUse.Optional));
+
+			XmlSchemaComplexContent ContentModel = new XmlSchemaComplexContent();
+			ContentModel.Content = Extension;
+
+			XmlSchemaComplexType SchemaType = new XmlSchemaComplexType();
+			SchemaType.ContentModel = ContentModel;
+
+			XmlSchemaElement Element = new XmlSchemaElement();
+			Element.Name = "Do";
+			Element.SchemaType = SchemaType;
+			return Element;
+		}
+
+		/// <summary>
+		/// Creates an element representing a conditional "Choose" block, which recursively contains another type
+		/// </summary>
+		/// <param name="InnerType">The base type for the do block to contain</param>
+		/// <returns>New schema element for the block</returns>
+		static XmlSchemaElement CreateChooseElement(ScriptSchemaStandardType InnerType)
+		{
+			// Create the "Option" element
+			XmlSchemaComplexContentExtension OptionExtension = new XmlSchemaComplexContentExtension();
+			OptionExtension.BaseTypeName = GetQualifiedTypeName(InnerType);
+			OptionExtension.Attributes.Add(CreateSchemaAttribute("If", ScriptSchemaStandardType.BalancedString, XmlSchemaUse.Required));
+
+			XmlSchemaComplexContent OptionContentModel = new XmlSchemaComplexContent();
+			OptionContentModel.Content = OptionExtension;
+
+			XmlSchemaComplexType OptionSchemaType = new XmlSchemaComplexType();
+			OptionSchemaType.ContentModel = OptionContentModel;
+
+			XmlSchemaElement OptionElement = new XmlSchemaElement();
+			OptionElement.Name = "Option";
+			OptionElement.SchemaType = OptionSchemaType;
+			OptionElement.MinOccurs = 0;
+			OptionElement.MaxOccursString = "unbounded";
+
+			// Create the "Otherwise" element
+			XmlSchemaElement OtherwiseElement = new XmlSchemaElement();
+			OtherwiseElement.Name = "Otherwise";
+			OtherwiseElement.SchemaTypeName = GetQualifiedTypeName(InnerType);
+			OtherwiseElement.MinOccurs = 0;
+			OtherwiseElement.MaxOccurs = 1;
+
+			// Create the "Choose" element
+			XmlSchemaSequence ChooseSequence = new XmlSchemaSequence();
+			ChooseSequence.Items.Add(OptionElement);
+			ChooseSequence.Items.Add(OtherwiseElement);
+
+			XmlSchemaComplexType ChooseSchemaType = new XmlSchemaComplexType();
+			ChooseSchemaType.Particle = ChooseSequence;
+
+			XmlSchemaElement ChooseElement = new XmlSchemaElement();
+			ChooseElement.Name = "Choose";
+			ChooseElement.SchemaType = ChooseSchemaType;
+			return ChooseElement;
+		}
+
 		/// <summary>
 		/// Constructs an XmlSchemaElement and initializes it with the given parameters
 		/// </summary>
diff --git a/Engine/Source/Programs/AutomationTool/BuildGraph/Script.cs b/Engine/Source/Programs/AutomationTool/BuildGraph/Script.cs
--- a/Engine/Source/Programs/AutomationTool/BuildGraph/Script.cs
+++ b/Engine/Source/Programs/AutomationTool/BuildGraph/Script.cs
@@ -243,45 +243,65 @@ namespace AutomationTool
 			}
 
 			// Read the root BuildGraph element
+			ReadGraphBody(Document.DocumentElement, File.Directory);
+			return true;
+		}
+
+		/// <summary>
+		/// Reads the contents of a graph
+		/// </summary>
+		/// <param name="Element">The parent element to read from</param>
+		/// <param name="BaseDirectory">Base directory to resolve includes against</param>
+		void ReadGraphBody(XmlElement Element, DirectoryReference BaseDirectory)
+		{
 			EnterScope();
-			foreach (ScriptElement Element in Document.DocumentElement.ChildNodes.OfType<ScriptElement>())
+			foreach (ScriptElement ChildElement in Element.ChildNodes.OfType<ScriptElement>())
 			{
-				switch (Element.Name)
+				switch (ChildElement.Name)
 				{
 					case "Include":
-						ReadInclude(Element, File.Directory);
+						ReadInclude(ChildElement, BaseDirectory);
 						break;
 					case "Property":
-						ReadProperty(Element);
+						ReadProperty(ChildElement);
 						break;
 					case "Local":
-						ReadLocalProperty(Element);
+						ReadLocalProperty(ChildElement);
 						break;
 					case "Agent":
-						ReadAgent(Element, null);
+						ReadAgent(ChildElement, null);
 						break;
 					case "Aggregate":
-						ReadAggregate(Element);
+						ReadAggregate(ChildElement);
+						break;
+					case "Report":
+						ReadReport(ChildElement);
 						break;
 					case "Notify":
-						ReadNotifier(Element);
+						ReadNotifier(ChildElement);
 						break;
 					case "Trigger":
-						ReadTrigger(Element);
+						ReadTrigger(ChildElement);
 						break;
 					case "Warning":
-						ReadDiagnostic(Element, LogEventType.Warning, null, null, null);
+						ReadDiagnostic(ChildElement, LogEventType.Warning, null, null, null);
 						break;
 					case "Error":
-						ReadDiagnostic(Element, LogEventType.Error, null, null, null);
+						ReadDiagnostic(ChildElement, LogEventType.Error, null, null, null);
+						break;
+					case "Do":
+						ReadBlock(ChildElement, x => ReadGraphBody(x, BaseDirectory));
+						break;
+					case "Choose":
+						ReadChoice(ChildElement, x => ReadGraphBody(x, BaseDirectory));
 						break;
 					default:
-						LogError(Element, "Invalid element '{0}'", Element.Name);
+						LogError(ChildElement, "Invalid element '{0}'", ChildElement.Name);
 						break;
 				}
 			}
 			LeaveScope();
-			return true;
+
 		}
 
 		/// <summary>
@@ -361,40 +381,55 @@ namespace AutomationTool
 					return;
 				}
 
-				// Read the root BuildGraph element
-				EnterScope();
-				foreach (ScriptElement ChildElement in Element.ChildNodes.OfType<ScriptElement>())
+				// Read the child elements
+				ReadTriggerBody(Element, Trigger);
+			}
+		}
+
+		/// <summary>
+		/// Reads the body of a trigger element
+		/// </summary>
+		/// <param name="Element">Xml element to read the definition from</param>
+		void ReadTriggerBody(XmlElement Element, ManualTrigger Trigger)
+		{
+			EnterScope();
+			foreach (ScriptElement ChildElement in Element.ChildNodes.OfType<ScriptElement>())
+			{
+				switch (ChildElement.Name)
 				{
-					switch (ChildElement.Name)
-					{
-						case "Property":
-							ReadProperty(ChildElement);
-							break;
-						case "Local":
-							ReadLocalProperty(ChildElement);
-							break;
-						case "Agent":
-							ReadAgent(ChildElement, Trigger);
-							break;
-						case "Aggregate":
-							ReadAggregate(ChildElement);
-							break;
-						case "Notifier":
-							ReadNotifier(ChildElement);
-							break;
-						case "Warning":
-							ReadDiagnostic(ChildElement, LogEventType.Warning, null, null, Trigger);
-							break;
-						case "Error":
-							ReadDiagnostic(ChildElement, LogEventType.Error, null, null, Trigger);
-							break;
-						default:
-							LogError(ChildElement, "Invalid element '{0}'", ChildElement.Name);
-							break;
-					}
+					case "Property":
+						ReadProperty(ChildElement);
+						break;
+					case "Local":
+						ReadLocalProperty(ChildElement);
+						break;
+					case "Agent":
+						ReadAgent(ChildElement, Trigger);
+						break;
+					case "Aggregate":
+						ReadAggregate(ChildElement);
+						break;
+					case "Notifier":
+						ReadNotifier(ChildElement);
+						break;
+					case "Warning":
+						ReadDiagnostic(ChildElement, LogEventType.Warning, null, null, Trigger);
+						break;
+					case "Error":
+						ReadDiagnostic(ChildElement, LogEventType.Error, null, null, Trigger);
+						break;
+					case "Do":
+						ReadBlock(ChildElement, x => ReadTriggerBody(x, Trigger));
+						break;
+					case "Choose":
+						ReadChoice(ChildElement, x => ReadTriggerBody(x, Trigger));
+						break;
+					default:
+						LogError(ChildElement, "Invalid element '{0}'", ChildElement.Name);
+						break;
 				}
-				LeaveScope();
 			}
+			LeaveScope();
 		}
 
 		/// <summary>
@@ -489,36 +524,53 @@ namespace AutomationTool
 				}
 
 				// Process all the child elements.
-				EnterScope();
-				foreach (ScriptElement ChildElement in Element.ChildNodes.OfType<ScriptElement>())
+				ReadAgentBody(Element, Agent, Trigger);
+			}
+		}
+
+		/// <summary>
+		/// Read the contents of an agent definition
+		/// </summary>
+		/// <param name="Element">Xml element to read the definition from</param>
+		/// <param name="ParentAgent">The agent to contain the definition</param>
+		/// <param name="ControllingTrigger">The enclosing trigger</param>
+		void ReadAgentBody(ScriptElement Element, Agent ParentAgent, ManualTrigger ControllingTrigger)
+		{
+			EnterScope();
+			foreach (ScriptElement ChildElement in Element.ChildNodes.OfType<ScriptElement>())
+			{
+				switch (ChildElement.Name)
 				{
-					switch (ChildElement.Name)
-					{
-						case "Property":
-							ReadProperty(ChildElement);
-							break;
-						case "Local":
-							ReadLocalProperty(ChildElement);
-							break;
-						case "Node":
-							ReadNode(ChildElement, Agent, Trigger);
-							break;
-						case "Aggregate":
-							ReadAggregate(ChildElement);
-							break;
-						case "Warning":
-							ReadDiagnostic(ChildElement, LogEventType.Warning, null, Agent, Trigger);
-							break;
-						case "Error":
-							ReadDiagnostic(ChildElement, LogEventType.Error, null, Agent, Trigger);
-							break;
-						default:
-							LogError(ChildElement, "Unexpected element type '{0}'", ChildElement.Name);
-							break;
-					}
+					case "Property":
+						ReadProperty(ChildElement);
+						break;
+					case "Local":
+						ReadLocalProperty(ChildElement);
+						break;
+					case "Node":
+						ReadNode(ChildElement, ParentAgent, ControllingTrigger);
+						break;
+					case "Aggregate":
+						ReadAggregate(ChildElement);
+						break;
+					case "Warning":
+						ReadDiagnostic(ChildElement, LogEventType.Warning, null, ParentAgent, ControllingTrigger);
+						break;
+					case "Error":
+						ReadDiagnostic(ChildElement, LogEventType.Error, null, ParentAgent, ControllingTrigger);
+						break;
+					case "Do":
+						ReadBlock(ChildElement, x => ReadAgentBody(x, ParentAgent, ControllingTrigger));
+						break;
+					case "Choose":
+						ReadChoice(ChildElement, x => ReadAgentBody(x, ParentAgent, ControllingTrigger));
+						break;
+					default:
+						LogError(ChildElement, "Unexpected element type '{0}'", ChildElement.Name);
+						break;
 				}
-				LeaveScope();
 			}
+			LeaveScope();
 		}
 
 		/// <summary>
@@ -535,6 +587,27 @@ namespace AutomationTool
 			}
 		}
 
+		/// <summary>
+		/// Reads the definition for a report
+		/// </summary>
+		/// <param name="Element">Xml element to read the definition from</param>
+		void ReadReport(ScriptElement Element)
+		{
+			string Name;
+			if (EvaluateCondition(Element) && TryReadObjectName(Element, out Name) && CheckNameIsUnique(Element, Name))
+			{
+				string[] RequiredNames = ReadListAttribute(Element, "Requires");
+
+				Report NewReport = new Report(Name);
+				foreach (Node ReferencedNode in ResolveReferences(Element, RequiredNames))
+				{
+					NewReport.Nodes.Add(ReferencedNode);
+					NewReport.Nodes.UnionWith(ReferencedNode.OrderDependencies);
+				}
+				Graph.NameToReport.Add(Name, NewReport);
+			}
+		}
+
 		/// <summary>
 		/// Reads the definition for a node, and adds it to the given agent
 		/// </summary>
@@ -627,29 +700,7 @@ namespace AutomationTool
 					}
 
 					// Add all the tasks
-					EnterScope();
-					foreach (ScriptElement ChildElement in Element.ChildNodes.OfType<ScriptElement>())
-					{
-						switch (ChildElement.Name)
-						{
-							case "Property":
-								ReadProperty(ChildElement);
-								break;
-							case "Local":
-								ReadLocalProperty(ChildElement);
-								break;
-							case "Warning":
-								ReadDiagnostic(ChildElement, LogEventType.Warning, NewNode, ParentAgent, ControllingTrigger);
-								break;
-							case "Error":
-								ReadDiagnostic(ChildElement, LogEventType.Error, NewNode, ParentAgent, ControllingTrigger);
-								break;
-							default:
-								ReadTask(ChildElement, NewNode.Tasks);
-								break;
-						}
-					}
-					LeaveScope();
+					ReadNodeBody(Element, NewNode, ParentAgent, ControllingTrigger);
 
 					// Add it to the current agent
 					ParentAgent.Nodes.Add(NewNode);
@@ -657,6 +708,76 @@ namespace AutomationTool
 			}
 		}
 
+		/// <summary>
+		/// Reads the contents of a node element
+		/// </summary>
+		/// <param name="Element">Xml element to read the definition from</param>
+		/// <param name="NewNode">The new node that has been created</param>
+		/// <param name="ParentAgent">Agent for the node to be added to</param>
+		/// <param name="ControllingTrigger">The controlling trigger for this node</param>
+		void ReadNodeBody(XmlElement Element, Node NewNode, Agent ParentAgent, ManualTrigger ControllingTrigger)
+		{
+			EnterScope();
+			foreach (ScriptElement ChildElement in Element.ChildNodes.OfType<ScriptElement>())
+			{
+				switch (ChildElement.Name)
+				{
+					case "Property":
+						ReadProperty(ChildElement);
+						break;
+					case "Local":
+						ReadLocalProperty(ChildElement);
+						break;
+					case "Warning":
+						ReadDiagnostic(ChildElement, LogEventType.Warning, NewNode, ParentAgent, ControllingTrigger);
+						break;
+					case "Error":
+						ReadDiagnostic(ChildElement, LogEventType.Error, NewNode, ParentAgent, ControllingTrigger);
+						break;
+					case "Do":
+						ReadBlock(ChildElement, x => ReadNodeBody(x, NewNode, ParentAgent, ControllingTrigger));
+						break;
+					case "Choose":
+						ReadChoice(ChildElement, x => ReadNodeBody(x, NewNode, ParentAgent, ControllingTrigger));
+						break;
+					default:
+						ReadTask(ChildElement, NewNode.Tasks);
+						break;
+				}
+			}
+			LeaveScope();
+		}
+
+		/// <summary>
+		/// Reads a block element
+		/// </summary>
+		/// <param name="Element">Xml element to read the definition from</param>
+		/// <param name="ReadContents">Delegate to read the contents of the element, if the condition evaluates to true</param>
+		void ReadBlock(ScriptElement Element, Action<ScriptElement> ReadContents)
+		{
+			if (EvaluateCondition(Element))
+			{
+				ReadContents(Element);
+			}
+		}
+
+		/// <summary>
+		/// Reads a "Choose" element 
+		/// </summary>
+		/// <param name="Element">Xml element to read the definition from</param>
+		/// <param name="ReadContents">Delegate to read the contents of the element, if the condition evaluates to true</param>
+		void ReadChoice(ScriptElement Element, Action<ScriptElement> ReadContents)
+		{
+			foreach (ScriptElement ChildElement in Element.ChildNodes.OfType<ScriptElement>())
+			{
+				if (ChildElement.Name == "Otherwise" || EvaluateCondition(ChildElement))
+				{
+					ReadContents(ChildElement);
+					break;
+				}
+			}
+		}
+
 		/// <summary>
 		/// Reads a task definition from the given element, and add it to the given list
 		/// </summary>
@@ -740,6 +861,7 @@ namespace AutomationTool
 				string[] ExceptNames = ReadListAttribute(Element, "Except");
 				string[] IndividualNodeNames = ReadListAttribute(Element, "Nodes");
 				string[] TriggerNames = ReadListAttribute(Element, "Triggers");
+				string[] ReportNames = ReadListAttribute(Element, "Reports");
 				string[] Users = ReadListAttribute(Element, "Users");
 				string[] Submitters = ReadListAttribute(Element, "Submitters");
 				bool? bWarnings = Element.HasAttribute("Warnings") ? (bool?)ReadBooleanAttribute(Element, "Warnings", true) : null;
@@ -803,6 +925,23 @@ namespace AutomationTool
 						}
 					}
 				}
+
+				// Add the users to the list of reports
+				if (ReportNames != null)
+				{
+					foreach (string ReportName in ReportNames)
+					{
+						Report Report;
+						if (Graph.NameToReport.TryGetValue(ReportName, out Report))
+						{
+							Report.NotifyUsers.UnionWith(Users);
+						}
+						else
+						{
+							LogError(Element, "Report '{0}' has not been defined", ReportName);
+						}
+					}
+				}
 			}
 		}
 
diff --git a/Engine/Source/Programs/AutomationTool/BuildGraph/Tasks/CompileTask.cs b/Engine/Source/Programs/AutomationTool/BuildGraph/Tasks/CompileTask.cs
--- a/Engine/Source/Programs/AutomationTool/BuildGraph/Tasks/CompileTask.cs
+++ b/Engine/Source/Programs/AutomationTool/BuildGraph/Tasks/CompileTask.cs
@@ -112,7 +112,7 @@ namespace AutomationTool
 			// Build everything
 			Dictionary<UE4Build.BuildTarget, BuildManifest> TargetToManifest = new Dictionary<UE4Build.BuildTarget,BuildManifest>();
             UE4Build Builder = new UE4Build(Job.OwnerCommand);
-			Builder.Build(Agenda, InDeleteBuildProducts: false, InUpdateVersionFiles: false, InForceNoXGE: false, InUseParallelExecutor: true, InTargetToManifest: TargetToManifest);
+			Builder.Build(Agenda, InDeleteBuildProducts: null, InUpdateVersionFiles: false, InForceNoXGE: false, InUseParallelExecutor: true, InTargetToManifest: TargetToManifest);
             UE4Build.CheckBuildProducts(Builder.BuildProductFiles);
 
 			// Tag all the outputs
diff --git a/Engine/Source/Programs/AutomationTool/BuildGraph/TempStorage.cs b/Engine/Source/Programs/AutomationTool/BuildGraph/TempStorage.cs
--- a/Engine/Source/Programs/AutomationTool/BuildGraph/TempStorage.cs
+++ b/Engine/Source/Programs/AutomationTool/BuildGraph/TempStorage.cs
@@ -93,16 +93,16 @@ namespace AutomationTool
 			}
 
 			// Check the timestamp of the file matches. On FAT filesystems writetime has a two seconds resolution (see http://msdn.microsoft.com/en-us/library/windows/desktop/ms724290%28v=vs.85%29.aspx)
-			long TimeDifference = Info.LastWriteTimeUtc.Ticks - LastWriteTimeUtcTicks;
-			if(TimeDifference < -2 || TimeDifference > +2)
+			TimeSpan TimeDifference = new TimeSpan(Info.LastWriteTimeUtc.Ticks - LastWriteTimeUtcTicks);
+			if(TimeDifference.TotalSeconds < -2 || TimeDifference.TotalSeconds > +2)
 			{
 				DateTime ExpectedLocal = new DateTime(LastWriteTimeUtcTicks, DateTimeKind.Utc).ToLocalTime();
 				if(RequireMatchingTimestamps())
 				{
-					CommandUtils.LogError("File date/time mismatch for {0} - was {1}, expected {2}", RelativePath, Info.LastWriteTime, ExpectedLocal);
+					CommandUtils.LogError("File date/time mismatch for {0} - was {1}, expected {2}, TimeDifference {3}", RelativePath, Info.LastWriteTime, ExpectedLocal, TimeDifference);
 					return false;
 				}
-				CommandUtils.Log("File date/time mismatch for {0} - was {1}, expected {2}", RelativePath, Info.LastWriteTime, ExpectedLocal);
+				CommandUtils.Log("File date/time mismatch for {0} - was {1}, expected {2}, TimeDifference {3}", RelativePath, Info.LastWriteTime, ExpectedLocal, TimeDifference);
 			}
 
 			return true;
@@ -892,4 +892,92 @@ namespace AutomationTool
 			}
 		}
 	}
+
+	/// <summary>
+	/// Commandlet to clean up all folders under a temp storage root that are older than a given number of days
+	/// </summary>
+	[Help("Removes folders in a given temp storage directory that are older than a certain time.")]
+	[Help("TempStorageDir=<Directory>", "Path to the root temp storage directory")]
+	[Help("Days=<N>", "Number of days to keep in temp storage")]
+	class CleanTempStorage : BuildCommand
+	{
+		/// <summary>
+		/// Entry point for the commandlet
+		/// </summary>
+		public override void ExecuteBuild()
+		{
+			string TempStorageDir = ParseParamValue("TempStorageDir", null);
+			if (TempStorageDir == null)
+			{
+				throw new AutomationException("Missing -TempStorageDir parameter");
+			}
+
+			string Days = ParseParamValue("Days", null);
+			if (Days == null)
+			{
+				throw new AutomationException("Missing -Days parameter");
+			}
+
+			double DaysValue;
+			if (!Double.TryParse(Days, out DaysValue))
+			{
+				throw new AutomationException("'{0}' is not a valid value for the -Days parameter", Days);
+			}
+
+			DateTime RetainTime = DateTime.UtcNow - TimeSpan.FromDays(DaysValue);
+
+			// Enumerate all the build directories
+			CommandUtils.Log("Scanning {0}...", TempStorageDir);
+			int NumBuilds = 0;
+			List<DirectoryInfo> BuildsToDelete = new List<DirectoryInfo>();
+			foreach (DirectoryInfo StreamDirectory in new DirectoryInfo(TempStorageDir).EnumerateDirectories().OrderBy(x => x.Name))
+			{
+				CommandUtils.Log("Scanning {0}...", StreamDirectory.FullName);
+				foreach (DirectoryInfo BuildDirectory in StreamDirectory.EnumerateDirectories())
+				{
+					if(!BuildDirectory.EnumerateFiles("*", SearchOption.AllDirectories).Any(x => x.LastWriteTimeUtc > RetainTime))
+					{
+						BuildsToDelete.Add(BuildDirectory);
+					}
+					NumBuilds++;
+				}
+			}
+			CommandUtils.Log("Found {0} builds; {1} to delete.", NumBuilds, BuildsToDelete.Count);
+
+			// Loop through them all, checking for files older than the delete time
+			for (int Idx = 0; Idx < BuildsToDelete.Count; Idx++)
+			{
+				try
+				{
+					CommandUtils.Log("[{0}/{1}] Deleting {2}...", Idx + 1, BuildsToDelete.Count, BuildsToDelete[Idx].FullName);
+					BuildsToDelete[Idx].Delete(true);
+				}
+				catch (Exception Ex)
+				{
+					CommandUtils.LogWarning("Failed to delete old manifest folder; will try one file at a time: {0}", Ex);
+					CommandUtils.DeleteDirectory_NoExceptions(true, BuildsToDelete[Idx].FullName);
+				}
+			}
+
+			// Try to delete any empty branch folders
+			foreach (DirectoryInfo StreamDirectory in new DirectoryInfo(TempStorageDir).EnumerateDirectories())
+			{
+				if(StreamDirectory.EnumerateDirectories().Count() == 0 && StreamDirectory.EnumerateFiles().Count() == 0)
+				{
+					try
+					{
+						StreamDirectory.Delete();
+					}
+					catch (IOException)
+					{
+						// only catch "directory is not empty type exceptions, if possible. Best we can do is check for IOException.
+					}
+					catch (Exception Ex)
+					{
+						CommandUtils.LogWarning("Unexpected failure trying to delete (potentially empty) stream directory {0}: {1}", StreamDirectory.FullName, Ex);
+					}
+				}
+			}
+		}
+	}
 }
diff --git a/Engine/Source/Programs/AutomationTool/GUBP/ElectricCommander.cs b/Engine/Source/Programs/AutomationTool/GUBP/ElectricCommander.cs
--- a/Engine/Source/Programs/AutomationTool/GUBP/ElectricCommander.cs
+++ b/Engine/Source/Programs/AutomationTool/GUBP/ElectricCommander.cs
@@ -451,7 +451,7 @@ namespace AutomationTool
 				}
 				JsonWriter.WriteArrayEnd();
 
-				JsonWriter.WriteArrayStart("Triggers");
+				JsonWriter.WriteArrayStart("Reports");
 				foreach (TriggerNode NodeToDo in OrderedToDo.OfType<TriggerNode>())
 				{
 					if(NodeToDo.ControllingTrigger == ExplicitTrigger)
@@ -466,6 +466,7 @@ namespace AutomationTool
 						{
 							JsonWriter.WriteValue("Notify", String.Join(" ", NodeToDo.RecipientsForFailureEmails));
 						}
+						JsonWriter.WriteValue("IsTrigger", true);
 						JsonWriter.WriteObjectEnd();
 					}
 				}
diff --git a/Engine/Source/Programs/AutomationTool/GUBP/TempStorage.cs b/Engine/Source/Programs/AutomationTool/GUBP/TempStorage.cs
--- a/Engine/Source/Programs/AutomationTool/GUBP/TempStorage.cs
+++ b/Engine/Source/Programs/AutomationTool/GUBP/TempStorage.cs
@@ -1268,81 +1268,4 @@ namespace AutomationTool
 
 		}
 	}
-
-	[Help("Removes folders in a given temp storage directory that are older than a certain time.")]
-	class CleanTempStorage : BuildCommand
-	{
-		public override void ExecuteBuild()
-		{
-			string TempStorageDir = ParseParamValue("TempStorageDir", null);
-			if(TempStorageDir == null)
-			{
-				throw new AutomationException("Missing -TempStorageDir parameter");
-			}
-
-			string Days = ParseParamValue("Days", null);
-			if(Days == null)
-			{
-				throw new AutomationException("Missing -Days parameter");
-			}
-
-			double DaysValue;
-			if(!Double.TryParse(Days, out DaysValue))
-			{
-				throw new AutomationException("'{0}' is not a valid value for the -Days parameter", Days);
-			}
-
-			DateTime DeleteTime = DateTime.UtcNow - TimeSpan.FromDays(DaysValue);
-
-			// Enumerate all the build directories
-			CommandUtils.Log("Scanning {0}...", TempStorageDir);
-			int NumBuilds = 0;
-			List<DirectoryInfo> BuildsToDelete = new List<DirectoryInfo>();
-			foreach(DirectoryInfo StreamDirectory in new DirectoryInfo(TempStorageDir).EnumerateDirectories())
-			{
-				foreach(DirectoryInfo BuildDirectory in StreamDirectory.EnumerateDirectories())
-				{
-					FileInfo ManifestFile = BuildDirectory.EnumerateFiles("*.TempManifest", SearchOption.AllDirectories).FirstOrDefault();
-					if(ManifestFile != null && ManifestFile.LastWriteTimeUtc < DeleteTime)
-					{
-						BuildsToDelete.Add(BuildDirectory);
-					}
-					NumBuilds++;
-				}
-			}
-			CommandUtils.Log("Found {0} builds; {1} to delete.", NumBuilds, BuildsToDelete.Count);
-
-			// Loop through them all, checking for files older than the delete time
-			for(int Idx = 0; Idx < BuildsToDelete.Count; Idx++)
-			{
-				try
-				{
-					CommandUtils.Log("[{0}/{1}] Deleting {2}...", Idx + 1, BuildsToDelete.Count, BuildsToDelete[Idx].FullName);
-					BuildsToDelete[Idx].Delete(true);
-				}
-				catch (Exception Ex)
-				{
-					CommandUtils.LogWarning("Failed to delete old manifest folder; will try one file at a time: {0}", Ex);
-					CommandUtils.DeleteDirectory_NoExceptions(true, BuildsToDelete[Idx].FullName);
-				}
-			}
-
-			// Try to delete any empty branch folders
-			foreach(DirectoryInfo StreamDirectory in new DirectoryInfo(TempStorageDir).EnumerateDirectories())
-			{
-				try
-				{
-					StreamDirectory.Delete();
-				}
-				catch (IOException)
-				{
-					// only catch "directory is not empty type exceptions, if possible. Best we can do is check for IOException.
-				}
-				catch (Exception Ex)
-				{
-					CommandUtils.LogWarning("Unexpected failure trying to delete (potentially empty) stream directory {0}: {1}", StreamDirectory.FullName, Ex);
-				}
-			}
-		}
-	}
 }
diff --git a/Engine/Source/Programs/AutomationTool/Scripts/AutomationScripts.Automation.csproj b/Engine/Source/Programs/AutomationTool/Scripts/AutomationScripts.Automation.csproj
--- a/Engine/Source/Programs/AutomationTool/Scripts/AutomationScripts.Automation.csproj
+++ b/Engine/Source/Programs/AutomationTool/Scripts/AutomationScripts.Automation.csproj
@@ -91,6 +91,7 @@
     <Compile Include="Properties\AssemblyInfo.cs" />
     <Compile Include="RunProjectCommand.Automation.cs" />
     <Compile Include="StageTargetCommand.Automation.cs" />
+    <Compile Include="StreamCopyDescription.Automation.cs" />
     <Compile Include="Tests.Automation.cs" />
     <Compile Include="TestSnippet.Automation.cs" />
     <Compile Include="UE4BuildUtils.cs" />
diff --git a/Engine/Source/Programs/AutomationTool/Scripts/BuildCommonTools.Automation.cs b/Engine/Source/Programs/AutomationTool/Scripts/BuildCommonTools.Automation.cs
--- a/Engine/Source/Programs/AutomationTool/Scripts/BuildCommonTools.Automation.cs
+++ b/Engine/Source/Programs/AutomationTool/Scripts/BuildCommonTools.Automation.cs
@@ -50,7 +50,7 @@ public class BuildCommonTools : BuildCommand
 
 		// Build everything. We don't want to touch version files for GitHub builds -- these are "programmer builds" and won't have a canonical build version
 		UE4Build Builder = new UE4Build(this);
-		Builder.Build(Agenda, InDeleteBuildProducts:true, InUpdateVersionFiles: false);
+		Builder.Build(Agenda, InUpdateVersionFiles: false);
 
 		// Add UAT and UBT to the build products
 		Builder.AddUATFilesToBuildProducts();
diff --git a/Engine/Source/Programs/AutomationTool/Scripts/StreamCopyDescription.Automation.cs b/Engine/Source/Programs/AutomationTool/Scripts/StreamCopyDescription.Automation.cs
new file mode 100644
--- /dev/null
+++ b/Engine/Source/Programs/AutomationTool/Scripts/StreamCopyDescription.Automation.cs
@@ -0,0 +1,95 @@
+﻿// Copyright 1998-2016 Epic Games, Inc. All Rights Reserved.
+
+using System;
+using System.Collections.Generic;
+using System.Diagnostics;
+using System.IO;
+using System.Text.RegularExpressions;
+using System.Threading;
+using System.Reflection;
+using System.Linq;
+using AutomationTool;
+using UnrealBuildTool;
+
+namespace AutomationTool
+{
+	/// <summary>
+	/// Generates a changelist description for copying a stream to its parent
+	/// </summary>
+	[RequireP4]
+	public class StreamCopyDescription : BuildCommand
+	{
+		/// <summary>
+		/// Execute the command
+		/// </summary>
+		public override void ExecuteBuild()
+		{
+			// Get the source and target streams
+			string Stream = ParseParamValue("Stream", P4Env.BuildRootP4);
+			string Changes = ParseParamValue("Changes", null);
+
+			// Get changes which haven't been copied up from the current stream
+			string Descriptions;
+			if (Changes == null)
+			{
+				ProcessResult Result = P4.P4(String.Format("interchanges -l -S {0}", Stream), AllowSpew: false);
+				Descriptions = Result.Output.Replace("\r\n", "\n");
+			}
+			else
+			{
+				ProcessResult Result = P4.P4(String.Format("changes -l {0}", Stream), AllowSpew: false);
+				Descriptions = Result.Output.Replace("\r\n", "\n");
+			}
+
+			// Clean any workspace names that may reveal internal information
+			Descriptions = Regex.Replace(Descriptions, "(Change[^@]*)@.*", "$1", RegexOptions.Multiline);
+
+			// Remove changes by the build machine
+			Descriptions = Regex.Replace(Descriptions, "[^\n]*buildmachine\n(\n|\t[^\n]*\n)*", "");
+
+			// Get the last submitted change in the source stream
+			List<P4Connection.ChangeRecord> ChangeRecords;
+			if (!P4.Changes(out ChangeRecords, String.Format("-m1 {0}/...", Stream), AllowSpew: false))
+			{
+				throw new AutomationException("Couldn't get changes for this branch");
+			}
+
+			// Figure out the target stream
+			ProcessResult StreamResult = P4.P4(String.Format("stream -o {0}", Stream), AllowSpew: false);
+			if (StreamResult.ExitCode != 0)
+			{
+				throw new AutomationException("Couldn't get stream description for {0}", Stream);
+			}
+			string Target = P4Spec.FromString(StreamResult.Output).GetField("Parent");
+			if(Target == null)
+			{
+				throw new AutomationException("Couldn't get parent stream for {0}", Stream);
+			}
+
+			// Write the output file
+			string OutputFileName = Path.Combine(CommandUtils.CmdEnv.LocalRoot, "Engine", "Intermediate", "Changes.txt");
+			using (StreamWriter Writer = new StreamWriter(OutputFileName))
+			{
+				Writer.WriteLine("Copying {0} to {1} (Source: {0} @ {2})", Stream, Target.Trim(), ChangeRecords[0].CL);
+				Writer.WriteLine();
+				Writer.WriteLine("==========================");
+				Writer.WriteLine("MAJOR FEATURES + CHANGES");
+				Writer.WriteLine("==========================");
+				Writer.WriteLine();
+
+				foreach (string Line in Descriptions.Split('\n'))
+				{
+					string TrimLine = Line.TrimStart();
+					if(!TrimLine.StartsWith("#codereview") && !TrimLine.StartsWith("#rb"))
+					{
+						Writer.WriteLine(Line);
+					}
+				}
+			}
+			Log("Written {0}.", OutputFileName);
+
+			// Open it with the default text editor
+			Process.Start(OutputFileName);
+		}
+	}
+}
diff --git a/Engine/Source/Programs/UnrealGameSync/UnrealGameSync/Forms/MainWindow.cs b/Engine/Source/Programs/UnrealGameSync/UnrealGameSync/Forms/MainWindow.cs
--- a/Engine/Source/Programs/UnrealGameSync/UnrealGameSync/Forms/MainWindow.cs
+++ b/Engine/Source/Programs/UnrealGameSync/UnrealGameSync/Forms/MainWindow.cs
@@ -2037,7 +2037,7 @@ namespace UnrealGameSync
 		private void BrowseForProject()
 		{
 			OpenFileDialog Dialog = new OpenFileDialog();
-			Dialog.Filter = "Project files (*.uproject, *.uprojectdirs)|*.uproject;*.uprojectdirs|All files (*.*)|*.*" ;
+			Dialog.Filter = "Project files (*.uproject)|*.uproject|Project directory lists (*.uprojectdirs)|*.uprojectdirs|All files (*.*)|*.*" ;
 			if(Dialog.ShowDialog() == System.Windows.Forms.DialogResult.OK)
 			{
 				OpenProject(Dialog.FileName);
diff --git a/Engine/Source/Programs/UnrealGameSync/UnrealGameSync/Properties/AssemblyInfo.cs b/Engine/Source/Programs/UnrealGameSync/UnrealGameSync/Properties/AssemblyInfo.cs
--- a/Engine/Source/Programs/UnrealGameSync/UnrealGameSync/Properties/AssemblyInfo.cs
+++ b/Engine/Source/Programs/UnrealGameSync/UnrealGameSync/Properties/AssemblyInfo.cs
@@ -34,5 +34,5 @@ using System.Runtime.InteropServices;
 // You can specify all the values or you can default the Build and Revision Numbers 
 // by using the '*' as shown below:
 // [assembly: AssemblyVersion("1.0.*")]
-[assembly: AssemblyVersion("1.69.*")]
-[assembly: AssemblyFileVersion("1.69.0.0")]
+[assembly: AssemblyVersion("1.70.*")]
+[assembly: AssemblyFileVersion("1.70.0.0")]
diff --git a/Engine/Source/Runtime/Core/Private/Misc/OutputDevice.cpp b/Engine/Source/Runtime/Core/Private/Misc/OutputDevice.cpp
--- a/Engine/Source/Runtime/Core/Private/Misc/OutputDevice.cpp
+++ b/Engine/Source/Runtime/Core/Private/Misc/OutputDevice.cpp
@@ -142,21 +142,46 @@ void StaticFailDebug( const TCHAR* Error, const ANSICHAR* File, int32 Line, cons
 	FCString::Strncat( GErrorHist, TEXT( "\r\n\r\n" ), ARRAY_COUNT( GErrorHist ) );
 }
 
-void FDebug::ConditionallyEmitBeginCrashUATMarker()
+void FDebug::OutputMultiLineCallstack(const ANSICHAR* File, int32 Line, const FName& LogName, const TCHAR* Heading, TCHAR* Message, ELogVerbosity::Type Verbosity)
 {
-	const bool bWriteUATMarkers = FParse::Param( FCommandLine::Get(), TEXT( "CrashForUAT" ) ) && FParse::Param( FCommandLine::Get(), TEXT( "stdout" ) );
+	const bool bWriteUATMarkers = FParse::Param(FCommandLine::Get(), TEXT("CrashForUAT")) && FParse::Param(FCommandLine::Get(), TEXT("stdout"));
+
 	if (bWriteUATMarkers)
 	{
-		UE_LOG( LogOutputDevice, Error, TEXT( "\nbegin: stack for UAT" ) );
+		FMsg::Logf(File, Line, LogName, Verbosity, TEXT("begin: stack for UAT"));
+	}
+
+	FMsg::Logf(File, Line, LogName, Verbosity, TEXT("%s"), Heading);
+	FMsg::Logf(File, Line, LogName, Verbosity, TEXT(""));
+
+	for (TCHAR* LineStart = Message;; )
+	{
+		// Find the end of the current line
+		TCHAR* LineEnd = LineStart;
+		while (*LineEnd != 0 && *LineEnd != '\r' && *LineEnd != '\n')
+		{
+			LineEnd++;
+		}
+
+		// Output it
+		TCHAR LineEndCharacter = *LineEnd;
+		*LineEnd = 0;
+		FMsg::Logf(File, Line, LogName, Verbosity, TEXT("%s"), LineStart);
+		*LineEnd = LineEndCharacter;
+
+		// Quit if this was the last line
+		if (*LineEnd == 0)
+		{
+			break;
+		}
+
+		// Move to the next line
+		LineStart = (LineEnd[0] == '\r' && LineEnd[1] == '\n') ? (LineEnd + 2) : (LineEnd + 1);
 	}
-}
 
-void FDebug::ConditionallyEmitEndCrashUATMarker()
-{
-	const bool bWriteUATMarkers = FParse::Param( FCommandLine::Get(), TEXT( "CrashForUAT" ) ) && FParse::Param( FCommandLine::Get(), TEXT( "stdout" ) );
 	if (bWriteUATMarkers)
 	{
-		UE_LOG( LogOutputDevice, Error, TEXT( "\nend: stack for UAT" ) );
+		FMsg::Logf(File, Line, LogName, Verbosity, TEXT("end: stack for UAT"));
 	}
 }
 
@@ -268,9 +293,9 @@ void FDebug::EnsureFailed(const ANSICHAR* Expr, const ANSICHAR* File, int32 Line
 		FMemory::SystemFree( StackTrace );
 
 		// Dump the error and flush the log.
-		FDebug::ConditionallyEmitBeginCrashUATMarker();
-		UE_LOG(LogOutputDevice, Warning, TEXT("=== Handled error: ===") LINE_TERMINATOR LINE_TERMINATOR TEXT("%s"), ErrorMsg);
-		FDebug::ConditionallyEmitEndCrashUATMarker();
+#if !NO_LOGGING
+		FDebug::OutputMultiLineCallstack(__FILE__, __LINE__, LogOutputDevice.GetCategoryName(), TEXT("=== Handled ensure: ==="), ErrorMsg, ELogVerbosity::Warning);
+#endif
 		GLog->Flush();
 
 		// Submit the error report to the server! (and display a balloon in the system tray)
diff --git a/Engine/Source/Runtime/Core/Private/Windows/WindowsPlatformOutputDevices.cpp b/Engine/Source/Runtime/Core/Private/Windows/WindowsPlatformOutputDevices.cpp
--- a/Engine/Source/Runtime/Core/Private/Windows/WindowsPlatformOutputDevices.cpp
+++ b/Engine/Source/Runtime/Core/Private/Windows/WindowsPlatformOutputDevices.cpp
@@ -129,10 +129,9 @@ void FOutputDeviceWindowsError::HandleError()
 	FCoreDelegates::OnHandleSystemError.Broadcast();
 
 	// Dump the error and flush the log.
-	FDebug::ConditionallyEmitBeginCrashUATMarker();
-	UE_LOG( LogWindows, Error, TEXT( "=== Critical error: ===" ) LINE_TERMINATOR TEXT( "%s" ) LINE_TERMINATOR, GErrorHist );
-	FDebug::ConditionallyEmitEndCrashUATMarker();
-
+#if !NO_LOGGING
+	FDebug::OutputMultiLineCallstack(__FILE__, __LINE__, LogWindows.GetCategoryName(), TEXT("=== Critical error: ==="), GErrorHist, ELogVerbosity::Error);
+#endif
 	GLog->PanicFlushThreadedLogs();
 
 	// Unhide the mouse.
diff --git a/Engine/Source/Runtime/Core/Public/Misc/OutputDevice.h b/Engine/Source/Runtime/Core/Public/Misc/OutputDevice.h
--- a/Engine/Source/Runtime/Core/Public/Misc/OutputDevice.h
+++ b/Engine/Source/Runtime/Core/Public/Misc/OutputDevice.h
@@ -409,11 +409,8 @@ struct CORE_API FMsg
  **/
 struct CORE_API FDebug
 {
-	/** Conditionally emits a special UAT marker. */
-	static void ConditionallyEmitBeginCrashUATMarker();
-
-	/** Conditionally emits a special UAT marker. */
-	static void ConditionallyEmitEndCrashUATMarker();
+	/** Outputs a callstack message, separating out each line into a separate log call by modifying the buffer in-place */
+	static void OutputMultiLineCallstack(const ANSICHAR* File, int32 Line, const FName& LogName, const TCHAR* Heading, TCHAR* Message, ELogVerbosity::Type Verbosity);
 
 	/** Logs final assert message and exits the program. */
 	static void VARARGS AssertFailed(const ANSICHAR* Expr, const ANSICHAR* File, int32 Line, const TCHAR* Format = TEXT(""), ...);
