Fixed error in StreamCopyDescription script, where the changes parameter wasn't being used to get descriptions from another stream.

diff --git a/Engine/Source/Programs/AutomationTool/Scripts/StreamCopyDescription.Automation.cs b/Engine/Source/Programs/AutomationTool/Scripts/StreamCopyDescription.Automation.cs
--- a/Engine/Source/Programs/AutomationTool/Scripts/StreamCopyDescription.Automation.cs
+++ b/Engine/Source/Programs/AutomationTool/Scripts/StreamCopyDescription.Automation.cs
@@ -37,7 +37,7 @@ namespace AutomationTool
 			}
 			else
 			{
-				ProcessResult Result = P4.P4(String.Format("changes -l {0}", Stream), AllowSpew: false);
+				ProcessResult Result = P4.P4(String.Format("changes -l {0}", Changes), AllowSpew: false);
 				Descriptions = Result.Output.Replace("\r\n", "\n");
 			}
 
