Suppress warnings about dsym files being overwritten on Mac until we have a proper fix.

diff --git a/Engine/Source/Programs/AutomationTool/GUBP/TempStorage.cs b/Engine/Source/Programs/AutomationTool/GUBP/TempStorage.cs
--- a/Engine/Source/Programs/AutomationTool/GUBP/TempStorage.cs
+++ b/Engine/Source/Programs/AutomationTool/GUBP/TempStorage.cs
@@ -63,7 +63,7 @@ namespace AutomationTool
 					// this is a bit of a hack, but UAT itself creates these, so we need to allow them to be 
 					bool bOkToBeDifferent = Name.Contains("Engine/Binaries/DotNET/");
 					// Lets just allow all mac warnings to be silent
-					bOkToBeDifferent = bOkToBeDifferent || Name.Contains("Engine/Binaries/Mac/");
+					bOkToBeDifferent = bOkToBeDifferent || Name.Contains("/Binaries/Mac/");
 					// this is a problem with mac compiles
 					bOkToBeDifferent = bOkToBeDifferent || Name.EndsWith("MacOS/libogg.dylib");
 					bOkToBeDifferent = bOkToBeDifferent || Name.EndsWith("MacOS/libvorbis.dylib");
