Copying //UE4/Dev-Build to //UE4/Dev-Main (Source: //UE4/Dev-Build @ 2972140)

==========================
MAJOR FEATURES + CHANGES
==========================

Change 2959679 on 2016/04/28 by [REDACTED_PERSON]

	UGS: Show the original author for changes with a #ROBOMERGE-AUTHOR tag.

Change 2959695 on 2016/04/28 by [REDACTED_PERSON]

	UGS: Only filter out changes from by buildmachine that contain the string "CIS Counter".

Change 2960798 on 2016/04/29 by [REDACTED_PERSON]

	Remove C++ version of ParallelExecutor. Now implemented in C# as part of UAT.

Change 2960928 on 2016/04/29 by [REDACTED_PERSON]

	UGS: Change filter for buildmachine changes to only include rebuilt lightmaps.

Change 2963214 on 2016/05/02 by [REDACTED_PERSON]

	BuildGraph: Allow specifying optional dependencies for a node, indicating that the build products from an upstream node are desired, but should not block the node from running.

Change 2964454 on 2016/05/03 by [REDACTED_PERSON]

	Change PostBuildInfoTool to PostBadgeStatus, and add position-independent argument parsing.

Change 2964533 on 2016/05/03 by [REDACTED_PERSON]

	BuildGraph: Add the ability to generate summary badges from BuildGraph scripts, which can be pushed into a separate database for consumption by UGS.

Change 2964852 on 2016/05/03 by [REDACTED_PERSON]

	BuildGraph: Add a task which can submit a set of files to Perforce, optionally creating and using a different workspace to do so.

Change 2966856 on 2016/05/04 by [REDACTED_PERSON]

	EC: Allow specifying a filter for the changes considered when looking for the most recent change. Allows filtering out content changes for UGS builds, code-only builds, etc...

Change 2966867 on 2016/05/04 by [REDACTED_PERSON]

	EC: Restore code to always set time in CIS state; we never want large builds to trigger off their defined interval.

Change 2967504 on 2016/05/05 by [REDACTED_PERSON]

	UAT: Make sure the intermediate directory exists before writing out the list of changes in StreamCopyDescription.

Change 2967778 on 2016/05/05 by [REDACTED_PERSON]

	UAT: Detect the P4 environment by querying Perforce for the setting of P4PORT, rather than assuming it's set in an environment variable. Windows stores this setting in the registry rather than the environment, but it's also valid to be set via P4CONFIG.

Change 2967815 on 2016/05/05 by [REDACTED_PERSON]

	EC: Copy the initial resource pool setting from the stream settings into an EC property

Change 2967873 on 2016/05/05 by [REDACTED_PERSON]

	EC: Allow stream settings to be stored directly in /GUBP_V5/Streams/ rather than having to be in a child property sheet.

Change 2969294 on 2016/05/06 by [REDACTED_PERSON]

	EC: Extend ConformResources command to allow updating the pools that resources are assigned to, and to limit the number of machines which are syncing at once. Also added new EC procedure to allow specifying these arguments.

Change 2969371 on 2016/05/06 by [REDACTED_PERSON]

	EC: Allow overriding the stream and workspace identifier synced by the builders. Overriding the stream allows syncing a narrower view of files (eg. Dev-Main vs Main), and overriding the workspace identifier allows sharing a workspace between two streams.

Change 2970623 on 2016/05/09 by [REDACTED_PERSON]

	UAT: Prevent Ctrl-C handler delegate from being garbage collected and failing to be triggered.

Change 2970627 on 2016/05/09 by [REDACTED_PERSON]

	UAT: Don't limit the list of valid target platforms specified on the command line to just those that we have initialized. Ignoring the platform if the SDK is not installed is never what the user wants.

Change 2972140 on 2016/05/10 by [REDACTED_PERSON]

	Change 'Engine, Tools and Monolithics' to include QAGame and Template editors, but exclude everything downstream of a trigger.

diff --git a/Engine/Source/Programs/AutomationTool/AutomationUtils/CommandUtils.cs b/Engine/Source/Programs/AutomationTool/AutomationUtils/CommandUtils.cs
--- a/Engine/Source/Programs/AutomationTool/AutomationUtils/CommandUtils.cs
+++ b/Engine/Source/Programs/AutomationTool/AutomationUtils/CommandUtils.cs
@@ -1698,36 +1698,6 @@ namespace AutomationTool
 			}
 		}
 
-		/// <summary>
-		/// Gets environment variables set by a batch file.
-		/// </summary>
-		/// <param name="BatchFileName">Filename that sets the environment variables</param>
-		/// <param name="AlsoSet">True if found variables should be automatically set withing this process.</param>
-		/// <returns>Dictionary of environment variables set by the batch file.</returns>
-		public static CaselessDictionary<string> GetEnvironmentVariablesFromBatchFile(string BatchFileName, bool AlsoSet = false)
-		{
-			CaselessDictionary<string> Result;
-			try
-			{
-				Result = HarvestEnvVars.HarvestEnvVarsFromBatchFile(BatchFileName, "", HarvestEnvVars.EPathOverride.User);
-			}
-			catch (Exception Ex)
-			{
-				throw new AutomationException(Ex, "Failed to harvest environment variables");
-			}
-
-			if (AlsoSet)
-			{
-				// Set the environment variables
-				foreach (var Envvar in Result)
-				{
-					Environment.SetEnvironmentVariable(Envvar.Key, Envvar.Value);
-				}
-			}
-
-			return Result;
-		}
-
 		#endregion
 
 		#region CommandLine
diff --git a/Engine/Source/Programs/AutomationTool/AutomationUtils/LocalP4Environment.cs b/Engine/Source/Programs/AutomationTool/AutomationUtils/LocalP4Environment.cs
--- a/Engine/Source/Programs/AutomationTool/AutomationUtils/LocalP4Environment.cs
+++ b/Engine/Source/Programs/AutomationTool/AutomationUtils/LocalP4Environment.cs
@@ -25,11 +25,7 @@ namespace AutomationTool
 		protected override void InitEnvironment(P4Connection Connection, CommandEnvironment CmdEnv)
 		{
 			var HostName = Environment.MachineName.ToLower();
-			var P4PortEnv = Environment.GetEnvironmentVariable("P4PORT");
-			if (String.IsNullOrEmpty(P4PortEnv))
-			{
-				P4PortEnv = DetectP4Port();
-			}
+			var P4PortEnv = DetectP4Port();
 
 			var UserName = CommandUtils.GetEnvVar(EnvVarNames.User);
 			if (String.IsNullOrEmpty(UserName))
@@ -297,22 +293,46 @@ namespace AutomationTool
 		/// <returns>Source control server address.</returns>
 		private static string DetectP4Port()
 		{
-			// Try to read the P4PORT environment and check if it is set correctly
-			var P4PortEnv = Environment.GetEnvironmentVariable(EnvVarNames.P4Port);
-
-			// If not, try to fallback to Mapping.P4Port and set this as P4PORT before continueing
-			if (!String.IsNullOrEmpty(P4PortEnv))
-			{
-				Log.TraceWarning("P4PORT is not set. Falling back to {0} which is set to {1}.", EnvVarNames.P4Port, P4PortEnv);
-			}
-			else
+			string P4PortEnv = Environment.GetEnvironmentVariable("P4PORT");
+			if(String.IsNullOrEmpty(P4PortEnv))
 			{
-				// If that fails as well, we just give it a shot with perforce:1666 and hope that this works
-				Log.TraceWarning("P4PORT is not set. Trying to fallback to perforce:1666");
-				P4PortEnv = "perforce:1666";
-			}
-			Environment.SetEnvironmentVariable("P4PORT", P4PortEnv);
+				// If it's not set, spawn Perforce to get the current server port setting
+				ProcessResult Result = CommandUtils.Run(HostPlatform.Current.P4Exe, "set P4PORT", null, CommandUtils.ERunOptions.NoLoggingOfRunCommand);
+				if (Result.ExitCode == 0)
+				{
+					const string KeyName = "P4PORT=";
+					if (Result.Output.StartsWith(KeyName))
+					{
+						int LastIdx = Result.Output.IndexOfAny(new char[] { ' ', '\n' });
+						if (LastIdx == -1)
+						{
+							LastIdx = Result.Output.Length;
+						}
+						P4PortEnv = Result.Output.Substring(KeyName.Length, LastIdx - KeyName.Length);
+					}
+				}
+
+				// Otherwise fallback to the uebp variables, or the default
+				if(String.IsNullOrEmpty(P4PortEnv))
+				{
+					// Try to read the P4PORT environment and check if it is set correctly
+					P4PortEnv = Environment.GetEnvironmentVariable(EnvVarNames.P4Port);
+
+					// If not, try to fallback to Mapping.P4Port and set this as P4PORT before continueing
+					if (!String.IsNullOrEmpty(P4PortEnv))
+					{
+						Log.TraceWarning("P4PORT is not set. Falling back to {0} which is set to {1}.", EnvVarNames.P4Port, P4PortEnv);
+					}
+					else
+					{
+						// If that fails as well, we just give it a shot with perforce:1666 and hope that this works
+						Log.TraceWarning("P4PORT is not set. Trying to fallback to perforce:1666");
+						P4PortEnv = "perforce:1666";
+					}
+				}
 
+				Environment.SetEnvironmentVariable("P4PORT", P4PortEnv);
+			}
 			return P4PortEnv;
 		}
 	}
diff --git a/Engine/Source/Programs/AutomationTool/AutomationUtils/ProjectParams.cs b/Engine/Source/Programs/AutomationTool/AutomationUtils/ProjectParams.cs
--- a/Engine/Source/Programs/AutomationTool/AutomationUtils/ProjectParams.cs
+++ b/Engine/Source/Programs/AutomationTool/AutomationUtils/ProjectParams.cs
@@ -162,11 +162,14 @@ namespace AutomationTool
 					{
 						// Look up platform names in the command line: -Platform_1 -Platform_2 ... -Platform_k
 						TargetPlatforms = new List<UnrealTargetPlatform>();
-						foreach (var Plat in CommandUtils.KnownTargetPlatforms)
+						foreach (UnrealTargetPlatform Plat in Enum.GetValues(typeof(UnrealTargetPlatform)))
 						{
-							if (Command.ParseParam(Plat.ToString()))
+							if (Plat != UnrealTargetPlatform.Unknown)
 							{
-								TargetPlatforms.Add(Plat);
+								if (Command.ParseParam(Plat.ToString()))
+								{
+									TargetPlatforms.Add(Plat);
+								}
 							}
 						}
 					}
diff --git a/Engine/Source/Programs/AutomationTool/BuildGraph/Badge.cs b/Engine/Source/Programs/AutomationTool/BuildGraph/Badge.cs
new file mode 100644
--- /dev/null
+++ b/Engine/Source/Programs/AutomationTool/BuildGraph/Badge.cs
@@ -0,0 +1,50 @@
+﻿using System;
+using System.Collections.Generic;
+using System.Diagnostics;
+using System.Linq;
+using System.Text;
+using System.Threading.Tasks;
+
+namespace AutomationTool
+{
+	/// <summary>
+	/// Defines a badge which gives an at-a-glance summary of part of the build, and can be displayed in UGS
+	/// </summary>
+	class Badge
+	{
+		/// <summary>
+		/// Name of this badge
+		/// </summary>
+		public readonly string Name;
+
+		/// <summary>
+		/// Depot path to the project that this badge applies to. Used for filtering in UGS.
+		/// </summary>
+		public readonly string Project;
+
+		/// <summary>
+		/// Set of nodes that this badge reports the status of
+		/// </summary>
+		public HashSet<Node> Nodes = new HashSet<Node>();
+
+		/// <summary>
+		/// Constructor
+		/// </summary>
+		/// <param name="InName">Name of this report</param>
+		/// <param name="InProject">Depot path to the project that this badge applies to</param>
+		public Badge(string InName, string InProject)
+		{
+			Name = InName;
+			Project = InProject;
+		}
+
+		/// <summary>
+		/// Get the name of this badge
+		/// </summary>
+		/// <returns>The name of this badge</returns>
+		public override string ToString()
+		{
+			return Name;
+		}
+	}
+}
diff --git a/Engine/Source/Programs/AutomationTool/BuildGraph/BuildGraph.Automation.csproj b/Engine/Source/Programs/AutomationTool/BuildGraph/BuildGraph.Automation.csproj
--- a/Engine/Source/Programs/AutomationTool/BuildGraph/BuildGraph.Automation.csproj
+++ b/Engine/Source/Programs/AutomationTool/BuildGraph/BuildGraph.Automation.csproj
@@ -62,6 +62,7 @@
       <Link>Properties\MetaData.cs</Link>
     </Compile>
     <Compile Include="Agent.cs" />
+    <Compile Include="Badge.cs" />
     <Compile Include="BuildGraph.cs" />
     <Compile Include="Graph.cs" />
     <Compile Include="Node.cs" />
@@ -81,6 +82,7 @@
     <Compile Include="Tasks\RenameTask.cs" />
     <Compile Include="Tasks\SpawnTask.cs" />
     <Compile Include="Tasks\LogTask.cs" />
+    <Compile Include="Tasks\SubmitTask.cs" />
     <Compile Include="Tasks\TagTask.cs" />
     <Compile Include="Tasks\CommandTask.cs" />
     <Compile Include="Tasks\CompileTask.cs" />
diff --git a/Engine/Source/Programs/AutomationTool/BuildGraph/Graph.cs b/Engine/Source/Programs/AutomationTool/BuildGraph/Graph.cs
--- a/Engine/Source/Programs/AutomationTool/BuildGraph/Graph.cs
+++ b/Engine/Source/Programs/AutomationTool/BuildGraph/Graph.cs
@@ -101,6 +101,11 @@ namespace AutomationTool
 		/// </summary>
 		public Dictionary<string, Node[]> AggregateNameToNodes = new Dictionary<string,Node[]>(StringComparer.InvariantCultureIgnoreCase);
 
+		/// <summary>
+		/// List of badges that can be displayed for this build
+		/// </summary>
+		public List<Badge> Badges = new List<Badge>();
+
 		/// <summary>
 		/// Diagnostic messages for this graph
 		/// </summary>
@@ -256,6 +261,9 @@ namespace AutomationTool
 			// Create a new list of aggregates for everything that's left
 			AggregateNameToNodes = AggregateNameToNodes.Where(x => x.Value.All(y => RetainNodes.Contains(y))).ToDictionary(Pair => Pair.Key, Pair => Pair.Value);
 
+			// Remove any badges which do not have all their dependencies
+			Badges.RemoveAll(x => x.Nodes.Any(y => !RetainNodes.Contains(y)));
+
 			// Remove any diagnostics which are no longer part of the graph
 			Diagnostics.RemoveAll(x => (x.EnclosingNode != null && !RetainNodes.Contains(x.EnclosingNode)) || (x.EnclosingAgent != null && !Agents.Contains(x.EnclosingAgent)));
 		}
@@ -311,6 +319,18 @@ namespace AutomationTool
 					Writer.WriteEndElement();
 				}
 
+				foreach (Badge Badge in Badges)
+				{
+					Writer.WriteStartElement("Badge");
+					Writer.WriteAttributeString("Name", Badge.Name);
+					if (Badge.Project != null)
+					{
+						Writer.WriteAttributeString("Project", Badge.Project);
+					}
+					Writer.WriteAttributeString("Requires", String.Join(";", Badge.Nodes.Select(x => x.Name)));
+					Writer.WriteEndElement();
+				}
+
 				Writer.WriteEndElement();
 			}
 		}
@@ -375,6 +395,33 @@ namespace AutomationTool
 				}
 				JsonWriter.WriteArrayEnd();
 
+				// Write all the badges
+				JsonWriter.WriteArrayStart("Badges");
+				foreach (Badge Badge in Badges)
+				{
+					Node[] Dependencies = Badge.Nodes.Where(x => NodesToExecute.Contains(x)).ToArray();
+					if (Dependencies.Length > 0)
+					{
+						// Reduce that list to the smallest subset of direct dependencies
+						HashSet<Node> DirectDependencies = new HashSet<Node>(Dependencies);
+						foreach (Node Dependency in Dependencies)
+						{
+							DirectDependencies.ExceptWith(Dependency.OrderDependencies);
+						}
+
+						JsonWriter.WriteObjectStart();
+						JsonWriter.WriteValue("Name", Badge.Name);
+						if (!String.IsNullOrEmpty(Badge.Project))
+						{
+							JsonWriter.WriteValue("Project", Badge.Project);
+						}
+						JsonWriter.WriteValue("AllDependencies", String.Join(";", Agents.SelectMany(x => x.Nodes).Where(x => Dependencies.Contains(x)).Select(x => x.Name)));
+						JsonWriter.WriteValue("DirectDependencies", String.Join(";", DirectDependencies.Select(x => x.Name)));
+						JsonWriter.WriteObjectEnd();
+					}
+				}
+				JsonWriter.WriteArrayEnd();
+
 				// Write all the triggers and reports. 
 				JsonWriter.WriteArrayStart("Reports");
 				foreach (Report Report in NameToReport.Values)
diff --git a/Engine/Source/Programs/AutomationTool/BuildGraph/Schema.cs b/Engine/Source/Programs/AutomationTool/BuildGraph/Schema.cs
--- a/Engine/Source/Programs/AutomationTool/BuildGraph/Schema.cs
+++ b/Engine/Source/Programs/AutomationTool/BuildGraph/Schema.cs
@@ -116,6 +116,7 @@ namespace AutomationTool
 		NodeBody,
 		Aggregate,
 		Report,
+		Badge,
 		Notify,
 		Include,
 		Property,
@@ -278,6 +279,7 @@ namespace AutomationTool
 			NewSchema.Items.Add(CreateNodeBodyType(TaskNameToType));
 			NewSchema.Items.Add(CreateAggregateType());
 			NewSchema.Items.Add(CreateReportType());
+			NewSchema.Items.Add(CreateBadgeType());
 			NewSchema.Items.Add(CreateNotifyType());
 			NewSchema.Items.Add(CreateIncludeType());
 			NewSchema.Items.Add(CreatePropertyType(ScriptSchemaStandardType.Property));
@@ -401,6 +403,7 @@ namespace AutomationTool
 			GraphChoice.Items.Add(CreateSchemaElement("Trigger", ScriptSchemaStandardType.Trigger));
 			GraphChoice.Items.Add(CreateSchemaElement("Aggregate", ScriptSchemaStandardType.Aggregate));
 			GraphChoice.Items.Add(CreateSchemaElement("Report", ScriptSchemaStandardType.Report));
+			GraphChoice.Items.Add(CreateSchemaElement("Badge", ScriptSchemaStandardType.Badge));
 			GraphChoice.Items.Add(CreateSchemaElement("Notify", ScriptSchemaStandardType.Notify));
 			GraphChoice.Items.Add(CreateSchemaElement("Warning", ScriptSchemaStandardType.Warning));
 			GraphChoice.Items.Add(CreateSchemaElement("Error", ScriptSchemaStandardType.Error));
@@ -571,12 +574,27 @@ namespace AutomationTool
 		/// <returns>Type definition for a report</returns>
 		static XmlSchemaType CreateReportType()
 		{
-			XmlSchemaComplexType AggregateType = new XmlSchemaComplexType();
-			AggregateType.Name = GetTypeName(ScriptSchemaStandardType.Report);
-			AggregateType.Attributes.Add(CreateSchemaAttribute("Name", ScriptSchemaStandardType.Name, XmlSchemaUse.Required));
-			AggregateType.Attributes.Add(CreateSchemaAttribute("Requires", ScriptSchemaStandardType.NameOrTagList, XmlSchemaUse.Required));
-			AggregateType.Attributes.Add(CreateSchemaAttribute("If", ScriptSchemaStandardType.BalancedString, XmlSchemaUse.Optional));
-			return AggregateType;
+			XmlSchemaComplexType ReportType = new XmlSchemaComplexType();
+			ReportType.Name = GetTypeName(ScriptSchemaStandardType.Report);
+			ReportType.Attributes.Add(CreateSchemaAttribute("Name", ScriptSchemaStandardType.Name, XmlSchemaUse.Required));
+			ReportType.Attributes.Add(CreateSchemaAttribute("Requires", ScriptSchemaStandardType.NameOrTagList, XmlSchemaUse.Required));
+			ReportType.Attributes.Add(CreateSchemaAttribute("If", ScriptSchemaStandardType.BalancedString, XmlSchemaUse.Optional));
+			return ReportType;
+		}
+
+		/// <summary>
+		/// Creates the schema type representing the badge type
+		/// </summary>
+		/// <returns>Type definition for a badge</returns>
+		static XmlSchemaType CreateBadgeType()
+		{
+			XmlSchemaComplexType BadgeType = new XmlSchemaComplexType();
+			BadgeType.Name = GetTypeName(ScriptSchemaStandardType.Badge);
+			BadgeType.Attributes.Add(CreateSchemaAttribute("Name", ScriptSchemaStandardType.Name, XmlSchemaUse.Required));
+			BadgeType.Attributes.Add(CreateSchemaAttribute("Requires", ScriptSchemaStandardType.NameOrTagList, XmlSchemaUse.Required));
+			BadgeType.Attributes.Add(CreateSchemaAttribute("Project", ScriptSchemaStandardType.BalancedString, XmlSchemaUse.Required));
+			BadgeType.Attributes.Add(CreateSchemaAttribute("If", ScriptSchemaStandardType.BalancedString, XmlSchemaUse.Optional));
+			return BadgeType;
 		}
 
 		/// <summary>
diff --git a/Engine/Source/Programs/AutomationTool/BuildGraph/Script.cs b/Engine/Source/Programs/AutomationTool/BuildGraph/Script.cs
--- a/Engine/Source/Programs/AutomationTool/BuildGraph/Script.cs
+++ b/Engine/Source/Programs/AutomationTool/BuildGraph/Script.cs
@@ -277,6 +277,9 @@ namespace AutomationTool
 					case "Report":
 						ReadReport(ChildElement);
 						break;
+					case "Badge":
+						ReadBadge(ChildElement);
+						break;
 					case "Notify":
 						ReadNotifier(ChildElement);
 						break;
@@ -608,6 +611,28 @@ namespace AutomationTool
 			}
 		}
 
+		/// <summary>
+		/// Reads the definition for a badge
+		/// </summary>
+		/// <param name="Element">Xml element to read the definition from</param>
+		void ReadBadge(ScriptElement Element)
+		{
+			string Name;
+			if (EvaluateCondition(Element) && TryReadObjectName(Element, out Name))
+			{
+				string[] RequiredNames = ReadListAttribute(Element, "Requires");
+				string Project = ReadAttribute(Element, "Project");
+
+				Badge NewBadge = new Badge(Name, Project);
+				foreach (Node ReferencedNode in ResolveReferences(Element, RequiredNames))
+				{
+					NewBadge.Nodes.Add(ReferencedNode);
+					NewBadge.Nodes.UnionWith(ReferencedNode.OrderDependencies);
+				}
+				Graph.Badges.Add(NewBadge);
+			}
+		}
+
 		/// <summary>
 		/// Reads the definition for a node, and adds it to the given agent
 		/// </summary>
diff --git a/Engine/Source/Programs/AutomationTool/BuildGraph/Tasks/SubmitTask.cs b/Engine/Source/Programs/AutomationTool/BuildGraph/Tasks/SubmitTask.cs
new file mode 100644
--- /dev/null
+++ b/Engine/Source/Programs/AutomationTool/BuildGraph/Tasks/SubmitTask.cs
@@ -0,0 +1,138 @@
+﻿using System;
+using System.Collections.Generic;
+using System.Linq;
+using System.Text;
+using System.Threading.Tasks;
+using System.Xml;
+using UnrealBuildTool;
+
+namespace AutomationTool.Tasks
+{
+	/// <summary>
+	/// Parameters for the submit task
+	/// </summary>
+	public class SubmitTaskParameters
+	{
+		/// <summary>
+		/// The description for the submitted changelist
+		/// </summary>
+		[TaskParameter]
+		public string Description;
+
+		/// <summary>
+		/// The files to submit
+		/// </summary>
+		[TaskParameter]
+		public string Files;
+
+		/// <summary>
+		/// The filetype for the submitted files
+		/// </summary>
+		[TaskParameter(Optional = true)]
+		public string FileType;
+
+		/// <summary>
+		/// The workspace name. If specified, a new workspace will be created using the given stream and root directory to submit the files.
+		/// </summary>
+		[TaskParameter(Optional=true)]
+		public string Workspace;
+
+		/// <summary>
+		/// The stream for the workspace. If not specified, defaults to the current stream.
+		/// </summary>
+		[TaskParameter(Optional = true)]
+		public string Stream;
+
+		/// <summary>
+		/// Root directory for the stream. If not specified, defaults to the current root directory.
+		/// </summary>
+		[TaskParameter(Optional = true)]
+		public string RootDir;
+	}
+
+	/// <summary>
+	/// Task which submits a the version files in the current branch
+	/// </summary>
+	[TaskElement("Submit", typeof(SubmitTaskParameters))]
+	public class SubmitTask : CustomTask
+	{
+		/// <summary>
+		/// Parameters for the task
+		/// </summary>
+		SubmitTaskParameters Parameters;
+
+		/// <summary>
+		/// Construct a version task
+		/// </summary>
+		/// <param name="InParameters">Parameters for this task</param>
+		public SubmitTask(SubmitTaskParameters InParameters)
+		{
+			Parameters = InParameters;
+		}
+
+		/// <summary>
+		/// Execute the task.
+		/// </summary>
+		/// <param name="Job">Information about the current job</param>
+		/// <param name="BuildProducts">Set of build products produced by this node.</param>
+		/// <param name="TagNameToFileSet">Mapping from tag names to the set of files they include</param>
+		/// <returns>True if the task succeeded</returns>
+		public override bool Execute(JobContext Job, HashSet<FileReference> BuildProducts, Dictionary<string, HashSet<FileReference>> TagNameToFileSet)
+		{
+			HashSet<FileReference> Files = ResolveFilespec(CommandUtils.RootDirectory, Parameters.Files, TagNameToFileSet);
+			if (CommandUtils.AllowSubmit && Files.Count > 0)
+			{
+				// Get the connection that we're going to submit with
+				P4Connection SubmitP4 = CommandUtils.P4;
+				if (Parameters.Workspace != null)
+				{
+					// Create a brand new workspace
+					P4ClientInfo Client = new P4ClientInfo();
+					Client.Owner = Environment.UserName;
+					Client.Host = Environment.MachineName;
+					Client.Stream = Parameters.Stream ?? CommandUtils.P4Env.BuildRootP4;
+					Client.RootPath = Parameters.RootDir ?? CommandUtils.RootDirectory.FullName;
+					Client.Name = Parameters.Workspace;
+					Client.Options = P4ClientOption.NoAllWrite | P4ClientOption.Clobber | P4ClientOption.NoCompress | P4ClientOption.Unlocked | P4ClientOption.NoModTime | P4ClientOption.RmDir;
+					Client.LineEnd = P4LineEnd.Local;
+					CommandUtils.P4.CreateClient(Client, AllowSpew: false);
+
+					// Create a new connection for it
+					SubmitP4 = new P4Connection(Client.Owner, Client.Name);
+				}
+
+				// Get the latest version of it
+				int NewCL = SubmitP4.CreateChange(Description: Parameters.Description);
+				foreach(FileReference File in Files)
+				{
+					SubmitP4.Revert(String.Format("-k \"{0}\"", File.FullName));
+					SubmitP4.Sync(String.Format("-k \"{0}\"", File.FullName), AllowSpew: false);
+					SubmitP4.Add(NewCL, String.Format("\"{0}\"", File.FullName));
+					SubmitP4.Edit(NewCL, String.Format("\"{0}\"", File.FullName));
+					if (Parameters.FileType != null)
+					{
+						SubmitP4.P4(String.Format("reopen -t \"{0}\" \"{1}\"", Parameters.FileType, File.FullName), AllowSpew: false);
+					}
+				}
+
+				// Submit it
+				int SubmittedCL;
+				SubmitP4.Submit(NewCL, out SubmittedCL);
+				if (SubmittedCL <= 0)
+				{
+					throw new AutomationException("Submit failed.");
+				}
+				CommandUtils.Log("Submitted in changelist {0}", SubmittedCL);
+			}
+			return true;
+		}
+
+		/// <summary>
+		/// Output this task out to an XML writer.
+		/// </summary>
+		public override void Write(XmlWriter Writer)
+		{
+			Write(Writer, Parameters);
+		}
+	}
+}
diff --git a/Engine/Source/Programs/AutomationTool/Program.cs b/Engine/Source/Programs/AutomationTool/Program.cs
--- a/Engine/Source/Programs/AutomationTool/Program.cs
+++ b/Engine/Source/Programs/AutomationTool/Program.cs
@@ -15,6 +15,11 @@ namespace AutomationTool
 {
 	public class Program
 	{
+		/// <summary>
+		/// Keep a persistent reference to the delegate for handling Ctrl-C events. Since it's passed to non-managed code, we have to prevent it from being garbage collected.
+		/// </summary>
+		static ProcessManager.CtrlHandlerDelegate CtrlHandlerDelegateInstance = CtrlHandler;
+
 		[STAThread]
 		public static int Main()
 		{
@@ -45,7 +50,7 @@ namespace AutomationTool
 				var Domain = AppDomain.CurrentDomain;
 				Domain.ProcessExit += Domain_ProcessExit;
 				Domain.DomainUnload += Domain_ProcessExit;
-				HostPlatform.Current.SetConsoleCtrlHandler(CtrlHandler);
+				HostPlatform.Current.SetConsoleCtrlHandler(CtrlHandlerDelegateInstance);
 
 				var Version = AssemblyUtils.ExecutableVersion;
 				Log.TraceVerbose("{0} ver. {1}", Version.ProductName, Version.ProductVersion);
diff --git a/Engine/Source/Programs/AutomationTool/Scripts/StreamCopyDescription.Automation.cs b/Engine/Source/Programs/AutomationTool/Scripts/StreamCopyDescription.Automation.cs
--- a/Engine/Source/Programs/AutomationTool/Scripts/StreamCopyDescription.Automation.cs
+++ b/Engine/Source/Programs/AutomationTool/Scripts/StreamCopyDescription.Automation.cs
@@ -67,7 +67,9 @@ namespace AutomationTool
 			}
 
 			// Write the output file
-			string OutputFileName = Path.Combine(CommandUtils.CmdEnv.LocalRoot, "Engine", "Intermediate", "Changes.txt");
+			string OutputDirName = Path.Combine(CommandUtils.CmdEnv.LocalRoot, "Engine", "Intermediate");
+			CommandUtils.CreateDirectory(OutputDirName);
+			string OutputFileName = Path.Combine(OutputDirName, "Changes.txt");
 			using (StreamWriter Writer = new StreamWriter(OutputFileName))
 			{
 				Writer.WriteLine("Copying {0} to {1} (Source: {0} @ {2})", Stream, Target.Trim(), ChangeRecords[0].CL);
diff --git a/Engine/Source/Programs/ParallelExecutor/ParallelExecutor.Build.cs b/Engine/Source/Programs/ParallelExecutor/ParallelExecutor.Build.cs
deleted file mode 100644
diff --git a/Engine/Source/Programs/ParallelExecutor/ParallelExecutor.Target.cs b/Engine/Source/Programs/ParallelExecutor/ParallelExecutor.Target.cs
deleted file mode 100644
diff --git a/Engine/Source/Programs/ParallelExecutor/Private/BuildGraph.cpp b/Engine/Source/Programs/ParallelExecutor/Private/BuildGraph.cpp
deleted file mode 100644
diff --git a/Engine/Source/Programs/ParallelExecutor/Private/BuildGraph.h b/Engine/Source/Programs/ParallelExecutor/Private/BuildGraph.h
deleted file mode 100644
diff --git a/Engine/Source/Programs/ParallelExecutor/Private/ParallelExecutor.cpp b/Engine/Source/Programs/ParallelExecutor/Private/ParallelExecutor.cpp
deleted file mode 100644
diff --git a/Engine/Source/Programs/ParallelExecutor/Private/ParallelExecutor.h b/Engine/Source/Programs/ParallelExecutor/Private/ParallelExecutor.h
deleted file mode 100644
diff --git a/Engine/Source/Programs/UnrealGameSync/PostBadgeStatus/App.config b/Engine/Source/Programs/UnrealGameSync/PostBadgeStatus/App.config
new file mode 100644
--- /dev/null
+++ b/Engine/Source/Programs/UnrealGameSync/PostBadgeStatus/App.config
@@ -0,0 +1,6 @@
+﻿<?xml version="1.0" encoding="utf-8" ?>
+<configuration>
+    <startup> 
+        <supportedRuntime version="v4.0" sku=".NETFramework,Version=v4.5" />
+    </startup>
+</configuration>
\ No newline at end of file
diff --git a/Engine/Source/Programs/UnrealGameSync/PostBadgeStatus/PostBadgeStatus.csproj b/Engine/Source/Programs/UnrealGameSync/PostBadgeStatus/PostBadgeStatus.csproj
new file mode 100644
--- /dev/null
+++ b/Engine/Source/Programs/UnrealGameSync/PostBadgeStatus/PostBadgeStatus.csproj
@@ -0,0 +1,66 @@
+﻿<?xml version="1.0" encoding="utf-8"?>
+<Project ToolsVersion="12.0" DefaultTargets="Build" xmlns="http://schemas.microsoft.com/developer/msbuild/2003">
+  <Import Project="$(MSBuildExtensionsPath)\$(MSBuildToolsVersion)\Microsoft.Common.props" Condition="Exists('$(MSBuildExtensionsPath)\$(MSBuildToolsVersion)\Microsoft.Common.props')" />
+  <PropertyGroup>
+    <Configuration Condition=" '$(Configuration)' == '' ">Debug</Configuration>
+    <Platform Condition=" '$(Platform)' == '' ">AnyCPU</Platform>
+    <ProjectGuid>{DCBE3437-5DBF-41B4-B854-967BC0C93C48}</ProjectGuid>
+    <OutputType>Exe</OutputType>
+    <AppDesignerFolder>Properties</AppDesignerFolder>
+    <RootNamespace>PostBadgeStatus</RootNamespace>
+    <AssemblyName>PostBadgeStatus</AssemblyName>
+    <TargetFrameworkVersion>v4.5</TargetFrameworkVersion>
+    <FileAlignment>512</FileAlignment>
+  </PropertyGroup>
+  <PropertyGroup Condition="'$(Configuration)|$(Platform)' == 'Debug|AnyCPU'">
+    <DebugSymbols>true</DebugSymbols>
+    <OutputPath>bin\Debug\</OutputPath>
+    <DefineConstants>DEBUG;TRACE</DefineConstants>
+    <DebugType>full</DebugType>
+    <PlatformTarget>AnyCPU</PlatformTarget>
+    <ErrorReport>prompt</ErrorReport>
+    <CodeAnalysisRuleSet>MinimumRecommendedRules.ruleset</CodeAnalysisRuleSet>
+    <Prefer32Bit>true</Prefer32Bit>
+  </PropertyGroup>
+  <PropertyGroup Condition="'$(Configuration)|$(Platform)' == 'Release|AnyCPU'">
+    <OutputPath>bin\Release\</OutputPath>
+    <DefineConstants>TRACE</DefineConstants>
+    <Optimize>true</Optimize>
+    <DebugType>pdbonly</DebugType>
+    <PlatformTarget>AnyCPU</PlatformTarget>
+    <ErrorReport>prompt</ErrorReport>
+    <CodeAnalysisRuleSet>MinimumRecommendedRules.ruleset</CodeAnalysisRuleSet>
+    <Prefer32Bit>true</Prefer32Bit>
+  </PropertyGroup>
+  <ItemGroup>
+    <Reference Include="System" />
+    <Reference Include="System.Core" />
+    <Reference Include="System.Xml.Linq" />
+    <Reference Include="System.Data.DataSetExtensions" />
+    <Reference Include="Microsoft.CSharp" />
+    <Reference Include="System.Data" />
+    <Reference Include="System.Xml" />
+  </ItemGroup>
+  <ItemGroup>
+    <Compile Include="Program.cs" />
+    <Compile Include="Properties\AssemblyInfo.cs" />
+  </ItemGroup>
+  <ItemGroup>
+    <None Include="App.config" />
+  </ItemGroup>
+  <Choose>
+    <When Condition="Exists('NotForLicensees\ProgramSettings.cs')">
+      <ItemGroup>
+        <Compile Include="NotForLicensees\ProgramSettings.cs" />
+      </ItemGroup>
+    </When>
+  </Choose>
+  <Import Project="$(MSBuildToolsPath)\Microsoft.CSharp.targets" />
+  <!-- To modify your build process, add your task inside one of the targets below and uncomment it. 
+       Other similar extension points exist, see Microsoft.Common.targets.
+  <Target Name="BeforeBuild">
+  </Target>
+  <Target Name="AfterBuild">
+  </Target>
+  -->
+</Project>
\ No newline at end of file
diff --git a/Engine/Source/Programs/UnrealGameSync/PostBadgeStatus/Program.cs b/Engine/Source/Programs/UnrealGameSync/PostBadgeStatus/Program.cs
new file mode 100644
--- /dev/null
+++ b/Engine/Source/Programs/UnrealGameSync/PostBadgeStatus/Program.cs
@@ -0,0 +1,66 @@
+﻿// Copyright 1998-2016 Epic Games, Inc. All Rights Reserved.
+
+using System;
+using System.Collections.Generic;
+using System.Data.SqlClient;
+using System.Linq;
+using System.Text;
+using System.Threading.Tasks;
+
+namespace WriteBadgeStatus
+{
+	class Program
+	{
+		static int Main(string[] Args)
+		{
+			// Parse all the parameters
+			List<string> Arguments = new List<string>(Args);
+			string Name = ParseParam(Arguments, "Name");
+			string Change = ParseParam(Arguments, "Change");
+			string Project = ParseParam(Arguments, "Project");
+			string Database = ParseParam(Arguments, "Database");
+			string Status = ParseParam(Arguments, "Status");
+			string Url = ParseParam(Arguments, "Url");
+
+			// Check we've got all the arguments we need (and no more)
+			if (Arguments.Count > 0 || Name == null || Change == null || Project == null || Database == null || Status == null || Url == null)
+			{
+				Console.WriteLine("Syntax:");
+				Console.WriteLine("  PostBadgeStatus -Name=<Name> -Change=<CL> -Project=<DepotPath> -Database=<Connection String> -Status=<Status> -Url=<Url>");
+				return 1;
+			}
+
+			// Open the connection and add the entry
+			using(SqlConnection Connection = new SqlConnection(Database))
+			{
+				Connection.Open();
+				using (SqlCommand Command = new SqlCommand("INSERT INTO dbo.CIS (ChangeNumber, BuildType, Result, URL, Project, ArchivePath) VALUES (@ChangeNumber, @BuildType, @Result, @URL, @Project, @ArchivePath)", Connection))
+				{
+					Command.Parameters.AddWithValue("@ChangeNumber", Change);
+					Command.Parameters.AddWithValue("@BuildType", Name);
+					Command.Parameters.AddWithValue("@Result", Status);
+					Command.Parameters.AddWithValue("@URL", Url);
+					Command.Parameters.AddWithValue("@Project", Project);
+					Command.Parameters.AddWithValue("@ArchivePath", "");
+					Command.ExecuteNonQuery();
+				}
+			}
+			return 0;
+		}
+
+		static string ParseParam(List<string> Arguments, string ParamName)
+		{
+			string ParamPrefix = String.Format("-{0}=", ParamName);
+			for (int Idx = 0; Idx < Arguments.Count; Idx++)
+			{
+				if (Arguments[Idx].StartsWith(ParamPrefix, StringComparison.InvariantCultureIgnoreCase))
+				{
+					string ParamValue = Arguments[Idx].Substring(ParamPrefix.Length);
+					Arguments.RemoveAt(Idx);
+					return ParamValue;
+				}
+			}
+			return null;
+		}
+	}
+}
diff --git a/Engine/Source/Programs/UnrealGameSync/PostBadgeStatus/Properties/AssemblyInfo.cs b/Engine/Source/Programs/UnrealGameSync/PostBadgeStatus/Properties/AssemblyInfo.cs
new file mode 100644
--- /dev/null
+++ b/Engine/Source/Programs/UnrealGameSync/PostBadgeStatus/Properties/AssemblyInfo.cs
@@ -0,0 +1,38 @@
+﻿// Copyright 1998-2016 Epic Games, Inc. All Rights Reserved.
+
+using System.Reflection;
+using System.Runtime.CompilerServices;
+using System.Runtime.InteropServices;
+
+// General Information about an assembly is controlled through the following 
+// set of attributes. Change these attribute values to modify the information
+// associated with an assembly.
+[assembly: AssemblyTitle("WriteBadgeStatus")]
+[assembly: AssemblyDescription("")]
+[assembly: AssemblyConfiguration("")]
+[assembly: AssemblyCompany("")]
+[assembly: AssemblyProduct("WriteBadgeStatus")]
+[assembly: AssemblyCopyright("Copyright © Epic Games 2016")]
+[assembly: AssemblyTrademark("")]
+[assembly: AssemblyCulture("")]
+
+// Setting ComVisible to false makes the types in this assembly not visible 
+// to COM components.  If you need to access a type in this assembly from 
+// COM, set the ComVisible attribute to true on that type.
+[assembly: ComVisible(false)]
+
+// The following GUID is for the ID of the typelib if this project is exposed to COM
+[assembly: Guid("fb419b69-e06d-420c-bb0a-321faf3f99a1")]
+
+// Version information for an assembly consists of the following four values:
+//
+//      Major Version
+//      Minor Version 
+//      Build Number
+//      Revision
+//
+// You can specify all the values or you can default the Build and Revision Numbers 
+// by using the '*' as shown below:
+// [assembly: AssemblyVersion("1.0.*")]
+[assembly: AssemblyVersion("1.0.0.0")]
+[assembly: AssemblyFileVersion("1.0.0.0")]
diff --git a/Engine/Source/Programs/UnrealGameSync/PostBadgeStatus/bin/Release/PostBadgeStatus.exe.config b/Engine/Source/Programs/UnrealGameSync/PostBadgeStatus/bin/Release/PostBadgeStatus.exe.config
new file mode 100644
--- /dev/null
+++ b/Engine/Source/Programs/UnrealGameSync/PostBadgeStatus/bin/Release/PostBadgeStatus.exe.config
@@ -0,0 +1,6 @@
+﻿<?xml version="1.0" encoding="utf-8" ?>
+<configuration>
+    <startup> 
+        <supportedRuntime version="v4.0" sku=".NETFramework,Version=v4.5" />
+    </startup>
+</configuration>
\ No newline at end of file
diff --git a/Engine/Source/Programs/UnrealGameSync/PostBuildInfoTool/App.config b/Engine/Source/Programs/UnrealGameSync/PostBuildInfoTool/App.config
deleted file mode 100644
diff --git a/Engine/Source/Programs/UnrealGameSync/PostBuildInfoTool/PostBuildInfoTool.csproj b/Engine/Source/Programs/UnrealGameSync/PostBuildInfoTool/PostBuildInfoTool.csproj
deleted file mode 100644
diff --git a/Engine/Source/Programs/UnrealGameSync/PostBuildInfoTool/Program.cs b/Engine/Source/Programs/UnrealGameSync/PostBuildInfoTool/Program.cs
deleted file mode 100644
diff --git a/Engine/Source/Programs/UnrealGameSync/PostBuildInfoTool/Properties/AssemblyInfo.cs b/Engine/Source/Programs/UnrealGameSync/PostBuildInfoTool/Properties/AssemblyInfo.cs
deleted file mode 100644
diff --git a/Engine/Source/Programs/UnrealGameSync/PostBuildInfoTool/bin/Release/PostBuildInfoTool.exe.config b/Engine/Source/Programs/UnrealGameSync/PostBuildInfoTool/bin/Release/PostBuildInfoTool.exe.config
deleted file mode 100644
diff --git a/Engine/Source/Programs/UnrealGameSync/UnrealGameSync/Forms/MainWindow.cs b/Engine/Source/Programs/UnrealGameSync/UnrealGameSync/Forms/MainWindow.cs
--- a/Engine/Source/Programs/UnrealGameSync/UnrealGameSync/Forms/MainWindow.cs
+++ b/Engine/Source/Programs/UnrealGameSync/UnrealGameSync/Forms/MainWindow.cs
@@ -825,7 +825,7 @@ namespace UnrealGameSync
 				for(int ChangeIdx = 0; ChangeIdx < Changes.Count; ChangeIdx++)
 				{
 					PerforceChangeSummary Change = Changes[ChangeIdx];
-					if(String.Compare(Change.User, "buildmachine", true) != 0 || PromotedChangeNumbers.Contains(Change.Number))
+					if(ShouldShowChange(Change) || PromotedChangeNumbers.Contains(Change.Number))
 					{
 						SortedChangeNumbers.Add(Change.Number);
 
@@ -915,6 +915,15 @@ namespace UnrealGameSync
 			UpdateSyncActionCheckboxes();
 		}
 
+		bool ShouldShowChange(PerforceChangeSummary Change)
+		{
+			if(String.Compare(Change.User, "buildmachine", true) == 0 && Change.Description.IndexOf("lightmaps", StringComparison.InvariantCultureIgnoreCase) == -1)
+			{
+				return false;
+			}
+			return true;
+		}
+
 		void UpdateBuildMetadataCallback()
 		{
 			Invoke(new MethodInvoker(UpdateBuildMetadata));
diff --git a/Engine/Source/Programs/UnrealGameSync/UnrealGameSync/PerforceMonitor.cs b/Engine/Source/Programs/UnrealGameSync/UnrealGameSync/PerforceMonitor.cs
--- a/Engine/Source/Programs/UnrealGameSync/UnrealGameSync/PerforceMonitor.cs
+++ b/Engine/Source/Programs/UnrealGameSync/UnrealGameSync/PerforceMonitor.cs
@@ -223,6 +223,32 @@ namespace UnrealGameSync
 				NewestChangeNumber = Math.Min(NewestChangeNumber, NewChanges.First().Number);
 			}
 
+			// Fixup any ROBOMERGE authors
+			const string RoboMergePrefix = "#ROBOMERGE-AUTHOR:";
+			foreach(PerforceChangeSummary Change in NewChanges)
+			{
+				if(Change.Description.StartsWith(RoboMergePrefix))
+				{
+					int StartIdx = RoboMergePrefix.Length;
+					while(StartIdx < Change.Description.Length && Change.Description[StartIdx] == ' ')
+					{
+						StartIdx++;
+					}
+
+					int EndIdx = StartIdx;
+					while(EndIdx < Change.Description.Length && !Char.IsWhiteSpace(Change.Description[EndIdx]))
+					{
+						EndIdx++;
+					}
+
+					if(EndIdx > StartIdx)
+					{
+						Change.User = Change.Description.Substring(StartIdx, EndIdx - StartIdx);
+						Change.Description = "ROBOMERGE: " + Change.Description.Substring(EndIdx).TrimStart();
+					}
+				}
+			}
+
 			// Insert them into the builds list
 			if(NewChanges.Count > 0 || MaxChanges < CurrentMaxChanges)
 			{
diff --git a/Engine/Source/Programs/UnrealGameSync/UnrealGameSync/Properties/AssemblyInfo.cs b/Engine/Source/Programs/UnrealGameSync/UnrealGameSync/Properties/AssemblyInfo.cs
--- a/Engine/Source/Programs/UnrealGameSync/UnrealGameSync/Properties/AssemblyInfo.cs
+++ b/Engine/Source/Programs/UnrealGameSync/UnrealGameSync/Properties/AssemblyInfo.cs
@@ -34,5 +34,5 @@ using System.Runtime.InteropServices;
 // You can specify all the values or you can default the Build and Revision Numbers 
 // by using the '*' as shown below:
 // [assembly: AssemblyVersion("1.0.*")]
-[assembly: AssemblyVersion("1.70.*")]
-[assembly: AssemblyFileVersion("1.70.0.0")]
+[assembly: AssemblyVersion("1.72.*")]
+[assembly: AssemblyFileVersion("1.72.0.0")]
diff --git a/Engine/Source/Programs/UnrealGameSync/UnrealGameSync/UnrealGameSync.csproj b/Engine/Source/Programs/UnrealGameSync/UnrealGameSync/UnrealGameSync.csproj
--- a/Engine/Source/Programs/UnrealGameSync/UnrealGameSync/UnrealGameSync.csproj
+++ b/Engine/Source/Programs/UnrealGameSync/UnrealGameSync/UnrealGameSync.csproj
@@ -112,9 +112,7 @@
     <Compile Include="ArchiveManifest.cs" />
     <Compile Include="BuildStep.cs" />
     <Compile Include="ChildProcess.cs" />
-    <Compile Include="Controls\StatusPanel.cs">
-      <SubType>Component</SubType>
-    </Compile>
+    <Compile Include="Controls\StatusPanel.cs" />
     <Compile Include="Controls\BuildListControl.cs">
       <SubType>Component</SubType>
     </Compile>
