Copying //UE4/Dev-Editor to //UE4/Dev-Main (Source: //UE4/Dev-Editor @ 2973866)

diff --git a/Engine/Source/Programs/AutomationTool/AutomationUtils/UBTUtils.cs b/Engine/Source/Programs/AutomationTool/AutomationUtils/UBTUtils.cs
--- a/Engine/Source/Programs/AutomationTool/AutomationUtils/UBTUtils.cs
+++ b/Engine/Source/Programs/AutomationTool/AutomationUtils/UBTUtils.cs
@@ -36,7 +36,7 @@ namespace AutomationTool
 			{
 				CommandLine += " -nocreatestub";
 			}
-			CommandLine += " -NoHotReloadFromIDE";
+			CommandLine += " -NoHotReload";
 			if (bJunkDeleted || GlobalCommandLine.IgnoreJunk)
 			{
 				// UBT has already deleted junk files, make sure it doesn't do it again
diff --git a/Engine/Source/Programs/AutomationTool/Localization/LocalizationProvider.cs b/Engine/Source/Programs/AutomationTool/Localization/LocalizationProvider.cs
--- a/Engine/Source/Programs/AutomationTool/Localization/LocalizationProvider.cs
+++ b/Engine/Source/Programs/AutomationTool/Localization/LocalizationProvider.cs
@@ -131,7 +131,7 @@ public abstract class LocalizationProvider
 		}
 		else
 		{
-			throw new AutomationException("Could not find a localization provider for: '" + InLocalizationProviderId + "'");
+			BuildCommand.LogWarning("Could not find a localization provider for '{0}'", InLocalizationProviderId);
 		}
 
 		return null;
diff --git a/Engine/Source/Programs/AutomationTool/OneSkyLocalization/OneSkyLocalizationProvider.cs b/Engine/Source/Programs/AutomationTool/OneSkyLocalization/OneSkyLocalizationProvider.cs
--- a/Engine/Source/Programs/AutomationTool/OneSkyLocalization/OneSkyLocalizationProvider.cs
+++ b/Engine/Source/Programs/AutomationTool/OneSkyLocalization/OneSkyLocalizationProvider.cs
@@ -114,7 +114,7 @@ public class OneSkyLocalizationProvider : LocalizationProvider
 						using (Stream FileStream = ExportFile.OpenWrite())
 						{
 							MemoryStream.CopyTo(FileStream);
-							Console.WriteLine("[SUCCESS] Exporting: '{0}' as '{1}' ({2})", OneSkyFile.Filename, ExportFile.FullName, Culture);
+							Console.WriteLine("[SUCCESS] Exported '{0}' as '{1}' ({2})", OneSkyFile.Filename, ExportFile.FullName, Culture);
 						}
 
 						if (ExportFileWasReadOnly)
@@ -151,11 +151,11 @@ public class OneSkyLocalizationProvider : LocalizationProvider
 				}
 				else if (ExportTranslationState == UploadedFile.ExportTranslationState.NoContent)
 				{
-					Console.WriteLine("[WARNING] Exporting: '{0}' ({1}) has no translations!", OneSkyFile.Filename, Culture);
+					Console.WriteLine("[SUCCESS] Skipped '{0}' ({1}) as it has no translations", OneSkyFile.Filename, Culture);
 				}
 				else
 				{
-					Console.WriteLine("[FAILED] Exporting: '{0}' ({1})", OneSkyFile.Filename, Culture);
+					Console.WriteLine("[FAILED] Failed to get translation data for '{0}' ({1})", OneSkyFile.Filename, Culture);
 				}
 			}
 		}
diff --git a/Engine/Source/Programs/AutomationTool/Scripts/BuildPluginCommand.Automation.cs b/Engine/Source/Programs/AutomationTool/Scripts/BuildPluginCommand.Automation.cs
--- a/Engine/Source/Programs/AutomationTool/Scripts/BuildPluginCommand.Automation.cs
+++ b/Engine/Source/Programs/AutomationTool/Scripts/BuildPluginCommand.Automation.cs
@@ -21,76 +21,136 @@ class BuildPlugin : BuildCommand
 		string PluginParam = ParseParamValue("Plugin");
 		if(PluginParam == null)
 		{
-			throw new AutomationException("Plugin file name was not specified via the -plugin argument");
+			throw new AutomationException("Missing -Plugin=... argument");
 		}
 
-		// Read the plugin
+		// Check it exists
 		FileReference PluginFile = new FileReference(PluginParam);
-		DirectoryReference PluginDirectory = PluginFile.Directory;
-		PluginDescriptor Plugin = PluginDescriptor.FromFile(PluginFile);
+		if (!PluginFile.Exists())
+		{
+			throw new AutomationException("Plugin '{0}' not found", PluginFile.FullName);
+		}
 
-		// Clean the intermediate build directory
-		DirectoryReference IntermediateBuildDirectory = DirectoryReference.Combine(PluginDirectory, "Intermediate", "Build");
-		if(CommandUtils.DirectoryExists(IntermediateBuildDirectory.FullName))
+		// Get the output directory
+		string PackageParam = ParseParamValue("Package");
+		if (PackageParam == null)
 		{
-			CommandUtils.DeleteDirectory(IntermediateBuildDirectory.FullName);
+			throw new AutomationException("Missing -Package=... argument");
 		}
 
-		// Create a host project for the plugin. For script generator plugins, we need to have UHT be able to load it - and that can only happen if it's enabled in a project.
-		DirectoryReference HostProjectDirectory = DirectoryReference.Combine(new DirectoryReference(CommandUtils.CmdEnv.LocalRoot), "HostProject");
-		if (CommandUtils.DirectoryExists(HostProjectDirectory.FullName))
+		// Make sure the packaging directory is valid
+		DirectoryReference PackageDir = new DirectoryReference(PackageParam);
+		if (PluginFile.IsUnderDirectory(PackageDir))
+		{
+			throw new AutomationException("Packaged plugin output directory must be different to source");
+		}
+		if (PackageDir.IsUnderDirectory(DirectoryReference.Combine(CommandUtils.RootDirectory, "Engine")))
 		{
-			CommandUtils.DeleteDirectory(HostProjectDirectory.FullName);
+			throw new AutomationException("Output directory for packaged plugin must be outside engine directory");
 		}
 
-		DirectoryReference HostProjectPluginDirectory = DirectoryReference.Combine(HostProjectDirectory, "Plugins", PluginFile.GetFileNameWithoutExtension());
+		// Clear the output directory of existing stuff
+		if (PackageDir.Exists())
+		{
+			CommandUtils.DeleteDirectoryContents(PackageDir.FullName);
+		}
+		else
+		{
+			PackageDir.CreateDirectory();
+		}
 
-		string[] CopyPluginFiles = Directory.EnumerateFiles(PluginDirectory.FullName, "*", SearchOption.AllDirectories).ToArray();
-		foreach (string CopyPluginFile in CopyPluginFiles)
+		// Create a placeholder FilterPlugin.ini with instructions on how to use it
+		FileReference SourceFilterFile = FileReference.Combine(PluginFile.Directory, "Config", "FilterPlugin.ini");
+		if (!SourceFilterFile.Exists())
 		{
-			CommandUtils.CopyFile(CopyPluginFile, CommandUtils.MakeRerootedFilePath(CopyPluginFile, PluginDirectory.FullName, HostProjectPluginDirectory.FullName));
+			List<string> Lines = new List<string>();
+			Lines.Add("[FilterPlugin]");
+			Lines.Add("; This section lists additional files which will be packaged along with your plugin. Paths should be listed relative to the root plugin directory, and");
+			Lines.Add("; may include \"...\", \"*\", and \"?\" wildcards to match directories, files, and individual characters respectively.");
+			Lines.Add(";");
+			Lines.Add("; Examples:");
+			Lines.Add(";    /README.txt");
+			Lines.Add(";    /Extras/...");
+			Lines.Add(";    /Binaries/ThirdParty/*.dll");
+			SourceFilterFile.Directory.CreateDirectory();
+			CommandUtils.WriteAllLines_NoExceptions(SourceFilterFile.FullName, Lines.ToArray());
 		}
 
-		FileReference HostProjectPluginFile = FileReference.Combine(HostProjectPluginDirectory, PluginFile.GetFileName());
-		FileReference HostProjectFile = FileReference.Combine(HostProjectDirectory, "HostProject.uproject");
+		// Create a host project for the plugin. For script generator plugins, we need to have UHT be able to load it, which can only happen if it's enabled in a project.
+		FileReference HostProjectFile = FileReference.Combine(PackageDir, "HostProject", "HostProject.uproject");
+		FileReference HostProjectPluginFile = CreateHostProject(HostProjectFile, PluginFile);
+
+		// Read the plugin
+		CommandUtils.Log("Reading plugin from {0}...", HostProjectPluginFile);
+		PluginDescriptor Plugin = PluginDescriptor.FromFile(HostProjectPluginFile);
+
+		// Compile the plugin for all the target platforms
+		List<UnrealTargetPlatform> HostPlatforms = ParseParam("NoHostPlatform")? new List<UnrealTargetPlatform>() : new List<UnrealTargetPlatform> { BuildHostPlatform.Current.Platform };
+		List<UnrealTargetPlatform> TargetPlatforms = Rocket.RocketBuild.GetTargetPlatforms(this, BuildHostPlatform.Current.Platform).Where(x => Rocket.RocketBuild.IsCodeTargetPlatform(BuildHostPlatform.Current.Platform, x)).ToList();
+		FileReference[] BuildProducts = CompilePlugin(HostProjectFile, HostProjectPluginFile, Plugin, HostPlatforms, TargetPlatforms, "");
+
+		// Package up the final plugin data
+		PackagePlugin(HostProjectPluginFile, BuildProducts, PackageDir);
+
+		// Remove the host project
+		if(!ParseParam("NoDeleteHostProject"))
+		{
+			CommandUtils.DeleteDirectory(HostProjectFile.Directory.FullName);
+		}
+	}
+
+	FileReference CreateHostProject(FileReference HostProjectFile, FileReference PluginFile)
+	{
+		DirectoryReference HostProjectDir = HostProjectFile.Directory;
+		HostProjectDir.CreateDirectory();
+
+		// Create the new project descriptor
 		File.WriteAllText(HostProjectFile.FullName, "{ \"FileVersion\": 3, \"Plugins\": [ { \"Name\": \"" + PluginFile.GetFileNameWithoutExtension() + "\", \"Enabled\": true } ] }");
 
-		// Get any additional arguments from the commandline
-		string AdditionalArgs = "";
+		// Get the plugin directory in the host project, and copy all the files in
+		DirectoryReference HostProjectPluginDir = DirectoryReference.Combine(HostProjectDir, "Plugins", PluginFile.GetFileNameWithoutExtension());
+		CommandUtils.ThreadedCopyFiles(PluginFile.Directory.FullName, HostProjectPluginDir.FullName);
+		CommandUtils.DeleteDirectory(true, DirectoryReference.Combine(HostProjectPluginDir, "Intermediate").FullName);
 
-		// Build the host platforms
+		// Return the path to the plugin file in the host project
+		return FileReference.Combine(HostProjectPluginDir, PluginFile.GetFileName());
+	}
+
+	FileReference[] CompilePlugin(FileReference HostProjectFile, FileReference HostProjectPluginFile, PluginDescriptor Plugin, List<UnrealTargetPlatform> HostPlatforms, List<UnrealTargetPlatform> TargetPlatforms, string AdditionalArgs)
+	{
 		List<string> ReceiptFileNames = new List<string>();
-		UnrealTargetPlatform HostPlatform = BuildHostPlatform.Current.Platform;
-		if(!ParseParam("NoHostPlatform"))
+
+		// Build the host platforms
+		if(HostPlatforms.Count > 0)
 		{
-			if (Plugin.bCanBeUsedWithUnrealHeaderTool)
+			CommandUtils.Log("Building plugin for host platforms: {0}", String.Join(", ", HostPlatforms));
+			foreach (UnrealTargetPlatform HostPlatform in HostPlatforms)
 			{
-				BuildPluginWithUBT(PluginFile, Plugin, null, "UnrealHeaderTool", TargetRules.TargetType.Program, HostPlatform, UnrealTargetConfiguration.Development, ReceiptFileNames, String.Format("{0} -plugin {1}", AdditionalArgs, CommandUtils.MakePathSafeToUseWithCommandLine(HostProjectPluginFile.FullName)));
+				if (Plugin.bCanBeUsedWithUnrealHeaderTool)
+				{
+					CompilePluginWithUBT(null, HostProjectPluginFile, Plugin, "UnrealHeaderTool", TargetRules.TargetType.Program, HostPlatform, UnrealTargetConfiguration.Development, ReceiptFileNames, String.Format("{0} -plugin {1}", AdditionalArgs, CommandUtils.MakePathSafeToUseWithCommandLine(HostProjectPluginFile.FullName)));
+				}
+				CompilePluginWithUBT(HostProjectFile, HostProjectPluginFile, Plugin, "UE4Editor", TargetRules.TargetType.Editor, HostPlatform, UnrealTargetConfiguration.Development, ReceiptFileNames, AdditionalArgs);
 			}
-			BuildPluginWithUBT(PluginFile, Plugin, HostProjectFile, "UE4Editor", TargetRules.TargetType.Editor, HostPlatform, UnrealTargetConfiguration.Development, ReceiptFileNames, AdditionalArgs);
 		}
 
 		// Add the game targets
-		List<UnrealTargetPlatform> TargetPlatforms = Rocket.RocketBuild.GetTargetPlatforms(this, HostPlatform);
-		foreach(UnrealTargetPlatform TargetPlatform in TargetPlatforms)
+		if(TargetPlatforms.Count > 0)
 		{
-			if(Rocket.RocketBuild.IsCodeTargetPlatform(HostPlatform, TargetPlatform))
+			CommandUtils.Log("Building plugin for target platforms: {0}", String.Join(", ", TargetPlatforms));
+			foreach (UnrealTargetPlatform TargetPlatform in TargetPlatforms)
 			{
-				BuildPluginWithUBT(PluginFile, Plugin, HostProjectFile, "UE4Game", TargetRules.TargetType.Game, TargetPlatform, UnrealTargetConfiguration.Development, ReceiptFileNames, AdditionalArgs);
-				BuildPluginWithUBT(PluginFile, Plugin, HostProjectFile, "UE4Game", TargetRules.TargetType.Game, TargetPlatform, UnrealTargetConfiguration.Shipping, ReceiptFileNames, AdditionalArgs);
+				CompilePluginWithUBT(HostProjectFile, HostProjectPluginFile, Plugin, "UE4Game", TargetRules.TargetType.Game, TargetPlatform, UnrealTargetConfiguration.Development, ReceiptFileNames, null);
+				CompilePluginWithUBT(HostProjectFile, HostProjectPluginFile, Plugin, "UE4Game", TargetRules.TargetType.Game, TargetPlatform, UnrealTargetConfiguration.Shipping, ReceiptFileNames, null);
 			}
 		}
 
 		// Package the plugin to the output folder
-		string PackageDirectory = ParseParamValue("Package");
-		if(PackageDirectory != null)
-		{
-			List<BuildProduct> BuildProducts = GetBuildProductsFromReceipts(UnrealBuildTool.UnrealBuildTool.EngineDirectory, HostProjectDirectory, ReceiptFileNames);
-			PackagePlugin(HostProjectPluginFile, BuildProducts, PackageDirectory);
-		}
+		List<BuildProduct> BuildProducts = GetBuildProductsFromReceipts(UnrealBuildTool.UnrealBuildTool.EngineDirectory, HostProjectFile.Directory, ReceiptFileNames);
+		return BuildProducts.Select(x => new FileReference(x.Path)).ToArray();
 	}
 
-	void BuildPluginWithUBT(FileReference PluginFile, PluginDescriptor Plugin, FileReference HostProjectFile, string TargetName, TargetRules.TargetType TargetType, UnrealTargetPlatform Platform, UnrealTargetConfiguration Configuration, List<string> ReceiptFileNames, string InAdditionalArgs)
+	void CompilePluginWithUBT(FileReference HostProjectFile, FileReference HostProjectPluginFile, PluginDescriptor Plugin, string TargetName, TargetRules.TargetType TargetType, UnrealTargetPlatform Platform, UnrealTargetConfiguration Configuration, List<string> ReceiptFileNames, string InAdditionalArgs)
 	{
 		// Find a list of modules that need to be built for this plugin
 		List<string> ModuleNames = new List<string>();
@@ -115,7 +175,7 @@ class BuildPlugin : BuildCommand
 
 			string Architecture = UEBuildPlatform.GetBuildPlatform(Platform).CreateContext(HostProjectFile).GetActiveArchitecture();
 
-			string ReceiptFileName = TargetReceipt.GetDefaultPath(Path.GetDirectoryName(PluginFile.FullName), TargetName, Platform, Configuration, Architecture);
+			string ReceiptFileName = TargetReceipt.GetDefaultPath(HostProjectPluginFile.Directory.FullName, TargetName, Platform, Configuration, Architecture);
 			Arguments += String.Format(" -receipt {0}", CommandUtils.MakePathSafeToUseWithCommandLine(ReceiptFileName));
 			ReceiptFileNames.Add(ReceiptFileName);
 			
@@ -144,37 +204,33 @@ class BuildPlugin : BuildCommand
 		return BuildProducts;
 	}
 
-	static void PackagePlugin(FileReference PluginFile, List<BuildProduct> BuildProducts, string PackageDirectory)
+	static void PackagePlugin(FileReference SourcePluginFile, IEnumerable<FileReference> BuildProducts, DirectoryReference TargetDir)
 	{
-		// Clear the output directory
-		CommandUtils.DeleteDirectoryContents(PackageDirectory);
+		DirectoryReference SourcePluginDir = SourcePluginFile.Directory;
 
 		// Copy all the files to the output directory
-		List<string> MatchingFileNames = FilterPluginFiles(PluginFile.FullName, BuildProducts);
-		foreach(string MatchingFileName in MatchingFileNames)
+		FileReference[] SourceFiles = FilterPluginFiles(SourcePluginFile, BuildProducts).ToArray();
+		foreach(FileReference SourceFile in SourceFiles)
 		{
-			string SourceFileName = Path.Combine(Path.GetDirectoryName(PluginFile.FullName), MatchingFileName);
-			string TargetFileName = Path.Combine(PackageDirectory, MatchingFileName);
-			CommandUtils.CopyFile(SourceFileName, TargetFileName);
-			CommandUtils.SetFileAttributes(TargetFileName, ReadOnly: false);
+			FileReference TargetFile = FileReference.Combine(TargetDir, SourceFile.MakeRelativeTo(SourcePluginDir));
+			CommandUtils.CopyFile(SourceFile.FullName, TargetFile.FullName);
+			CommandUtils.SetFileAttributes(TargetFile.FullName, ReadOnly: false);
 		}
 
 		// Get the output plugin filename
-		FileReference TargetPluginFile = FileReference.Combine(new DirectoryReference(PackageDirectory), PluginFile.GetFileName());
+		FileReference TargetPluginFile = FileReference.Combine(TargetDir, SourcePluginFile.GetFileName());
 		PluginDescriptor NewDescriptor = PluginDescriptor.FromFile(TargetPluginFile);
 		NewDescriptor.bEnabledByDefault = true;
 		NewDescriptor.bInstalled = true;
 		NewDescriptor.Save(TargetPluginFile.FullName);
 	}
 
-	static List<string> FilterPluginFiles(string PluginFileName, List<BuildProduct> BuildProducts)
+	static IEnumerable<FileReference> FilterPluginFiles(FileReference PluginFile, IEnumerable<FileReference> BuildProducts)
 	{
-		string PluginDirectory = Path.GetDirectoryName(PluginFileName);
-
 		// Set up the default filter
 		FileFilter Filter = new FileFilter();
-		Filter.AddRuleForFile(PluginFileName, PluginDirectory, FileFilterType.Include);
-		Filter.AddRuleForFiles(BuildProducts.Select(x => x.Path), PluginDirectory, FileFilterType.Include);
+		Filter.AddRuleForFile(PluginFile.FullName, PluginFile.Directory.FullName, FileFilterType.Include);
+		Filter.AddRuleForFiles(BuildProducts, PluginFile.Directory, FileFilterType.Include);
 		Filter.Include("/Binaries/ThirdParty/...");
 		Filter.Include("/Resources/...");
 		Filter.Include("/Content/...");
@@ -182,10 +238,11 @@ class BuildPlugin : BuildCommand
 		Filter.Include("/Source/...");
 
 		// Add custom rules for each platform
-		string FilterFileName = Path.Combine(Path.GetDirectoryName(PluginFileName), "Config", "FilterPlugin.ini");
-		if(File.Exists(FilterFileName))
+		FileReference FilterFile = FileReference.Combine(PluginFile.Directory, "Config", "FilterPlugin.ini");
+		if(FilterFile.Exists())
 		{
-			Filter.ReadRulesFromFile(FilterFileName, "FilterPlugin");
+			CommandUtils.Log("Reading filter rules from {0}", FilterFile);
+			Filter.ReadRulesFromFile(FilterFile.FullName, "FilterPlugin");
 		}
 
 		// Apply the standard exclusion rules
@@ -193,6 +250,6 @@ class BuildPlugin : BuildCommand
 		Filter.ExcludeConfidentialPlatforms();
 
 		// Apply the filter to the plugin directory
-		return Filter.ApplyToDirectory(PluginDirectory, true);
+		return Filter.ApplyToDirectory(PluginFile.Directory, true);
 	}
 }
diff --git a/Engine/Source/Programs/AutomationTool/Scripts/Localisation.Automation.cs b/Engine/Source/Programs/AutomationTool/Scripts/Localisation.Automation.cs
--- a/Engine/Source/Programs/AutomationTool/Scripts/Localisation.Automation.cs
+++ b/Engine/Source/Programs/AutomationTool/Scripts/Localisation.Automation.cs
@@ -86,10 +86,13 @@ class Localise : BuildCommand
 			ProjectInfos.Add(GenerateProjectInfo(RootWorkingDirectory, ProjectName));
 		}
 
-		// Export all text from our localization provider
-		foreach (var ProjectInfo in ProjectInfos)
+		if (LocProvider != null)
 		{
-			LocProvider.DownloadProjectFromLocalizationProvider(ProjectInfo.ProjectName, ProjectInfo.ImportInfo);
+			// Export all text from our localization provider
+			foreach (var ProjectInfo in ProjectInfos)
+			{
+				LocProvider.DownloadProjectFromLocalizationProvider(ProjectInfo.ProjectName, ProjectInfo.ImportInfo);
+			}
 		}
 
 		// Setup editor arguments for SCC.
@@ -141,8 +144,8 @@ class Localise : BuildCommand
 		{
 			LogWarning("Skipping upload to the localization provider due to an earlier commandlet failure.");
 		}
-		else
-		{
+		else if (LocProvider != null)
+        {
 			// Upload all text to our localization provider
 			foreach (var ProjectInfo in ProjectInfos)
 			{
diff --git a/Engine/Source/Programs/UnrealBuildTool/Configuration/UEBuildTarget.cs b/Engine/Source/Programs/UnrealBuildTool/Configuration/UEBuildTarget.cs
--- a/Engine/Source/Programs/UnrealBuildTool/Configuration/UEBuildTarget.cs
+++ b/Engine/Source/Programs/UnrealBuildTool/Configuration/UEBuildTarget.cs
@@ -2773,17 +2773,21 @@ namespace UnrealBuildTool
 		{
 			List<UEBuildBinary> Result = new List<UEBuildBinary>();
 
-			foreach (UEBuildBinary DLLBinary in Binaries)
+			foreach (UEBuildBinary Binary in Binaries)
 			{
-				OnlyModule FoundOnlyModule = DLLBinary.FindOnlyModule(OnlyModules);
-				if (FoundOnlyModule != null)
+				// If we're doing an OnlyModule compile, we never want the executable that static libraries are linked into for monolithic builds
+				if(Binary.Config.Type != UEBuildBinaryType.Executable)
 				{
-					Result.Add(DLLBinary);
-
-					if (!String.IsNullOrEmpty(FoundOnlyModule.OnlyModuleSuffix))
+					OnlyModule FoundOnlyModule = Binary.FindOnlyModule(OnlyModules);
+					if (FoundOnlyModule != null)
 					{
-						DLLBinary.Config.OriginalOutputFilePaths = DLLBinary.Config.OutputFilePaths;
-						DLLBinary.Config.OutputFilePaths = DLLBinary.Config.OutputFilePaths.Select(Path => AddModuleFilenameSuffix(FoundOnlyModule.OnlyModuleName, Path, FoundOnlyModule.OnlyModuleSuffix)).ToList();
+						Result.Add(Binary);
+
+						if (!String.IsNullOrEmpty(FoundOnlyModule.OnlyModuleSuffix))
+						{
+							Binary.Config.OriginalOutputFilePaths = Binary.Config.OutputFilePaths;
+							Binary.Config.OutputFilePaths = Binary.Config.OutputFilePaths.Select(Path => AddModuleFilenameSuffix(FoundOnlyModule.OnlyModuleName, Path, FoundOnlyModule.OnlyModuleSuffix)).ToList();
+						}
 					}
 				}
 			}
diff --git a/Engine/Source/Programs/UnrealBuildTool/System/UnrealBuildTool.cs b/Engine/Source/Programs/UnrealBuildTool/System/UnrealBuildTool.cs
--- a/Engine/Source/Programs/UnrealBuildTool/System/UnrealBuildTool.cs
+++ b/Engine/Source/Programs/UnrealBuildTool/System/UnrealBuildTool.cs
@@ -1648,8 +1648,9 @@ namespace UnrealBuildTool
 					return ECompilationResult.Succeeded;
 				}
 
-				UEBuildConfiguration.bHotReloadFromIDE = UEBuildConfiguration.bAllowHotReloadFromIDE && TargetDescs.Count == 1 && !TargetDescs[0].bIsEditorRecompile && ShouldDoHotReloadFromIDE(TargetDescs[0]);
-				bool bIsHotReload = UEBuildConfiguration.bHotReloadFromIDE || (TargetDescs.Count == 1 && TargetDescs[0].OnlyModules.Count > 0 && TargetDescs[0].ForeignPlugins.Count == 0);
+				bool bNoHotReload = UnrealBuildTool.CommandLineContains("-NoHotReload");
+				UEBuildConfiguration.bHotReloadFromIDE = !bNoHotReload && UEBuildConfiguration.bAllowHotReloadFromIDE && TargetDescs.Count == 1 && !TargetDescs[0].bIsEditorRecompile && ShouldDoHotReloadFromIDE(TargetDescs[0]);
+				bool bIsHotReload = !bNoHotReload && (UEBuildConfiguration.bHotReloadFromIDE || (TargetDescs.Count == 1 && TargetDescs[0].OnlyModules.Count > 0 && TargetDescs[0].ForeignPlugins.Count == 0));
 				TargetDescriptor HotReloadTargetDesc = bIsHotReload ? TargetDescs[0] : null;
 
 				if (ProjectFileGenerator.bGenerateProjectFiles)
diff --git a/Engine/Source/Runtime/AssetRegistry/Private/AssetRegistry.cpp b/Engine/Source/Runtime/AssetRegistry/Private/AssetRegistry.cpp
--- a/Engine/Source/Runtime/AssetRegistry/Private/AssetRegistry.cpp
+++ b/Engine/Source/Runtime/AssetRegistry/Private/AssetRegistry.cpp
@@ -2275,12 +2275,12 @@ void FAssetRegistry::AddAssetData(FAssetData* AssetData)
 	// Populate the class map if adding blueprint
 	if (ClassGeneratorNames.Contains(AssetData->AssetClass))
 	{
-		auto GeneratedClassPtr = AssetData->TagsAndValues.Find("GeneratedClass");
-		auto ParentClassPtr = AssetData->TagsAndValues.Find("ParentClass");
-		if ( GeneratedClassPtr && ParentClassPtr )
+		const FString GeneratedClass = AssetData->GetTagValueRef<FString>("GeneratedClass");
+		const FString ParentClass = AssetData->GetTagValueRef<FString>("ParentClass");
+		if ( !GeneratedClass.IsEmpty() && !ParentClass.IsEmpty() )
 		{
-			const FName GeneratedClassFName = FName(*ExportTextPathToObjectName(*GeneratedClassPtr));
-			const FName ParentClassFName = FName(*ExportTextPathToObjectName(*ParentClassPtr));
+			const FName GeneratedClassFName = *ExportTextPathToObjectName(GeneratedClass);
+			const FName ParentClassFName = *ExportTextPathToObjectName(ParentClass);
 			CachedInheritanceMap.Add(GeneratedClassFName, ParentClassFName);
 		}
 	}
@@ -2364,20 +2364,19 @@ void FAssetRegistry::UpdateAssetData(FAssetData* AssetData, const FAssetData& Ne
 	// Update the class map if updating a blueprint
 	if (ClassGeneratorNames.Contains(AssetData->AssetClass))
 	{
-		auto OldGeneratedClassPtr = AssetData->TagsAndValues.Find("GeneratedClass");
-
-		if ( OldGeneratedClassPtr )
+		const FString OldGeneratedClass = AssetData->GetTagValueRef<FString>("GeneratedClass");
+		if ( !OldGeneratedClass.IsEmpty() )
 		{
-			const FName OldGeneratedClassFName = FName(*ExportTextPathToObjectName(*OldGeneratedClassPtr));
+			const FName OldGeneratedClassFName = *ExportTextPathToObjectName(OldGeneratedClass);
 			CachedInheritanceMap.Remove(OldGeneratedClassFName);
 		}
 
-		const FString* NewGeneratedClassPtr = NewAssetData.TagsAndValues.Find("GeneratedClass");
-		const FString* NewParentClassPtr = NewAssetData.TagsAndValues.Find("ParentClass");
-		if ( NewGeneratedClassPtr && NewParentClassPtr )
+		const FString NewGeneratedClass = NewAssetData.GetTagValueRef<FString>("GeneratedClass");
+		const FString NewParentClass = NewAssetData.GetTagValueRef<FString>("ParentClass");
+		if ( !NewGeneratedClass.IsEmpty() && !NewParentClass.IsEmpty() )
 		{
-			const FName NewGeneratedClassFName = FName(*ExportTextPathToObjectName(*NewGeneratedClassPtr));
-			const FName NewParentClassFName = FName(*ExportTextPathToObjectName(*NewParentClassPtr));
+			const FName NewGeneratedClassFName = *ExportTextPathToObjectName(*NewGeneratedClass);
+			const FName NewParentClassFName = *ExportTextPathToObjectName(*NewParentClass);
 			CachedInheritanceMap.Add(NewGeneratedClassFName, NewParentClassFName);
 		}
 	}
@@ -2398,11 +2397,10 @@ bool FAssetRegistry::RemoveAssetData(FAssetData* AssetData)
 		// Remove from the class map if removing a blueprint
 		if (ClassGeneratorNames.Contains(AssetData->AssetClass))
 		{
-			auto OldGeneratedClassPtr = AssetData->TagsAndValues.Find("GeneratedClass");
-
-			if ( OldGeneratedClassPtr )
+			const FString OldGeneratedClass = AssetData->GetTagValueRef<FString>("GeneratedClass");
+			if ( !OldGeneratedClass.IsEmpty() )
 			{
-				const FName OldGeneratedClassFName = FName(*ExportTextPathToObjectName(*OldGeneratedClassPtr));
+				const FName OldGeneratedClassFName = *ExportTextPathToObjectName(OldGeneratedClass);
 				CachedInheritanceMap.Remove(OldGeneratedClassFName);
 			}
 		}
diff --git a/Engine/Source/Runtime/AssetRegistry/Public/AssetData.h b/Engine/Source/Runtime/AssetRegistry/Public/AssetData.h
--- a/Engine/Source/Runtime/AssetRegistry/Public/AssetData.h
+++ b/Engine/Source/Runtime/AssetRegistry/Public/AssetData.h
@@ -241,6 +241,14 @@ public:
 		return Package;
 	}
 
+	/** Try and get the value associated with the given tag as a type converted value */
+	template <typename ValueType>
+	bool GetTagValue(const FName InTagName, ValueType& OutTagValue) const;
+
+	/** Try and get the value associated with the given tag as a type converted value, or an empty value if it doesn't exist */
+	template <typename ValueType>
+	ValueType GetTagValueRef(const FName InTagName) const;
+
 	/** Returns true if the asset is loaded */
 	bool IsAssetLoaded() const
 	{
@@ -332,4 +340,121 @@ public:
 
 		return Ar;
 	}
+
+private:
+	bool GetTagValueStringImpl(const FName InTagName, FString& OutTagValue) const
+	{
+		if (const FString* FoundValue = TagsAndValues.Find(InTagName))
+		{
+			bool bIsHandled = false;
+			if (FTextStringHelper::IsComplexText(**FoundValue))
+			{
+				FText TmpText;
+				if (FTextStringHelper::ReadFromString(**FoundValue, TmpText))
+				{
+					bIsHandled = true;
+					OutTagValue = TmpText.ToString();
+				}
+			}
+
+			if (!bIsHandled)
+			{
+				OutTagValue = *FoundValue;
+			}
+
+			return true;
+		}
+
+		return false;
+	}
+
+	bool GetTagValueTextImpl(const FName InTagName, FText& OutTagValue) const
+	{
+		if (const FString* FoundValue = TagsAndValues.Find(InTagName))
+		{
+			if (!FTextStringHelper::ReadFromString(**FoundValue, OutTagValue))
+			{
+				OutTagValue = FText::FromString(*FoundValue);
+			}
+			return true;
+		}
+		return false;
+	}
+
+	bool GetTagValueNameImpl(const FName InTagName, FName& OutTagValue) const
+	{
+		FString StrValue;
+		if (GetTagValueStringImpl(InTagName, StrValue))
+		{
+			OutTagValue = *StrValue;
+			return true;
+		}
+		return false;
+	}
 };
+
+template <typename ValueType>
+inline bool FAssetData::GetTagValue(const FName InTagName, ValueType& OutTagValue) const
+{
+	if (const FString* FoundValue = TagsAndValues.Find(InTagName))
+	{
+		FMemory::Memzero(&OutTagValue, sizeof(ValueType));
+		LexicalConversion::FromString(OutTagValue, **FoundValue);
+		return true;
+	}
+	return false;
+}
+
+template <>
+inline bool FAssetData::GetTagValue<FString>(const FName InTagName, FString& OutTagValue) const
+{
+	return GetTagValueStringImpl(InTagName, OutTagValue);
+}
+
+template <>
+inline bool FAssetData::GetTagValue<FText>(const FName InTagName, FText& OutTagValue) const
+{
+	return GetTagValueTextImpl(InTagName, OutTagValue);
+}
+
+template <>
+inline bool FAssetData::GetTagValue<FName>(const FName InTagName, FName& OutTagValue) const
+{
+	return GetTagValueNameImpl(InTagName, OutTagValue);
+}
+
+template <typename ValueType>
+inline ValueType FAssetData::GetTagValueRef(const FName InTagName) const
+{
+	ValueType TmpValue;
+	FMemory::Memzero(&TmpValue, sizeof(ValueType));
+	if (const FString* FoundValue = TagsAndValues.Find(InTagName))
+	{
+		LexicalConversion::FromString(TmpValue, **FoundValue);
+	}
+	return TmpValue;
+}
+
+template <>
+inline FString FAssetData::GetTagValueRef<FString>(const FName InTagName) const
+{
+	FString TmpValue;
+	GetTagValueStringImpl(InTagName, TmpValue);
+	return TmpValue;
+}
+
+template <>
+inline FText FAssetData::GetTagValueRef<FText>(const FName InTagName) const
+{
+	FText TmpValue;
+	GetTagValueTextImpl(InTagName, TmpValue);
+	return TmpValue;
+}
+
+template <>
+inline FName FAssetData::GetTagValueRef<FName>(const FName InTagName) const
+{
+	FName TmpValue;
+	GetTagValueNameImpl(InTagName, TmpValue);
+	return TmpValue;
+}
diff --git a/Engine/Source/Runtime/Core/Private/Internationalization/Text.cpp b/Engine/Source/Runtime/Core/Private/Internationalization/Text.cpp
--- a/Engine/Source/Runtime/Core/Private/Internationalization/Text.cpp
+++ b/Engine/Source/Runtime/Core/Private/Internationalization/Text.cpp
@@ -1461,15 +1461,17 @@ bool TextBiDi::IsControlCharacter(const TCHAR InChar)
 		|| InChar == TEXT('\u2069'); // POP DIRECTIONAL ISOLATE
 }
 
+#define LOC_DEFINE_REGION
+const FString FTextStringHelper::InvTextMarker = TEXT("INVTEXT");
+const FString FTextStringHelper::NsLocTextMarker = TEXT("NSLOCTEXT");
+const FString FTextStringHelper::LocTextMarker = TEXT("LOCTEXT");
+#undef LOC_DEFINE_REGION
+
 bool FTextStringHelper::ReadFromString_ComplexText(const TCHAR* Buffer, FText& OutValue, const TCHAR* Namespace, int32* OutNumCharsRead)
 {
 #define LOC_DEFINE_REGION
 	const TCHAR* const Start = Buffer;
 
-	static const FString InvTextMarker = TEXT("INVTEXT");
-	static const FString NsLocTextMarker = TEXT("NSLOCTEXT");
-	static const FString LocTextMarker = TEXT("LOCTEXT");
-
 	auto ExtractQuotedString = [&](FString& OutStr) -> const TCHAR*
 	{
 		int32 CharsRead = 0;
@@ -1743,4 +1745,11 @@ bool FTextStringHelper::WriteToString(FString& Buffer, const FText& Value, const
 	return true;
 }
 
+bool FTextStringHelper::IsComplexText(const TCHAR* Buffer)
+{
+#define LOC_DEFINE_REGION
+	return FCString::Strstr(Buffer, *InvTextMarker) || FCString::Strstr(Buffer, *NsLocTextMarker) || FCString::Strstr(Buffer, *LocTextMarker);
+#undef LOC_DEFINE_REGION
+}
+
 #undef LOCTEXT_NAMESPACE
diff --git a/Engine/Source/Runtime/Core/Private/Mac/MacApplication.cpp b/Engine/Source/Runtime/Core/Private/Mac/MacApplication.cpp
--- a/Engine/Source/Runtime/Core/Private/Mac/MacApplication.cpp
+++ b/Engine/Source/Runtime/Core/Private/Mac/MacApplication.cpp
@@ -638,7 +638,23 @@ void FMacApplication::ProcessEvent(const FDeferredMacEvent& Event)
 			const bool bHaveText = [[Event.DraggingPasteboard types] containsObject:NSPasteboardTypeString];
 			const bool bHaveFiles = [[Event.DraggingPasteboard types] containsObject:NSFilenamesPboardType];
 
-			if (bHaveFiles)
+			if (bHaveFiles && bHaveText)
+			{
+				TArray<FString> FileList;
+
+				NSArray *Files = [Event.DraggingPasteboard propertyListForType:NSFilenamesPboardType];
+				for (int32 Index = 0; Index < [Files count]; Index++)
+				{
+					NSString* FilePath = [Files objectAtIndex: Index];
+					const FString ListElement = FString([FilePath fileSystemRepresentation]);
+					FileList.Add(ListElement);
+				}
+
+				NSString* Text = [Event.DraggingPasteboard stringForType:NSPasteboardTypeString];
+
+				MessageHandler->OnDragEnterExternal(EventWindow.ToSharedRef(), FString(Text), FileList);
+			}
+			else if (bHaveFiles)
 			{
 				TArray<FString> FileList;
 
diff --git a/Engine/Source/Runtime/Core/Private/Misc/AutomationTest.cpp b/Engine/Source/Runtime/Core/Private/Misc/AutomationTest.cpp
--- a/Engine/Source/Runtime/Core/Private/Misc/AutomationTest.cpp
+++ b/Engine/Source/Runtime/Core/Private/Misc/AutomationTest.cpp
@@ -6,6 +6,43 @@
 
 DEFINE_LOG_CATEGORY_STATIC(LogAutomationTest, Warning, All);
 
+FString FAutomationTestInfo::GetTestAsString() const
+{
+	return FString::Printf(TEXT("%s,%s,%i,%d,%s,%s,%i"), *DisplayName, *TestName, TestFlags, NumParticipantsRequired, *TestParameter, *SourceFile, SourceFileLine);
+}
+
+void FAutomationTestInfo::ParseStringInfo(const FString& InTestInfo)
+{
+	//Split New Test Name into string array
+	TArray<FString> Pieces;
+	InTestInfo.ParseIntoArray(Pieces, TEXT(","), false);
+
+	// We should always have at least 3 parameters
+	check(Pieces.Num() >= 3)
+
+	DisplayName = Pieces[0];
+	TestName = Pieces[1];
+	TestFlags = uint8(FCString::Atoi(*Pieces[2]));
+
+	NumParticipantsRequired = FCString::Atoi(*Pieces[3]);
+
+	// Optional Parameters
+	if ( Pieces.Num() >= 5 )
+	{
+		TestParameter = Pieces[4];
+	}
+
+	if ( Pieces.Num() >= 6 )
+	{
+		SourceFile = Pieces[5];
+	}
+
+	if ( Pieces.Num() >= 7 )
+	{
+		SourceFileLine = FCString::Atoi(*Pieces[6]);
+	}
+}
+
 
 void FAutomationTestFramework::FAutomationTestFeedbackContext::Serialize( const TCHAR* V, ELogVerbosity::Type Verbosity, const class FName& Category )
 {
@@ -774,7 +811,7 @@ void FAutomationTestBase::GenerateTestNames(TArray<FAutomationTestInfo>& TestInf
 		}
 
 		// Add the test info to our collection
-		FAutomationTestInfo NewTestInfo( CompleteBeautifiedNames, CompleteTestName, GetTestFlags(), GetRequiredDeviceNum(), ParameterNames[ParameterIndex] );
+		FAutomationTestInfo NewTestInfo(CompleteBeautifiedNames, CompleteTestName, GetTestFlags(), GetRequiredDeviceNum(), ParameterNames[ParameterIndex], GetFileName(), GetFileLine());
 		
 		TestInfo.Add( NewTestInfo );
 	}
diff --git a/Engine/Source/Runtime/Core/Private/Misc/ConfigCacheIni.cpp b/Engine/Source/Runtime/Core/Private/Misc/ConfigCacheIni.cpp
--- a/Engine/Source/Runtime/Core/Private/Misc/ConfigCacheIni.cpp
+++ b/Engine/Source/Runtime/Core/Private/Misc/ConfigCacheIni.cpp
@@ -531,13 +531,73 @@ void FConfigFile::Read( const FString& Filename )
 	}
 }
 
-bool FConfigFile::ShouldExportQuotedString(const FString& PropertyValue) const
+bool FConfigFile::ShouldExportQuotedString(const FString& PropertyValue)
 {
-	// The value should be exported as quoted string if it begins with a space (which is stripped on import) or
-	// when it contains '//' (interpreted as a comment when importing).
-	return **PropertyValue == TEXT(' ') || FCString::Strstr(*PropertyValue, TEXT("//")) != nullptr;
+	bool bEscapeNextChar = false;
+	bool bIsWithinQuotes = false;
+
+	// The value should be exported as quoted string if...
+	const TCHAR* const DataPtr = *PropertyValue;
+	for (const TCHAR* CharPtr = DataPtr; *CharPtr; ++CharPtr)
+	{
+		const TCHAR ThisChar = *CharPtr;
+		const TCHAR NextChar = *(CharPtr + 1);
+
+		const bool bIsFirstChar = CharPtr == DataPtr;
+		const bool bIsLastChar = NextChar == 0;
+
+		if (ThisChar == TEXT('"') && !bEscapeNextChar)
+		{
+			bIsWithinQuotes = !bIsWithinQuotes;
+		}
+		bEscapeNextChar = ThisChar == TEXT('\\') && bIsWithinQuotes && !bEscapeNextChar;
+
+		// ... it begins or ends with a space (which is stripped on import)
+		if (ThisChar == TEXT(' ') && (bIsFirstChar || bIsLastChar))
+		{
+			return true;
+		}
+
+		// ... it begins with a '"' (which would be treated as a quoted string)
+		if (ThisChar == TEXT('"') && bIsFirstChar)
+		{
+			return true;
+		}
+
+		// ... it ends with a '\' (which would be treated as a line extension)
+		if (ThisChar == TEXT('\\') && bIsLastChar)
+		{
+			return true;
+		}
+
+		// ... it contains unquoted '{' or '}' (which are stripped on import)
+		if ((ThisChar == TEXT('{') || ThisChar == TEXT('}')) && !bIsWithinQuotes)
+		{
+			return true;
+		}
+		
+		// ... it contains unquoted '//' (interpreted as a comment when importing)
+		if ((ThisChar == TEXT('/') && NextChar == TEXT('/')) && !bIsWithinQuotes)
+		{
+			return true;
+		}
+	}
+
+	return false;
 }
 
+FString FConfigFile::GenerateExportedPropertyLine(const FString& PropertyName, const FString& PropertyValue)
+{
+	const bool bShouldQuote = ShouldExportQuotedString(PropertyValue);
+	if (bShouldQuote)
+	{
+		return FString::Printf(TEXT("%s=\"%s\"") LINE_TERMINATOR, *PropertyName, *PropertyValue.ReplaceCharWithEscapedChar());
+	}
+	else
+	{
+		return FString::Printf(TEXT("%s=%s") LINE_TERMINATOR, *PropertyName, *PropertyValue);
+	}
+}
 
 #if ALLOW_INI_OVERRIDE_FROM_COMMANDLINE
 
@@ -905,13 +965,6 @@ bool FConfigFile::Write( const FString& Filename, bool bDoRemoteWrite/* = true*/
 						bWroteASectionProperty = true;
 					}
 
-					// Print the property to the file.
-					TCHAR QuoteString[2] = {0,0};
-					if (ShouldExportQuotedString(PropertyValue))
-					{
-						QuoteString[0] = TEXT('\"');
-					}
-
 					// Write out our property, if it is an array we need to write out the entire array.
 					TArray<FConfigValue> CompletePropertyToWrite;
 					Section.MultiFind( PropertyName, CompletePropertyToWrite, true );
@@ -924,8 +977,7 @@ bool FConfigFile::Write( const FString& Filename, bool bDoRemoteWrite/* = true*/
 					{
 						for (const FConfigValue& ConfigValue : CompletePropertyToWrite)
 						{
-							Text += FString::Printf( TEXT("%s=%s%s%s") LINE_TERMINATOR, 
-								*PropertyName.ToString(), QuoteString, *ConfigValue.GetSavedValue(), QuoteString);	
+							Text += GenerateExportedPropertyLine(PropertyName.ToString(), ConfigValue.GetSavedValue());
 						}
 					}
 
@@ -1165,16 +1217,7 @@ void FConfigFile::SaveSourceToBackupFile()
 		{
 			const FName PropertyName = PropertyIterator.Key();
 			const FString& PropertyValue = PropertyIterator.Value().GetSavedValue();
-
-			// Print the property to the file.
-			TCHAR QuoteString[2] = {0,0};
-			if (ShouldExportQuotedString(PropertyValue))
-			{
-				QuoteString[0] = TEXT('\"');
-			}
-
-			Text += FString::Printf( TEXT("%s=%s%s%s") LINE_TERMINATOR, 
-				*PropertyName.ToString(), QuoteString, *PropertyValue, QuoteString);	
+			Text += FConfigFile::GenerateExportedPropertyLine(PropertyName.ToString(), PropertyValue);
 		}
 		Text += LINE_TERMINATOR;
 	}
@@ -1250,7 +1293,7 @@ void FConfigFile::ProcessPropertyAndWriteForDefaults( const TArray< FConfigValue
 			for (const FString& NextElement : ArrayProperties)
 			{
 				FString PropertyNameWithRemoveOp = PropertyName.Replace(TEXT("+"), TEXT("-"));
-				OutText += FString::Printf(TEXT("%s=%s") LINE_TERMINATOR, *PropertyNameWithRemoveOp, *NextElement);
+				OutText += GenerateExportedPropertyLine(PropertyNameWithRemoveOp, NextElement);
 			}
 		}
 	}
@@ -1258,14 +1301,7 @@ void FConfigFile::ProcessPropertyAndWriteForDefaults( const TArray< FConfigValue
 	// Write the properties out to a file.
 	for ( const FConfigValue& PropertyIt : InCompletePropertyToProcess)
 	{
-		FString PropertyValue = PropertyIt.GetSavedValue();
-
-		if (ShouldExportQuotedString(PropertyValue))
-		{
-			PropertyValue = FString::Printf(TEXT("\"%s\""), *PropertyValue);
-		}
-
-		OutText += FString::Printf(TEXT("%s=%s") LINE_TERMINATOR, *PropertyName, *PropertyValue);
+		OutText += GenerateExportedPropertyLine(PropertyName, PropertyIt.GetSavedValue());
 	}
 }
 
@@ -1708,7 +1744,7 @@ void FConfigCacheIni::SetString( const TCHAR* Section, const TCHAR* Key, const T
 		Sec->Add( Key, Value );
 		File->Dirty = true;
 	}
-	else if( FCString::Stricmp(*ConfigValue->GetSavedValue(),Value)!=0 )
+	else if( FCString::Strcmp(*ConfigValue->GetSavedValue(),Value)!=0 )
 	{
 		File->Dirty = true;
 		*ConfigValue = FConfigValue(Value);
@@ -1739,7 +1775,7 @@ void FConfigCacheIni::SetText( const TCHAR* Section, const TCHAR* Key, const FTe
 		Sec->Add( Key, StrValue );
 		File->Dirty = true;
 	}
-	else if( FCString::Stricmp(*ConfigValue->GetSavedValue(), *StrValue)!=0 )
+	else if( FCString::Strcmp(*ConfigValue->GetSavedValue(), *StrValue)!=0 )
 	{
 		File->Dirty = true;
 		*ConfigValue = FConfigValue(StrValue);
@@ -3170,8 +3206,7 @@ private:
 			{
 				if (SectionLine.StartsWith(FString::Printf(TEXT("%s="), *PropertyName)))
 				{
-					UpdatedSection += FString::Printf(TEXT("%s=%s"), *PropertyName, *PropertyValue);
-					UpdatedSection += LINE_TERMINATOR;
+					UpdatedSection += FConfigFile::GenerateExportedPropertyLine(PropertyName, PropertyValue);
 					bWrotePropertyOnPass = true;
 				}
 				else
@@ -3251,8 +3286,7 @@ private:
 		// Make sure we dont leave much whitespace, and append the property name/value entry
 		ClearTrailingWhitespace(PreText);
 		PreText += LINE_TERMINATOR;
-		PreText += FString::Printf(TEXT("%s=%s"), *PropertyName, *PropertyValue);
-		PreText += LINE_TERMINATOR;
+		PreText += FConfigFile::GenerateExportedPropertyLine(PropertyName, PropertyValue);
 		PreText += LINE_TERMINATOR;
 	}
 
diff --git a/Engine/Source/Runtime/Core/Private/Misc/FbxErrors.cpp b/Engine/Source/Runtime/Core/Private/Misc/FbxErrors.cpp
--- a/Engine/Source/Runtime/Core/Private/Misc/FbxErrors.cpp
+++ b/Engine/Source/Runtime/Core/Private/Misc/FbxErrors.cpp
@@ -42,6 +42,8 @@ FName FFbxErrors::Generic_Mesh_LOD_NoFileSelected(TEXT("Generic_Mesh_LOD_NoFileS
 
 FName FFbxErrors::Generic_Mesh_LOD_MultipleFilesSelected(TEXT("Generic_Mesh_LOD_MultipleFilesSelected"));
 
+FName FFbxErrors::Generic_Mesh_SkinxxNameError(TEXT("Generic_Mesh_SkinxxNameError"));
+
 /** Static Mesh **/
 FName FFbxErrors::StaticMesh_TooManyMaterials(TEXT("StaticMesh_TooManyMaterials"));
 
diff --git a/Engine/Source/Runtime/Core/Private/Misc/OutputDevice.cpp b/Engine/Source/Runtime/Core/Private/Misc/OutputDevice.cpp
--- a/Engine/Source/Runtime/Core/Private/Misc/OutputDevice.cpp
+++ b/Engine/Source/Runtime/Core/Private/Misc/OutputDevice.cpp
@@ -257,6 +257,9 @@ void FDebug::EnsureFailed(const ANSICHAR* Expr, const ANSICHAR* File, int32 Line
 	// Is there a debugger attached?  If not we'll submit an error report.
 	if (FPlatformMisc::IsDebuggerPresent())
 	{
+#if !NO_LOGGING
+		UE_LOG(LogOutputDevice, Error, TEXT("%s [File:%s] [Line: %i]"), ErrorString, ANSI_TO_TCHAR(File), Line);
+#endif
 		return;
 	}
 
@@ -294,7 +297,7 @@ void FDebug::EnsureFailed(const ANSICHAR* Expr, const ANSICHAR* File, int32 Line
 
 		// Dump the error and flush the log.
 #if !NO_LOGGING
-		FDebug::OutputMultiLineCallstack(__FILE__, __LINE__, LogOutputDevice.GetCategoryName(), TEXT("=== Handled ensure: ==="), ErrorMsg, ELogVerbosity::Warning);
+		FDebug::OutputMultiLineCallstack(__FILE__, __LINE__, LogOutputDevice.GetCategoryName(), TEXT("=== Handled ensure: ==="), ErrorMsg, ELogVerbosity::Error);
 #endif
 		GLog->Flush();
 
@@ -351,6 +354,11 @@ void FDebug::EnsureFailed(const ANSICHAR* Expr, const ANSICHAR* File, int32 Line
 		// If we fail to generate the string to identify the crash we don't know if we should skip sending the report,
 		// so we will just send the report anyway.
 		bShouldSendNewReport = true;
+
+		// Add message to log even without stacktrace. It is useful for testing fail on ensure.
+#if !NO_LOGGING
+		UE_LOG(LogOutputDevice, Error, TEXT("%s [File:%s] [Line: %i]"), ErrorString, ANSI_TO_TCHAR(File), Line);
+#endif
 	}
 
 	if ( bShouldSendNewReport )
diff --git a/Engine/Source/Runtime/Core/Private/Tests/Internationalization/TextTest.cpp b/Engine/Source/Runtime/Core/Private/Tests/Internationalization/TextTest.cpp
--- a/Engine/Source/Runtime/Core/Private/Tests/Internationalization/TextTest.cpp
+++ b/Engine/Source/Runtime/Core/Private/Tests/Internationalization/TextTest.cpp
@@ -480,7 +480,7 @@ bool FTextTest::RunTest (const FString& Parameters)
 				// When changes are made to FormattedTestLayer2, please pull out the newly translated LEET string and update the below if-statement to keep the test passing!
 				FString LEETTranslatedString = FormattedTestLayer2.ToString();
 
-				FString DesiredOutput = FString(TEXT("\x2021") TEXT("\xAB") TEXT("\xAB") TEXT("L0r3m") TEXT("\xBB") TEXT(" \"L0r3m 1p$um\" ") TEXT("\xAB") TEXT("Ip$um") TEXT("\xBB") TEXT("\xBB") TEXT(" | ") TEXT("\xAB") TEXT("\xAB") TEXT("L0r3m") TEXT("\xBB") TEXT(" \"L0r3m 1p$um\" ") TEXT("\xAB") TEXT("Ip$um") TEXT("\xBB") TEXT("\xBB") TEXT(" | ") TEXT("\xAB") TEXT("5.5421") TEXT("\xBB") TEXT(" | ") TEXT("\xAB") TEXT("5010.89221") TEXT("\xBB") TEXT(" | ") TEXT("\xAB") TEXT("Aug 20, 2080, 9:33:22 AM") TEXT("\xBB") TEXT(" | ") TEXT("\xAB") TEXT("92%") TEXT("\xBB") TEXT(" | ") TEXT("\xAB") TEXT("\xA4") TEXT("\xA0") TEXT("100.25") TEXT("\xBB") TEXT("\x2021"));
+				FString DesiredOutput = FString(TEXT("\x2021") TEXT("\xAB") TEXT("\xAB") TEXT("L0r3m") TEXT("\xBB") TEXT(" \"L0r3m 1p$um\" ") TEXT("\xAB") TEXT("Ip$um") TEXT("\xBB") TEXT("\xBB") TEXT(" | ") TEXT("\xAB") TEXT("\xAB") TEXT("L0r3m") TEXT("\xBB") TEXT(" \"L0r3m 1p$um\" ") TEXT("\xAB") TEXT("Ip$um") TEXT("\xBB") TEXT("\xBB") TEXT(" | ") TEXT("\xAB") TEXT("5.5421") TEXT("\xBB") TEXT(" | ") TEXT("\xAB") TEXT("5010.89221") TEXT("\xBB") TEXT(" | ") TEXT("\xAB") TEXT("Aug 20, 2080, 9:33:22 AM") TEXT("\xBB") TEXT(" | ") TEXT("\xAB") TEXT("92%") TEXT("\xBB") TEXT(" | ") TEXT("\xAB") TEXT("$") TEXT("\xA0") TEXT("100.25") TEXT("\xBB") TEXT("\x2021"));
 				// Convert the baked string into an FText, which will be leetified, then compare it to the rebuilt FText
 				if(LEETTranslatedString != DesiredOutput)
 				{
diff --git a/Engine/Source/Runtime/Core/Private/Windows/WindowsApplication.cpp b/Engine/Source/Runtime/Core/Private/Windows/WindowsApplication.cpp
--- a/Engine/Source/Runtime/Core/Private/Windows/WindowsApplication.cpp
+++ b/Engine/Source/Runtime/Core/Private/Windows/WindowsApplication.cpp
@@ -1973,17 +1973,17 @@ HRESULT FWindowsApplication::OnOLEDragEnter( const HWND HWnd, const FDragDropOLE
 		return 0;
 	}
 
-	switch (OLEData.Type)
+	if ((OLEData.Type & FDragDropOLEData::Text) && (OLEData.Type & FDragDropOLEData::Files))
 	{
-		case FDragDropOLEData::Text:
-			*CursorEffect = MessageHandler->OnDragEnterText(Window.ToSharedRef(), OLEData.OperationText);
-			break;
-		case FDragDropOLEData::Files:
-			*CursorEffect = MessageHandler->OnDragEnterFiles(Window.ToSharedRef(), OLEData.OperationFilenames);
-			break;
-		case FDragDropOLEData::None:
-		default:
-			break;
+		*CursorEffect = MessageHandler->OnDragEnterExternal(Window.ToSharedRef(), OLEData.OperationText, OLEData.OperationFilenames);
+	}
+	else if (OLEData.Type & FDragDropOLEData::Text)
+	{
+		*CursorEffect = MessageHandler->OnDragEnterText(Window.ToSharedRef(), OLEData.OperationText);
+	}
+	else if (OLEData.Type & FDragDropOLEData::Files)
+	{
+		*CursorEffect = MessageHandler->OnDragEnterFiles(Window.ToSharedRef(), OLEData.OperationFilenames);
 	}
 
 	return 0;
diff --git a/Engine/Source/Runtime/Core/Private/Windows/WindowsWindow.cpp b/Engine/Source/Runtime/Core/Private/Windows/WindowsWindow.cpp
--- a/Engine/Source/Runtime/Core/Private/Windows/WindowsWindow.cpp
+++ b/Engine/Source/Runtime/Core/Private/Windows/WindowsWindow.cpp
@@ -857,18 +857,18 @@ FDragDropOLEData DecipherOLEData(IDataObject* DataObjectPointer)
 	if (bHaveUnicodeText && S_OK == DataObjectPointer->GetData(&FormatEtc_UNICODE, &StorageMedium))
 	{
 		FOLEResourceGuard ResourceGuard(StorageMedium);
-		OLEData.Type = FDragDropOLEData::Text;
+		OLEData.Type |= FDragDropOLEData::Text;
 		OLEData.OperationText = static_cast<TCHAR*>(ResourceGuard.DataPointer);
 	}
-	else if (bHaveAnsiText && S_OK == DataObjectPointer->GetData(&FormatEtc_Ansii, &StorageMedium))
+	if (bHaveAnsiText && S_OK == DataObjectPointer->GetData(&FormatEtc_Ansii, &StorageMedium))
 	{
 		FOLEResourceGuard ResourceGuard(StorageMedium);
-		OLEData.Type = FDragDropOLEData::Text;
+		OLEData.Type |= FDragDropOLEData::Text;
 		OLEData.OperationText = static_cast<ANSICHAR*>(ResourceGuard.DataPointer);
 	}
-	else if (bHaveFiles && S_OK == DataObjectPointer->GetData(&FormatEtc_File, &StorageMedium))
+	if (bHaveFiles && S_OK == DataObjectPointer->GetData(&FormatEtc_File, &StorageMedium))
 	{
-		OLEData.Type = FDragDropOLEData::Files;
+		OLEData.Type |= FDragDropOLEData::Files;
 
 		FOLEResourceGuard ResourceGuard(StorageMedium);
 		const DROPFILES* DropFiles = static_cast<DROPFILES*>(ResourceGuard.DataPointer);
diff --git a/Engine/Source/Runtime/Core/Public/GenericPlatform/GenericApplicationMessageHandler.h b/Engine/Source/Runtime/Core/Public/GenericPlatform/GenericApplicationMessageHandler.h
--- a/Engine/Source/Runtime/Core/Public/GenericPlatform/GenericApplicationMessageHandler.h
+++ b/Engine/Source/Runtime/Core/Public/GenericPlatform/GenericApplicationMessageHandler.h
@@ -429,6 +429,11 @@ public:
 		return EDropEffect::None;
 	}
 
+	virtual EDropEffect::Type OnDragEnterExternal( const TSharedRef< FGenericWindow >& Window, const FString& Text, const TArray< FString >& Files )
+	{
+		return EDropEffect::None;
+	}
+
 	virtual EDropEffect::Type OnDragOver( const TSharedPtr< FGenericWindow >& Window )
 	{
 		return EDropEffect::None;
diff --git a/Engine/Source/Runtime/Core/Public/Internationalization/Text.h b/Engine/Source/Runtime/Core/Public/Internationalization/Text.h
--- a/Engine/Source/Runtime/Core/Public/Internationalization/Text.h
+++ b/Engine/Source/Runtime/Core/Public/Internationalization/Text.h
@@ -670,12 +670,25 @@ public:
 	 * @param Value				The text value to write into the buffer.
 	 * @param bRequiresQuotes	True if the written text literal must be surrounded by quotes (eg, when saving as a delimited list)
 	 *
-	 * @return True if we read a valid FText instance into OutValue, false otherwise
+	 * @return True if we wrote a valid FText instance into Buffer, false otherwise
 	 */
 	static bool WriteToString(FString& Buffer, const FText& Value, const bool bRequiresQuotes = false);
 
+	/**
+	 * Test to see whether a given buffer contains complex text.
+	 *
+	 * @return True if it does, false otherwise
+	 */
+	static bool IsComplexText(const TCHAR* Buffer);
+
 private:
 	static bool ReadFromString_ComplexText(const TCHAR* Buffer, FText& OutValue, const TCHAR* Namespace, int32* OutNumCharsRead);
+
+#define LOC_DEFINE_REGION
+	static const FString InvTextMarker;
+	static const FString NsLocTextMarker;
+	static const FString LocTextMarker;
+#undef LOC_DEFINE_REGION
 };
 
 class CORE_API FTextBuilder
diff --git a/Engine/Source/Runtime/Core/Public/Misc/AutomationTest.h b/Engine/Source/Runtime/Core/Public/Misc/AutomationTest.h
--- a/Engine/Source/Runtime/Core/Public/Misc/AutomationTest.h
+++ b/Engine/Source/Runtime/Core/Public/Misc/AutomationTest.h
@@ -86,13 +86,13 @@ public:
 };
 
 /** Simple class to store the automation test info */
-class FAutomationTestInfo
+class CORE_API FAutomationTestInfo
 {
 public:
 
 	// Default constructor
-	FAutomationTestInfo( ) :
-		TestFlags( 0 )
+	FAutomationTestInfo( )
+		: TestFlags( 0 )
 		, NumParticipantsRequired( 0 )
 		, NumDevicesCurrentlyRunningTest( 0 )
 	{}
@@ -103,8 +103,8 @@ public:
 	 *
 	 * @param	InTestInfo - comma separated value string containing the test info
 	 */
-	FAutomationTestInfo( const FString& InTestInfo ) :
-		NumParticipantsRequired( 0 )
+	FAutomationTestInfo( const FString& InTestInfo )
+		: NumParticipantsRequired( 0 )
 		, NumDevicesCurrentlyRunningTest( 0 )
 	{
 		ParseStringInfo( InTestInfo );
@@ -118,10 +118,12 @@ public:
 	 * @param	InTestFlag - Test flags
 	 * @param	InParameterName - optional parameter. e.g. asset name
 	 */
-	FAutomationTestInfo( const FString& InDisplayName, const FString& InTestName, const uint32 InTestFlags, const int32 InNumParticipantsRequired, const FString& InParameterName = FString() )
+	FAutomationTestInfo(const FString& InDisplayName, const FString& InTestName, const uint32 InTestFlags, const int32 InNumParticipantsRequired, const FString& InParameterName = FString(), const FString& InSourceFile = FString(), int32 InSourceFileLine = 0)
 		: DisplayName( InDisplayName )
 		, TestName( InTestName )
 		, TestParameter( InParameterName )
+		, SourceFile( InSourceFile )
+		, SourceFileLine( InSourceFileLine )
 		, TestFlags( InTestFlags )
 		, NumParticipantsRequired( InNumParticipantsRequired )
 		, NumDevicesCurrentlyRunningTest( 0 )
@@ -156,10 +158,7 @@ public:
 	 *
 	 * @return The test as a string.
 	 */
-	const FString GetTestAsString()
-	{
-		return FString::Printf( TEXT("%s,%s,%i,%d,%s,"), *DisplayName, *TestName, TestFlags, NumParticipantsRequired, *TestParameter );
-	}
+	FString GetTestAsString() const;
 
 
 	/**
@@ -167,7 +166,7 @@ public:
 	 *
 	 * @return The test name.
 	 */
-	const FString GetTestName() const
+	FString GetTestName() const
 	{
 		return TestName;
 	}
@@ -183,6 +182,25 @@ public:
 		return TestParameter;
 	}
 
+	/**
+	 * Get the source file this test originated in.
+	 *
+	 * @return the source file.
+	 */
+	const FString GetSourceFile() const
+	{
+		return SourceFile;
+	}
+
+	/**
+	 * Get the line number in the source file this test originated on.
+	 *
+	 * @return the source line number.
+	 */
+	const int32 GetSourceFileLine() const
+	{
+		return SourceFileLine;
+	}
 
 	/**
 	 * Get the type of test.
@@ -259,27 +277,7 @@ private:
 	 *
 	 * @Param InTestInfo - the test information as a string.
 	 */
-	void ParseStringInfo( const FString& InTestInfo )
-	{
-		//Split New Test Name into string array
-		TArray<FString> Pieces;
-		InTestInfo.ParseIntoArray(Pieces, TEXT(","), true);
-
-		// We should always have at least 3 parameters
-		check( Pieces.Num() >= 3 )
-
-		DisplayName = Pieces[0];
-		TestName = Pieces[1];
-		TestFlags = uint8( FCString::Atoi( *Pieces[2] ) );
-
-		NumParticipantsRequired = FCString::Atoi( *Pieces[3] );
-
-		// parameter 4 is optional
-		if ( Pieces.Num() == 5 )
-		{
-			TestParameter = Pieces[4];
-		}
-	}
+	void ParseStringInfo(const FString& InTestInfo);
 
 private:
 	/** Display name used in the UI */
@@ -291,6 +289,12 @@ private:
 	/** Parameter - e.g. an asset name or map name */
 	FString TestParameter;
 
+	/** The source file this test originated in. */
+	FString SourceFile;
+
+	/** The line number in the source file this test originated on. */
+	int32 SourceFileLine;
+
 	/** The test flags. */
 	uint32 TestFlags;
 
@@ -749,7 +753,7 @@ public:
 	 * Returns whether this test has any errors associated with it or not
 	 *
 	 * @return true if this test has at least one error associated with it; false if not
- 	 */
+	 */
 	bool HasAnyErrors() const;
 
 	/**
@@ -792,6 +796,28 @@ public:
 		bSuppressLogs = bNewValue;
 	}
 
+	/**
+	 * Enqueues a new latent command.
+	 */
+	FORCEINLINE void AddCommand(TSharedPtr<IAutomationLatentCommand> NewCommand)
+	{
+		FAutomationTestFramework::GetInstance().EnqueueLatentCommand(NewCommand);
+	}
+
+	/**
+	 * Enqueues a new latent network command.
+	 */
+	FORCEINLINE void AddCommand(TSharedPtr<IAutomationNetworkCommand> NewCommand)
+	{
+		FAutomationTestFramework::GetInstance().EnqueueNetworkCommand(NewCommand);
+	}
+
+	/** Gets the filename where this test was defined. */
+	virtual FString GetFileName() const { return ""; }
+
+	/** Gets the line number where this test was defined. */
+	virtual int32 GetFileLine() const { return 0; }
+
 public:
 
 	/**
@@ -967,7 +993,7 @@ protected:
 	 * @param OutTestCommands - The parameters to be specified to each call to RunTests (for load all maps, it would be the map name to load)
 	 */
 	virtual void GetTests(TArray<FString>& OutBeautifiedNames, TArray <FString>& OutTestCommands) const = 0;
-	
+
 	/**
 	 * Virtual call to execute the automation test.  
 	 *
@@ -1082,7 +1108,7 @@ public: \
  * the automation test framework as a result of the macro.
  */
 
-#define IMPLEMENT_SIMPLE_AUTOMATION_TEST_PRIVATE( TClass, PrettyName, TFlags ) \
+#define IMPLEMENT_SIMPLE_AUTOMATION_TEST_PRIVATE( TClass, PrettyName, TFlags, FileName, LineNumber ) \
 	class TClass : public FAutomationTestBase \
 	{ \
 	public: \
@@ -1099,6 +1125,8 @@ public: \
 		virtual uint32 GetTestFlags() const override { return TFlags; } \
 		virtual bool IsStressTest() const { return false; } \
 		virtual uint32 GetRequiredDeviceNum() const override { return 1; } \
+		virtual FString GetFileName() const { return FileName; } \
+		virtual int32 GetFileLine() const { return LineNumber; } \
 	protected: \
 		virtual void GetTests(TArray<FString>& OutBeautifiedNames, TArray <FString>& OutTestCommands) const override \
 		{ \
@@ -1109,7 +1137,7 @@ public: \
 		virtual FString GetBeautifiedTestName() const override { return PrettyName; } \
 	};
 
-#define IMPLEMENT_COMPLEX_AUTOMATION_TEST_PRIVATE( TClass, PrettyName, TFlags ) \
+#define IMPLEMENT_COMPLEX_AUTOMATION_TEST_PRIVATE( TClass, PrettyName, TFlags, FileName, LineNumber ) \
 	class TClass : public FAutomationTestBase \
 	{ \
 	public: \
@@ -1126,13 +1154,15 @@ public: \
 		virtual uint32 GetTestFlags() const override { return ((TFlags) & ~(EAutomationTestFlags::SmokeFilter)); } \
 		virtual bool IsStressTest() const { return true; } \
 		virtual uint32 GetRequiredDeviceNum() const override { return 1; } \
+		virtual FString GetFileName() const { return FileName; } \
+		virtual int32 GetFileLine() const { return LineNumber; } \
 	protected: \
 		virtual void GetTests(TArray<FString>& OutBeautifiedNames, TArray <FString>& OutTestCommands) const override; \
 		virtual bool RunTest(const FString& Parameters) override; \
 		virtual FString GetBeautifiedTestName() const override { return PrettyName; } \
 	};
 
-#define IMPLEMENT_NETWORKED_AUTOMATION_TEST_PRIVATE(TClass, PrettyName, TFlags, NumParticipants) \
+#define IMPLEMENT_NETWORKED_AUTOMATION_TEST_PRIVATE(TClass, PrettyName, TFlags, NumParticipants, FileName, LineNumber) \
 	class TClass : public FAutomationTestBase \
 	{ \
 	public: \
@@ -1148,6 +1178,8 @@ public: \
 		} \
 		virtual uint32 GetTestFlags() const override { return ((TFlags) & ~(EAutomationTestFlags::EditorContext | EAutomationTestFlags::CommandletContext | EAutomationTestFlags::SmokeFilter)); } \
 		virtual uint32 GetRequiredDeviceNum() const override { return NumParticipants; } \
+		virtual FString GetFileName() const { return FileName; } \
+		virtual int32 GetFileLine() const { return LineNumber; } \
 	protected: \
 		virtual void GetTests(TArray<FString>& OutBeautifiedNames, TArray <FString>& OutTestCommands) const override \
 		{ \
@@ -1161,31 +1193,57 @@ public: \
 
 #if WITH_AUTOMATION_WORKER
 	#define IMPLEMENT_SIMPLE_AUTOMATION_TEST( TClass, PrettyName, TFlags ) \
-		IMPLEMENT_SIMPLE_AUTOMATION_TEST_PRIVATE(TClass, PrettyName, TFlags) \
+		IMPLEMENT_SIMPLE_AUTOMATION_TEST_PRIVATE(TClass, PrettyName, TFlags, __FILE__, __LINE__) \
 		namespace\
 		{\
 			TClass TClass##AutomationTestInstance( TEXT(#TClass) );\
 		}
 	#define IMPLEMENT_COMPLEX_AUTOMATION_TEST( TClass, PrettyName, TFlags ) \
-		IMPLEMENT_COMPLEX_AUTOMATION_TEST_PRIVATE(TClass, PrettyName, TFlags) \
+		IMPLEMENT_COMPLEX_AUTOMATION_TEST_PRIVATE(TClass, PrettyName, TFlags, __FILE__, __LINE__) \
 		namespace\
 		{\
 			TClass TClass##AutomationTestInstance( TEXT(#TClass) );\
 		}
 	#define IMPLEMENT_NETWORKED_AUTOMATION_TEST(TClass, PrettyName, TFlags, NumParticipants) \
-		IMPLEMENT_NETWORKED_AUTOMATION_TEST_PRIVATE(TClass, PrettyName, TFlags, NumParticipants) \
+		IMPLEMENT_NETWORKED_AUTOMATION_TEST_PRIVATE(TClass, PrettyName, TFlags, NumParticipants, __FILE__, __LINE__) \
 		namespace\
 		{\
 			TClass TClass##AutomationTestInstance( TEXT(#TClass) );\
 		}
 #else
 	#define IMPLEMENT_SIMPLE_AUTOMATION_TEST( TClass, PrettyName, TFlags ) \
-		IMPLEMENT_SIMPLE_AUTOMATION_TEST_PRIVATE(TClass, PrettyName, TFlags)
+		IMPLEMENT_SIMPLE_AUTOMATION_TEST_PRIVATE(TClass, PrettyName, TFlags, __FILE__, __LINE__)
 	#define IMPLEMENT_COMPLEX_AUTOMATION_TEST( TClass, PrettyName, TFlags ) \
-		IMPLEMENT_COMPLEX_AUTOMATION_TEST_PRIVATE(TClass, PrettyName, TFlags)
+		IMPLEMENT_COMPLEX_AUTOMATION_TEST_PRIVATE(TClass, PrettyName, TFlags, __FILE__, __LINE__)
 	#define IMPLEMENT_NETWORKED_AUTOMATION_TEST(TClass, PrettyName, TFlags, NumParticipants) \
-		IMPLEMENT_NETWORKED_AUTOMATION_TEST_PRIVATE(TClass, PrettyName, TFlags, NumParticipants)
+		IMPLEMENT_NETWORKED_AUTOMATION_TEST_PRIVATE(TClass, PrettyName, TFlags, NumParticipants, __FILE__, __LINE__)
 #endif // #if WITH_AUTOMATION_WORKER
 
 
+//////////////////////////////////////////////////////////////////////////
+// Basic Latent Commands
 
+/**
+ * Run some code latently with a predicate lambda.  If the predicate returns true, the latent action will be called 
+ * again next frame.  If it returns false, the command will stop running.
+ */
+class CORE_API FFunctionLatentCommand : public IAutomationLatentCommand
+{
+public:
+	FFunctionLatentCommand(TFunction<bool()> InLatentPredicate)
+		: LatentPredicate(InLatentPredicate)
+	{
+	}
+
+	virtual ~FFunctionLatentCommand()
+	{
+	}
+
+	virtual bool Update() override
+	{
+		return LatentPredicate();
+	}
+
+private:
+	TFunction<bool()> LatentPredicate;
+};
diff --git a/Engine/Source/Runtime/Core/Public/Misc/ConfigCacheIni.h b/Engine/Source/Runtime/Core/Public/Misc/ConfigCacheIni.h
--- a/Engine/Source/Runtime/Core/Public/Misc/ConfigCacheIni.h
+++ b/Engine/Source/Runtime/Core/Public/Misc/ConfigCacheIni.h
@@ -272,10 +272,13 @@ public:
 	 */
 	void ProcessSourceAndCheckAgainstBackup();
 
-private:
-
 	/** Checks if the PropertyValue should be exported in quotes when writing the ini to disk. */
-	bool ShouldExportQuotedString(const FString& PropertyValue) const;
+	static bool ShouldExportQuotedString(const FString& PropertyValue);
+
+	/** Generate a correctly escaped line to add to the config file for the given property */
+	static FString GenerateExportedPropertyLine(const FString& PropertyName, const FString& PropertyValue);
+
+private:
 
 	/** 
 	 * Save the source hierarchy which was loaded out to a backup file so we can check future changes in the base/default configs
diff --git a/Engine/Source/Runtime/Core/Public/Misc/FbxErrors.h b/Engine/Source/Runtime/Core/Public/Misc/FbxErrors.h
--- a/Engine/Source/Runtime/Core/Public/Misc/FbxErrors.h
+++ b/Engine/Source/Runtime/Core/Public/Misc/FbxErrors.h
@@ -50,6 +50,8 @@ struct CORE_API FFbxErrors
 
 	static FName Generic_Mesh_LOD_MultipleFilesSelected;
 
+	static FName Generic_Mesh_SkinxxNameError;
+
 	/** Static Mesh **/
 	static FName StaticMesh_TooManyMaterials;
 
diff --git a/Engine/Source/Runtime/Core/Public/Modules/ModuleManager.h b/Engine/Source/Runtime/Core/Public/Modules/ModuleManager.h
--- a/Engine/Source/Runtime/Core/Public/Modules/ModuleManager.h
+++ b/Engine/Source/Runtime/Core/Public/Modules/ModuleManager.h
@@ -461,6 +461,14 @@ protected:
 		: bCanProcessNewlyLoadedObjects(false)
 	{ }
 
+private:
+
+	/**
+	 * Prevent copy constructor from being triggered.
+	 */
+	FModuleManager(const FModuleManager&)
+	{ }
+
 protected:
 
 	/**
@@ -560,14 +568,14 @@ private:
 	mutable TOptional<TMap<FName, FString>> ModulePathsCache;
 
 	/** Multicast delegate that will broadcast a notification when modules are loaded, unloaded, or
-	    our set of known modules changes */
+		our set of known modules changes */
 	FModulesChangedEvent ModulesChangedEvent;
 	
 	/** Multicast delegate called to process any new loaded objects. */
 	FSimpleMulticastDelegate ProcessLoadedObjectsCallback;
 
 	/** When module manager is linked against an application that supports UObjects, this delegate will be primed
-	    at startup to provide information about whether a UObject package is loaded into memory. */
+		at startup to provide information about whether a UObject package is loaded into memory. */
 	FIsPackageLoadedCallback IsPackageLoaded;
 
 	/** Array of engine binaries directories. */
@@ -772,7 +780,7 @@ class FDefaultGameModuleImpl
 #else
 
 /** IMPLEMENT_PRIMARY_GAME_MODULE must be used for at least one game module in your game.  It sets the "name"
-    your game when compiling in monolithic mode. This is passed in by UBT from the .uproject name, and manually specifying a 
+	your game when compiling in monolithic mode. This is passed in by UBT from the .uproject name, and manually specifying a 
 	name is no longer necessary. */
 #if IS_MONOLITHIC
 	#if PLATFORM_DESKTOP
diff --git a/Engine/Source/Runtime/Core/Public/Windows/WindowsApplication.h b/Engine/Source/Runtime/Core/Public/Windows/WindowsApplication.h
--- a/Engine/Source/Runtime/Core/Public/Windows/WindowsApplication.h
+++ b/Engine/Source/Runtime/Core/Public/Windows/WindowsApplication.h
@@ -164,9 +164,9 @@ struct FDragDropOLEData
 {
 	enum EWindowsOLEDataType
 	{
-		None,
-		Text,
-		Files
+		None = 0,
+		Text = 1<<0,
+		Files = 1<<1,
 	};
 
 	FDragDropOLEData()
@@ -175,7 +175,7 @@ struct FDragDropOLEData
 
 	FString OperationText;
 	TArray<FString> OperationFilenames;
-	EWindowsOLEDataType Type;
+	uint8 Type;
 };
 
 
diff --git a/Engine/Source/Runtime/CoreUObject/Private/UObject/Obj.cpp b/Engine/Source/Runtime/CoreUObject/Private/UObject/Obj.cpp
--- a/Engine/Source/Runtime/CoreUObject/Private/UObject/Obj.cpp
+++ b/Engine/Source/Runtime/CoreUObject/Private/UObject/Obj.cpp
@@ -1638,16 +1638,16 @@ void UObject::LoadConfig( UClass* ConfigClass/*=NULL*/, const TCHAR* InFilename/
 					Key = FString::Printf(TEXT("%s[%i]"), *Property->GetName(), i);
 				}
 
-					FString Value;
+				FString Value;
 				const bool bFoundValue = GConfig->GetString( *ClassSection, *Key, Value, *PropFileName );
-					if (bFoundValue)
+				if (bFoundValue)
+				{
+					if (Property->ImportText(*Value, Property->ContainerPtrToValuePtr<uint8>(this, i), PortFlags, this) == NULL)
 					{
-						if (Property->ImportText(*Value, Property->ContainerPtrToValuePtr<uint8>(this, i), PortFlags, this) == NULL)
-						{
-							// this should be an error as the properties from the .ini / .int file are not correctly being read in and probably are affecting things in subtle ways
-							UE_LOG(LogObj, Error, TEXT("LoadConfig (%s): import failed for %s in: %s"), *GetPathName(), *Property->GetName(), *Value);
-						}
+						// this should be an error as the properties from the .ini / .int file are not correctly being read in and probably are affecting things in subtle ways
+						UE_LOG(LogObj, Error, TEXT("LoadConfig (%s): import failed for %s in: %s"), *GetPathName(), *Property->GetName(), *Value);
 					}
+				}
 
 #if !UE_BUILD_SHIPPING
 				if (!bFoundValue && !FPlatformProperties::RequiresCookedData())
@@ -1836,13 +1836,13 @@ void UObject::SaveConfig( uint64 Flags, const TCHAR* InFilename, FConfigCacheIni
 					FString CompleteKey = FString::Printf(TEXT("%s%s"), bIsADefaultIniWrite ? TEXT("+") : TEXT(""), *Key);
 
 					FScriptArrayHelper_InContainer ArrayHelper(Array, this);
-						for( int32 i=0; i<ArrayHelper.Num(); i++ )
-						{
-							FString	Buffer;
-							Array->Inner->ExportTextItem( Buffer, ArrayHelper.GetRawPtr(i), ArrayHelper.GetRawPtr(i), this, PortFlags );
-							Sec->Add(*CompleteKey, *Buffer);
-						}
+					for( int32 i=0; i<ArrayHelper.Num(); i++ )
+					{
+						FString	Buffer;
+						Array->Inner->ExportTextItem( Buffer, ArrayHelper.GetRawPtr(i), ArrayHelper.GetRawPtr(i), this, PortFlags );
+						Sec->Add(*CompleteKey, *Buffer);
 					}
+				}
 				else if( Property->Identical_InContainer(this, SuperClassDefaultObject) )
 				{
 					// If we are not writing it to config above, we should make sure that this property isn't stagnant in the cache.
@@ -1867,10 +1867,10 @@ void UObject::SaveConfig( uint64 Flags, const TCHAR* InFilename, FConfigCacheIni
 							Key = TempKey;
 						}
 
-							FString	Value;
-							Property->ExportText_InContainer( Index, Value, this, this, this, PortFlags );
-							Config->SetString( *Section, *Key, *Value, *PropFileName );
-						}
+						FString	Value;
+						Property->ExportText_InContainer( Index, Value, this, this, this, PortFlags );
+						Config->SetString( *Section, *Key, *Value, *PropFileName );
+					}
 					else if( Property->Identical_InContainer(this, SuperClassDefaultObject, Index) )
 					{
 						// If we are not writing it to config above, we should make sure that this property isn't stagnant in the cache.
diff --git a/Engine/Source/Runtime/Engine/Classes/Components/StaticMeshComponent.h b/Engine/Source/Runtime/Engine/Classes/Components/StaticMeshComponent.h
--- a/Engine/Source/Runtime/Engine/Classes/Components/StaticMeshComponent.h
+++ b/Engine/Source/Runtime/Engine/Classes/Components/StaticMeshComponent.h
@@ -440,11 +440,9 @@ public:
 	/**
 	 * Determines whether any of the component's LODs require override vertex color fixups
 	 *
-	 * @param	OutLODIndices	Indices of the LODs requiring fixup, if any
-	 *
 	 * @return	true if any LODs require override vertex color fixups
 	 */
-	bool RequiresOverrideVertexColorsFixup( TArray<int32>& OutLODIndices );
+	bool RequiresOverrideVertexColorsFixup();
 
 	/**
 	 * Update the vertex override colors if necessary (i.e. vertices from source mesh have changed from override colors)
@@ -492,7 +490,7 @@ private:
 	void InitResources();
 
 	/** Update the vertex override colors */
-	void PrivateFixupOverrideColors( const TArray<int32>& LODsToUpdate );
+	void PrivateFixupOverrideColors();
 
 protected:
 
diff --git a/Engine/Source/Runtime/Engine/Classes/EdGraph/EdGraphNode.h b/Engine/Source/Runtime/Engine/Classes/EdGraph/EdGraphNode.h
--- a/Engine/Source/Runtime/Engine/Classes/EdGraph/EdGraphNode.h
+++ b/Engine/Source/Runtime/Engine/Classes/EdGraph/EdGraphNode.h
@@ -11,6 +11,8 @@ class UEdGraphSchema;
 class UEdGraphPin;
 class SGraphNode;
 struct FEdGraphPinType;
+struct FSlateIcon;
+
 
 /** Enum used to define which way data flows into or out of this pin. */
 UENUM()
@@ -381,8 +383,12 @@ public:
 	 */
 	virtual FString GetDocumentationExcerptName() const;
 
-	/** @return name or brush to use in menu or on node */
-	virtual FName GetPaletteIcon(FLinearColor& OutColor) const { return TEXT("GraphEditor.Default_16x"); }
+	/** @return Icon to use in menu or on node */
+	DEPRECATED(4.13, "Please override 'virtual FSlateIcon GetIconAndTint(FLinearColor& OutColor) const;' instead.")
+	virtual FName GetPaletteIcon(FLinearColor& OutColor) const { return NAME_None; }
+
+	/** @return Icon to use in menu or on node */
+	virtual FSlateIcon GetIconAndTint(FLinearColor& OutColor) const;
 
 	/** Should we show the Palette Icon for this node on the node title */
 	virtual bool ShowPaletteIconOnNode() const { return false; }
diff --git a/Engine/Source/Runtime/Engine/Classes/EdGraph/EdGraphNode_Documentation.h b/Engine/Source/Runtime/Engine/Classes/EdGraph/EdGraphNode_Documentation.h
--- a/Engine/Source/Runtime/Engine/Classes/EdGraph/EdGraphNode_Documentation.h
+++ b/Engine/Source/Runtime/Engine/Classes/EdGraph/EdGraphNode_Documentation.h
@@ -32,7 +32,7 @@ public:
 	ENGINE_API virtual TSharedPtr<class INameValidatorInterface> MakeNameValidator() const override;
 	ENGINE_API virtual FString GetDocumentationLink() const override { return Link; }
 	ENGINE_API virtual FString GetDocumentationExcerptName() const override { return Excerpt; }
-	ENGINE_API virtual FName GetPaletteIcon(FLinearColor& OutColor) const override { return TEXT("GraphEditor.Documentation_16x");  }
+	ENGINE_API virtual FSlateIcon GetIconAndTint(FLinearColor& OutColor) const override;
 	//~ End UEdGraphNode Interface
 
 	/** Set the Bounds for the comment node */
diff --git a/Engine/Source/Runtime/Engine/Classes/Engine/Engine.h b/Engine/Source/Runtime/Engine/Classes/Engine/Engine.h
--- a/Engine/Source/Runtime/Engine/Classes/Engine/Engine.h
+++ b/Engine/Source/Runtime/Engine/Classes/Engine/Engine.h
@@ -1635,6 +1635,9 @@ public:
 	/** Initialize the game engine. */
 	virtual void Init(IEngineLoop* InEngineLoop);
 
+	/** Start the game, separate from the initialize call to allow for post initialize configuration before the game starts. */
+	virtual void Start();
+
 	/** Called at shutdown, just before the exit purge.	 */
 	virtual void PreExit();
 	virtual void ShutdownAudioDeviceManager();
diff --git a/Engine/Source/Runtime/Engine/Classes/Engine/GameEngine.h b/Engine/Source/Runtime/Engine/Classes/Engine/GameEngine.h
--- a/Engine/Source/Runtime/Engine/Classes/Engine/GameEngine.h
+++ b/Engine/Source/Runtime/Engine/Classes/Engine/GameEngine.h
@@ -99,6 +99,7 @@ public:
 	// UEngine interface
 
 	virtual void Init(class IEngineLoop* InEngineLoop) override;
+	virtual void Start() override;
 	virtual void PreExit() override;
 	virtual void Tick( float DeltaSeconds, bool bIdleMode ) override;
 	virtual float GetMaxTickRate( float DeltaTime, bool bAllowFrameRateSmoothing = true ) const override;
diff --git a/Engine/Source/Runtime/Engine/Classes/PhysicsEngine/BodyInstance.h b/Engine/Source/Runtime/Engine/Classes/PhysicsEngine/BodyInstance.h
--- a/Engine/Source/Runtime/Engine/Classes/PhysicsEngine/BodyInstance.h
+++ b/Engine/Source/Runtime/Engine/Classes/PhysicsEngine/BodyInstance.h
@@ -1077,6 +1077,7 @@ private:
 	friend struct FInitBodiesHelper<false>;
 	friend class FDerivedDataPhysXBinarySerializer;
 	friend class FBodyInstanceCustomizationHelper;
+	friend class FFoliageTypeCustomizationHelpers;
 
 #if WITH_BOX2D
 
diff --git a/Engine/Source/Runtime/Engine/Classes/Sound/AudioSettings.h b/Engine/Source/Runtime/Engine/Classes/Sound/AudioSettings.h
--- a/Engine/Source/Runtime/Engine/Classes/Sound/AudioSettings.h
+++ b/Engine/Source/Runtime/Engine/Classes/Sound/AudioSettings.h
@@ -63,6 +63,18 @@ class ENGINE_API UAudioSettings : public UDeveloperSettings
 	UPROPERTY(config, EditAnywhere, Category="Quality")
 	TArray<FAudioQualitySettings> QualityLevels;
 
+	/**
+	 * The format string to use when generating the filename for contexts within dialogue waves. This must generate unique names for your project.
+	 * Available format markers:
+	 *   * {DialogueGuid} - The GUID of the dialogue wave. Guaranteed to be unique and stable against asset renames.
+	 *   * {DialogueHash} - The hash of the dialogue wave. Not guaranteed to be unique or stable against asset renames, however may be unique enough if combined with the dialogue name.
+	 *   * {DialogueName} - The name of the dialogue wave. Not guaranteed to be unique or stable against asset renames, however may be unique enough if combined with the dialogue hash.
+	 *   * {ContextId}    - The ID of the context. Guaranteed to be unique within its dialogue wave. Not guaranteed to be stable against changes to the context.
+	 *   * {ContextIndex} - The index of the context within its parent dialogue wave. Guaranteed to be unique within its dialogue wave. Not guaranteed to be stable against contexts being removed.
+	 */
+	UPROPERTY(config, EditAnywhere, Category="Dialogue")
+	FString DialogueFilenameFormat;
+
 	const FAudioQualitySettings& GetQualityLevelSettings(int32 QualityLevel) const;
 
 private:
diff --git a/Engine/Source/Runtime/Engine/Classes/Sound/DialogueWave.h b/Engine/Source/Runtime/Engine/Classes/Sound/DialogueWave.h
--- a/Engine/Source/Runtime/Engine/Classes/Sound/DialogueWave.h
+++ b/Engine/Source/Runtime/Engine/Classes/Sound/DialogueWave.h
@@ -16,6 +16,7 @@ class UDialogueSoundWaveProxy;
 struct ENGINE_API FDialogueConstants
 {
 	static const FString DialogueNamespace;
+	static const FString SubtitleKeySuffix;
 #if WITH_EDITORONLY_DATA
 	static const FString PropertyName_AudioFile;
 	static const FString PropertyName_VoiceActorDirection;
@@ -58,6 +59,7 @@ struct ENGINE_API FDialogueContextMapping
 	UDialogueSoundWaveProxy* Proxy;
 
 	/** Gets the localization key to use for this context mapping */
+	FString GetLocalizationKey() const;
 	FString GetLocalizationKey(const FString& InOwnerDialogueWaveKey) const;
 };
 
@@ -75,18 +77,26 @@ class UDialogueWave : public UObject
 	UPROPERTY(EditAnywhere, Category=Filter, AssetRegistrySearchable)
 	uint32 bMature:1;
 
+	/**  */
+	UPROPERTY(EditAnywhere, Category=Script, meta=(InlineEditConditionToggle))
+	uint32 bOverride_SubtitleOverride : 1;
+
 	/** A localized version of the text that is actually spoken phonetically in the audio. */
-	UPROPERTY(EditAnywhere, Category=Script )
+	UPROPERTY(EditAnywhere, Category=Script)
 	FString SpokenText;
 
+	/** A localized version of the subtitle text that should be displayed for this audio. By default this will be the same as the Spoken Text. */
+	UPROPERTY(EditAnywhere, Category=Script, meta=(EditCondition="bOverride_SubtitleOverride"))
+	FString SubtitleOverride;
+
 #if WITH_EDITORONLY_DATA
 	/** Provides general notes to the voice actor intended to direct their performance, as well as contextual information to the translator. */
-	UPROPERTY(EditAnywhere, Category=Script )
+	UPROPERTY(EditAnywhere, Category=Script)
 	FString VoiceActorDirection;
 #endif // WITH_EDITORONLY_DATA
 
 	/* Mappings between dialogue contexts and associated soundwaves. */
-	UPROPERTY(EditAnywhere, Category=DialogueContexts )
+	UPROPERTY(EditAnywhere, Category=DialogueContexts)
 	TArray<FDialogueContextMapping> ContextMappings;
 
 	UPROPERTY()
@@ -116,6 +126,7 @@ public:
 	ENGINE_API FString GetContextLocalizationKey(const FDialogueContextMapping& ContextMapping) const;
 	ENGINE_API FString GetContextRecordedAudioFilename(const FDialogueContext& Context) const;
 	ENGINE_API FString GetContextRecordedAudioFilename(const FDialogueContextMapping& ContextMapping) const;
+	ENGINE_API static FString BuildRecordedAudioFilename(const FString& FormatString, const FGuid& DialogueGuid, const FString& DialogueName, const FString& ContextId, const int32 ContextIndex);
 	//~ End UDialogueWave Interface.
 
 	ENGINE_API void UpdateContext(FDialogueContextMapping& ContextMapping, USoundWave* SoundWave, UDialogueVoice* Speaker, const TArray<UDialogueVoice*>& Targets);
diff --git a/Engine/Source/Runtime/Engine/Private/Animation/Skeleton.cpp b/Engine/Source/Runtime/Engine/Private/Animation/Skeleton.cpp
--- a/Engine/Source/Runtime/Engine/Private/Animation/Skeleton.cpp
+++ b/Engine/Source/Runtime/Engine/Private/Animation/Skeleton.cpp
@@ -744,13 +744,14 @@ void USkeleton::CollectAnimationNotifies()
 	for (auto Iter = AssetList.CreateConstIterator(); Iter; ++Iter)
 	{
 		const FAssetData& Asset = *Iter;
-		const FString* SkeletonValue = Asset.TagsAndValues.Find(TEXT("Skeleton"));
-		if (SkeletonValue && *SkeletonValue == CurrentSkeletonName)
+		const FString SkeletonValue = Asset.GetTagValueRef<FString>(TEXT("Skeleton"));
+		if (SkeletonValue == CurrentSkeletonName)
 		{
-			if (const FString* Value = Asset.TagsAndValues.Find(USkeleton::AnimNotifyTag))
+			FString Value;
+			if (Asset.GetTagValue(USkeleton::AnimNotifyTag, Value))
 			{
 				TArray<FString> NotifyList;
-				Value->ParseIntoArray(NotifyList, *USkeleton::AnimNotifyTagDelimiter, true);
+				Value.ParseIntoArray(NotifyList, *USkeleton::AnimNotifyTagDelimiter, true);
 				for (auto NotifyIter = NotifyList.CreateConstIterator(); NotifyIter; ++NotifyIter)
 				{
 					FString NotifyName = *NotifyIter;
diff --git a/Engine/Source/Runtime/Engine/Private/Components/StaticMeshComponent.cpp b/Engine/Source/Runtime/Engine/Private/Components/StaticMeshComponent.cpp
--- a/Engine/Source/Runtime/Engine/Private/Components/StaticMeshComponent.cpp
+++ b/Engine/Source/Runtime/Engine/Private/Components/StaticMeshComponent.cpp
@@ -913,31 +913,19 @@ FTransform UStaticMeshComponent::GetSocketTransform(FName InSocketName, ERelativ
 }
 
 
-bool UStaticMeshComponent::RequiresOverrideVertexColorsFixup( TArray<int32>& OutLODIndices )
+bool UStaticMeshComponent::RequiresOverrideVertexColorsFixup()
 {
-	OutLODIndices.Empty();
 	bool bFixupRequired = false;
 
 #if WITH_EDITORONLY_DATA
 	if ( StaticMesh && StaticMesh->RenderData
-		&& StaticMesh->RenderData->DerivedDataKey != StaticMeshDerivedDataKey )
+		&& StaticMesh->RenderData->DerivedDataKey != StaticMeshDerivedDataKey
+		&& LODData.Num() > 0
+		&& LODData[0].OverrideVertexColors
+		&& LODData[0].OverrideVertexColors->GetNumVertices() > 0
+		&& LODData[0].PaintedVertices.Num() > 0 )
 	{
-		int32 NumLODs = StaticMesh->GetNumLODs();
-		// Iterate over each LOD to confirm which ones, if any, actually need to have their colors updated
-		for ( TArray<FStaticMeshComponentLODInfo>::TConstIterator LODIter( LODData ); LODIter; ++LODIter )
-		{
-			const FStaticMeshComponentLODInfo& CurCompLODInfo = *LODIter;
-
-			// Confirm that the LOD has override colors and cached vertex data. If it doesn't have both, it can't be fixed up.
-			if ( CurCompLODInfo.OverrideVertexColors && 
-				CurCompLODInfo.OverrideVertexColors->GetNumVertices() > 0 && 
-				CurCompLODInfo.PaintedVertices.Num() > 0 && 
-				LODIter.GetIndex() < NumLODs )
-			{
-				OutLODIndices.Add( LODIter.GetIndex() );
-				bFixupRequired = true;
-			}
-		}
+		bFixupRequired = true;
 	}
 #endif // WITH_EDITORONLY_DATA
 
@@ -1111,73 +1099,113 @@ bool UStaticMeshComponent::FixupOverrideColorsIfNecessary( bool bRebuildingStati
 
 	// Detect if there is a version mismatch between the source mesh and the component. If so, the component's LODs potentially
 	// need to have their override colors updated to match changes in the source mesh.
-	TArray<int32> LODsToUpdate;
 
-	if ( RequiresOverrideVertexColorsFixup( LODsToUpdate ) )
+	if ( RequiresOverrideVertexColorsFixup() )
 	{
 		// Check if we are building the static mesh.  If so we dont need to reregister this component as its already unregistered and will be reregistered
 		// when the static mesh is done building.  Having nested reregister contexts is not supported.
 		if( bRebuildingStaticMesh )
 		{
-			PrivateFixupOverrideColors( LODsToUpdate );
+			PrivateFixupOverrideColors();
 		}
 		else
 		{
 			// Detach this component because rendering changes are about to be applied
 			FComponentReregisterContext ReregisterContext( this );
-			PrivateFixupOverrideColors( LODsToUpdate );
+			PrivateFixupOverrideColors();
 		}
+
+		return true;
 	}
-	return LODsToUpdate.Num() > 0 ? true : false;
-#else
-	return false;
 #endif // WITH_EDITORONLY_DATA
 
+	return false;
 }
 
-void UStaticMeshComponent::PrivateFixupOverrideColors( const TArray<int32>& LODsToUpdate )
+void UStaticMeshComponent::InitResources()
 {
-#if WITH_EDITOR
-	UE_LOG(LogStaticMesh,Verbose,TEXT("Fixing up override colors for %s [%s]"),*GetPathName(),*StaticMesh->GetPathName());
-	for ( TArray<int32>::TConstIterator LODIdxIter( LODsToUpdate ); LODIdxIter; ++LODIdxIter )
+	for(int32 LODIndex = 0; LODIndex < LODData.Num(); LODIndex++)
 	{
-		FStaticMeshComponentLODInfo& CurCompLODInfo = LODData[ *LODIdxIter ];
-		FStaticMeshLODResources& CurRenderData = StaticMesh->RenderData->LODResources[ *LODIdxIter ];
-		TArray<FColor> NewOverrideColors;
-
-		CurCompLODInfo.BeginReleaseOverrideVertexColors();
-		FlushRenderingCommands();
-		RemapPaintedVertexColors(
-			CurCompLODInfo.PaintedVertices,
-			*CurCompLODInfo.OverrideVertexColors,
-			CurRenderData.PositionVertexBuffer,
-			&CurRenderData.VertexBuffer,
-			NewOverrideColors
-			);
-		if (NewOverrideColors.Num())
+		FStaticMeshComponentLODInfo &LODInfo = LODData[LODIndex];
+		if(LODInfo.OverrideVertexColors)
 		{
-			CurCompLODInfo.OverrideVertexColors->InitFromColorArray(NewOverrideColors);
+			BeginInitResource(LODInfo.OverrideVertexColors);
+			INC_DWORD_STAT_BY( STAT_InstVertexColorMemory, LODInfo.OverrideVertexColors->GetAllocatedSize() );
 		}
-		// Note: In order to avoid data loss we dont clear the colors if they cannot be fixed up.
-
-		// Initialize the vert. colors
-		BeginInitResource( CurCompLODInfo.OverrideVertexColors );
 	}
-	StaticMeshDerivedDataKey = StaticMesh->RenderData->DerivedDataKey;
-#endif // WITH_EDITOR
 }
 
-void UStaticMeshComponent::InitResources()
+void UStaticMeshComponent::PrivateFixupOverrideColors()
 {
-	for(int32 LODIndex = 0; LODIndex < LODData.Num(); LODIndex++)
+#if WITH_EDITOR
+	if (!StaticMesh || !StaticMesh->RenderData)
 	{
-		FStaticMeshComponentLODInfo &LODInfo = LODData[LODIndex];
-		if(LODInfo.OverrideVertexColors)
+		return;
+	}
+
+	const uint32 NumLODs = StaticMesh->RenderData->LODResources.Num();
+
+	// Initialize override vertex colors on any new LODs which have just been created
+	SetLODDataCount(NumLODs, LODData.Num());
+
+	FStaticMeshComponentLODInfo& LOD0Info = LODData[0];
+	if (LOD0Info.OverrideVertexColors)
+	{
+		for (uint32 LODIndex = 0; LODIndex < NumLODs; ++LODIndex)
 		{
+			FStaticMeshComponentLODInfo& LODInfo = LODData[LODIndex];
+
+			if (LODInfo.OverrideVertexColors == nullptr)
+			{
+				LODInfo.OverrideVertexColors = new FColorVertexBuffer;
+			}
+			else
+			{
+				LODInfo.BeginReleaseOverrideVertexColors();
+				FlushRenderingCommands();
+			}
+
+			FStaticMeshLODResources& CurRenderData = StaticMesh->RenderData->LODResources[LODIndex];
+
+			TArray<FColor> NewOverrideColors;
+
+			if (LOD0Info.PaintedVertices.Num() > 0)
+			{
+				// Build override colors for LOD, based on LOD0
+				RemapPaintedVertexColors(
+					LOD0Info.PaintedVertices,
+					*LOD0Info.OverrideVertexColors,
+					CurRenderData.PositionVertexBuffer,
+					&CurRenderData.VertexBuffer,
+					NewOverrideColors
+					);
+			}
+
+			if (NewOverrideColors.Num())
+			{
+				LODInfo.OverrideVertexColors->InitFromColorArray(NewOverrideColors);
+
+				// Update the PaintedVertices array
+				const int32 NumVerts = CurRenderData.GetNumVertices();
+				check(NumVerts == NewOverrideColors.Num());
+
+				LODInfo.PaintedVertices.Reserve(NumVerts);
+				for (int32 VertIndex = 0; VertIndex < NumVerts; ++VertIndex)
+				{
+					FPaintedVertex* Vertex = new(LODInfo.PaintedVertices) FPaintedVertex;
+					Vertex->Position = CurRenderData.PositionVertexBuffer.VertexPosition(VertIndex);
+					Vertex->Normal = CurRenderData.VertexBuffer.VertexTangentZ(VertIndex);
+					Vertex->Color = LODInfo.OverrideVertexColors->VertexColor(VertIndex);
+				}
+			}
+
 			BeginInitResource(LODInfo.OverrideVertexColors);
-			INC_DWORD_STAT_BY( STAT_InstVertexColorMemory, LODInfo.OverrideVertexColors->GetAllocatedSize() );
 		}
+
+		StaticMeshDerivedDataKey = StaticMesh->RenderData->DerivedDataKey;
 	}
+
+#endif // WITH_EDITOR
 }
 
 float GKeepPreCulledIndicesThreshold = .95f;
diff --git a/Engine/Source/Runtime/Engine/Private/DialogueWave.cpp b/Engine/Source/Runtime/Engine/Private/DialogueWave.cpp
--- a/Engine/Source/Runtime/Engine/Private/DialogueWave.cpp
+++ b/Engine/Source/Runtime/Engine/Private/DialogueWave.cpp
@@ -2,15 +2,17 @@
 
 #include "EnginePrivate.h"
 #include "ActiveSound.h"
+#include "Sound/AudioSettings.h"
 #include "Sound/DialogueWave.h"
 #include "Sound/DialogueSoundWaveProxy.h"
+#include "Sound/DialogueVoice.h"
 #include "Sound/SoundWave.h"
 #include "SubtitleManager.h"
 #include "GatherableTextData.h"
 #include "InternationalizationMetadata.h"
-#include "Sound/DialogueVoice.h"
 
 const FString FDialogueConstants::DialogueNamespace						= TEXT("Dialogue");
+const FString FDialogueConstants::SubtitleKeySuffix						= TEXT("_Subtitle");
 #if WITH_EDITORONLY_DATA
 const FString FDialogueConstants::PropertyName_AudioFile				= TEXT("AudioFile");
 const FString FDialogueConstants::PropertyName_VoiceActorDirection		= TEXT("VoiceActorDirection");
@@ -307,14 +309,12 @@ namespace
 		FDialogueHelper DialogueHelper;
 		if (DialogueHelper.ProcessDialogueWave(DialogueWave))
 		{
-			const FString& SpokenSource = DialogueWave->SpokenText;
-
-			if (!SpokenSource.IsEmpty())
+			auto FindOrAddDialogueTextData = [&PropertyLocalizationDataGatherer](const FString& InText) -> FGatherableTextData&
 			{
+				check(!InText.IsEmpty());
+
 				FTextSourceData SourceData;
-				{
-					SourceData.SourceString = SpokenSource;
-				}
+				SourceData.SourceString = InText;
 
 				auto& GatherableTextDataArray = PropertyLocalizationDataGatherer.GetGatherableTextDataArray();
 				FGatherableTextData* GatherableTextData = GatherableTextDataArray.FindByPredicate([&](const FGatherableTextData& Candidate)
@@ -330,12 +330,33 @@ namespace
 					GatherableTextData->SourceData = SourceData;
 				}
 
+				return *GatherableTextData;
+			};
+
+			// Gather the Spoken Text for each context
+			if (!DialogueWave->SpokenText.IsEmpty())
+			{
+				FGatherableTextData& GatherableTextData = FindOrAddDialogueTextData(DialogueWave->SpokenText);
+
+				const TArray<FTextSourceSiteContext>& Variations = DialogueHelper.GetContextSpecificVariations();
+				for (const FTextSourceSiteContext& Variation : Variations)
 				{
-					const TArray<FTextSourceSiteContext>& Variations = DialogueHelper.GetContextSpecificVariations();
-					for (const FTextSourceSiteContext& Variation : Variations)
-					{
-						GatherableTextData->SourceSiteContexts.Add(Variation);
-					}
+					GatherableTextData.SourceSiteContexts.Add(Variation);
+				}
+			}
+
+			// Gather the Subtitle Override for each context
+			if (!DialogueWave->SubtitleOverride.IsEmpty())
+			{
+				FGatherableTextData& GatherableTextData = FindOrAddDialogueTextData(DialogueWave->SubtitleOverride);
+
+				const TArray<FTextSourceSiteContext>& Variations = DialogueHelper.GetContextSpecificVariations();
+				for (const FTextSourceSiteContext& Variation : Variations)
+				{
+					FTextSourceSiteContext SubtitleVariation = Variation;
+					SubtitleVariation.KeyName += FDialogueConstants::SubtitleKeySuffix;
+					SubtitleVariation.InfoMetaData.RemoveField(FDialogueConstants::PropertyName_AudioFile);
+					GatherableTextData.SourceSiteContexts.Add(SubtitleVariation);
 				}
 			}
 		}
@@ -363,12 +384,16 @@ FDialogueContextMapping::FDialogueContextMapping()
 
 }
 
-FString FDialogueContextMapping::GetLocalizationKey(const FString& InOwnerDialogueWaveKey) const
+FString FDialogueContextMapping::GetLocalizationKey() const
 {
 	TMap<FString, FStringFormatArg> Args;
 	Args.Add(TEXT("ContextHash"), Context.GetContextHash());
-	const FString ContextKey = FString::Format(*LocalizationKeyFormat, Args);
+	return FString::Format(*LocalizationKeyFormat, Args);
+}
 
+FString FDialogueContextMapping::GetLocalizationKey(const FString& InOwnerDialogueWaveKey) const
+{
+	const FString ContextKey = GetLocalizationKey();
 	return FString::Printf(TEXT("%s_%s"), *InOwnerDialogueWaveKey, *ContextKey);
 }
 
@@ -460,6 +485,7 @@ UDialogueWave::UDialogueWave(const FObjectInitializer& ObjectInitializer)
 	} AutomaticRegistrationOfLocalizationGatherer;
 #endif
 
+	bOverride_SubtitleOverride = false;
 	ContextMappings.Add(FDialogueContextMapping());
 }
 
@@ -615,7 +641,33 @@ FString UDialogueWave::GetContextRecordedAudioFilename(const FDialogueContext& C
 
 FString UDialogueWave::GetContextRecordedAudioFilename(const FDialogueContextMapping& ContextMapping) const
 {
-	return FString::Printf(TEXT("%s.wav"), *ContextMapping.GetLocalizationKey(LocalizationGUID.ToString()));
+	const UAudioSettings* AudioSettings = GetDefault<UAudioSettings>();
+	const FString DialogueContextFilename = BuildRecordedAudioFilename(AudioSettings->DialogueFilenameFormat, LocalizationGUID, GetName(), ContextMapping.GetLocalizationKey(), ContextMappings.IndexOfByKey(ContextMapping));
+	return FString::Printf(TEXT("%s.wav"), *DialogueContextFilename);
+}
+
+FString UDialogueWave::BuildRecordedAudioFilename(const FString& FormatString, const FGuid& DialogueGuid, const FString& DialogueName, const FString& ContextId, const int32 ContextIndex)
+{
+	static const FString DialogueWaveSuffix = TEXT("_DialogueWave");
+	static const FString DialogueSuffix = TEXT("_Dialogue");
+
+	const FString DialogueHash = FString::Printf(TEXT("%08X"), FCrc::MemCrc32(&DialogueGuid, sizeof(FGuid)));
+
+	// Trim the asset name if it ends with a common leaf
+	FString TrimmedDialogueName = DialogueName;
+	if (!TrimmedDialogueName.RemoveFromEnd(DialogueWaveSuffix))
+	{
+		TrimmedDialogueName.RemoveFromEnd(DialogueSuffix);
+	}
+
+	TMap<FString, FStringFormatArg> Args;
+	Args.Add(TEXT("DialogueGuid"), DialogueGuid.ToString());
+	Args.Add(TEXT("DialogueHash"), DialogueHash);
+	Args.Add(TEXT("DialogueName"), TrimmedDialogueName);
+	Args.Add(TEXT("ContextId"), ContextId);
+	Args.Add(TEXT("ContextIndex"), ContextIndex);
+	
+	return FString::Format(*FormatString, Args);
 }
 
 // End UDialogueWave interface.
@@ -650,15 +702,27 @@ void UDialogueWave::UpdateMappingProxy(FDialogueContextMapping& ContextMapping)
 		UEngine::CopyPropertiesForUnrelatedObjects(ContextMapping.SoundWave, ContextMapping.Proxy);
 
 		FSubtitleCue NewSubtitleCue;
+
+		// Do we have a subtitle override?
 		FString Key = GetContextLocalizationKey(ContextMapping);
+		if (bOverride_SubtitleOverride)
+		{
+			Key += FDialogueConstants::SubtitleKeySuffix;
+		}
 
-		if (!(FText::FindText(FDialogueConstants::DialogueNamespace, Key, NewSubtitleCue.Text)))
+		// First try and find a context specific localization
+		if (!FText::FindText(FDialogueConstants::DialogueNamespace, Key, /*OUT*/NewSubtitleCue.Text))
 		{
+			// Failing that, try and find a general dialogue wave localization
 			Key = LocalizationGUID.ToString();
+			if (bOverride_SubtitleOverride)
+			{
+				Key += FDialogueConstants::SubtitleKeySuffix;
+			}
 
 			if (!FText::FindText(FDialogueConstants::DialogueNamespace, Key, /*OUT*/NewSubtitleCue.Text))
 			{
-				NewSubtitleCue.Text = FText::FromString(SpokenText);
+				NewSubtitleCue.Text = bOverride_SubtitleOverride ? FText::FromString(SubtitleOverride) : FText::FromString(SpokenText);
 			}
 		}
 		NewSubtitleCue.Time = 0.0f;
diff --git a/Engine/Source/Runtime/Engine/Private/EdGraph/EdGraphNode.cpp b/Engine/Source/Runtime/Engine/Private/EdGraph/EdGraphNode.cpp
--- a/Engine/Source/Runtime/Engine/Private/EdGraph/EdGraphNode.cpp
+++ b/Engine/Source/Runtime/Engine/Private/EdGraph/EdGraphNode.cpp
@@ -10,6 +10,7 @@
 #include "ScopedTransaction.h"
 #include "Editor/UnrealEd/Public/Kismet2/Kismet2NameValidators.h"
 #include "Editor/Kismet/Public/FindInBlueprintManager.h"
+#include "EditorStyle.h"
 #endif
 
 #define LOCTEXT_NAMESPACE "EdGraph"
@@ -220,6 +221,20 @@ FString UEdGraphNode::GetDocumentationExcerptName() const
 	return FString::Printf(TEXT("%s%s"), MyClass->GetPrefixCPP(), *MyClass->GetName());
 }
 
+PRAGMA_DISABLE_DEPRECATION_WARNINGS
+FSlateIcon UEdGraphNode::GetIconAndTint(FLinearColor& OutColor) const
+{
+	// @todo: Remove with GetPaletteIcon
+	FName DeprecatedName = GetPaletteIcon(OutColor);
+	if (!DeprecatedName.IsNone())
+	{
+		return FSlateIcon("EditorStyle", DeprecatedName);
+	}
+	
+	static const FSlateIcon Icon = FSlateIcon("EditorStyle", "GraphEditor.Default_16x");
+	return Icon;
+}
+PRAGMA_ENABLE_DEPRECATION_WARNINGS
 
 FString UEdGraphNode::GetDescriptiveCompiledName() const
 {
@@ -410,7 +425,9 @@ void UEdGraphNode::AddSearchMetaDataInfo(TArray<struct FSearchTagDataPair>& OutT
 
 	// Non-Searchable - Used to display the icon and color for this node for better visual identification.
 	FLinearColor GlyphColor = FLinearColor::White;
-	OutTaggedMetaData.Add(FSearchTagDataPair(FFindInBlueprintSearchTags::FiB_Glyph, FText::FromString(GetPaletteIcon(GlyphColor).ToString())));
+	FSlateIcon Icon = GetIconAndTint(GlyphColor);
+	OutTaggedMetaData.Add(FSearchTagDataPair(FFindInBlueprintSearchTags::FiB_Glyph, FText::FromName(Icon.GetStyleName())));
+	OutTaggedMetaData.Add(FSearchTagDataPair(FFindInBlueprintSearchTags::FiB_GlyphStyleSet, FText::FromName(Icon.GetStyleSetName())));
 	OutTaggedMetaData.Add(FSearchTagDataPair(FFindInBlueprintSearchTags::FiB_GlyphColor, FText::FromString(GlyphColor.ToString())));
 	OutTaggedMetaData.Add(FSearchTagDataPair(FFindInBlueprintSearchTags::FiB_Comment, FText::FromString(NodeComment)));
 }
diff --git a/Engine/Source/Runtime/Engine/Private/EdGraph/EdGraphNode_Documentation.cpp b/Engine/Source/Runtime/Engine/Private/EdGraph/EdGraphNode_Documentation.cpp
--- a/Engine/Source/Runtime/Engine/Private/EdGraph/EdGraphNode_Documentation.cpp
+++ b/Engine/Source/Runtime/Engine/Private/EdGraph/EdGraphNode_Documentation.cpp
@@ -50,6 +50,12 @@ FText UEdGraphNode_Documentation::GetNodeTitle( ENodeTitleType::Type TitleType )
 	return NodeTitle;
 }
 
+FSlateIcon UEdGraphNode_Documentation::GetIconAndTint(FLinearColor& OutColor) const
+{
+	static FSlateIcon Icon("EditorStyle", "GraphEditor.Documentation_16x");
+	return Icon;
+}
+
 FText UEdGraphNode_Documentation::GetPinNameOverride(const UEdGraphPin& Pin) const
 {
 	return GetNodeTitle( ENodeTitleType::ListView );
diff --git a/Engine/Source/Runtime/Engine/Private/EdGraph/EdGraphSchema.cpp b/Engine/Source/Runtime/Engine/Private/EdGraph/EdGraphSchema.cpp
--- a/Engine/Source/Runtime/Engine/Private/EdGraph/EdGraphSchema.cpp
+++ b/Engine/Source/Runtime/Engine/Private/EdGraph/EdGraphSchema.cpp
@@ -727,7 +727,7 @@ void UEdGraphSchema::GetContextMenuActions(const UEdGraph* CurrentGraph, const U
 			}
 
 			// Called by the EditableText widget when the user types a new comment for the selected node
-			static void OnNodeCommentTextCommitted(const FText& NewText, ETextCommit::Type /*CommitInfo*/, TWeakObjectPtr<UEdGraphNode> NodeWeakPtr)
+			static void OnNodeCommentTextCommitted(const FText& NewText, ETextCommit::Type CommitInfo, TWeakObjectPtr<UEdGraphNode> NodeWeakPtr)
 			{
 				// Apply the change to the selected actor
 				UEdGraphNode* SelectedNode = NodeWeakPtr.Get();
@@ -749,7 +749,12 @@ void UEdGraphSchema::GetContextMenuActions(const UEdGraph* CurrentGraph, const U
 					}
 				}
 
-				FSlateApplication::Get().DismissAllMenus();
+				// Only dismiss all menus if the text was committed via some user action.
+				// ETextCommit::Default implies that focus was switched by some other means. If this is because a submenu was opened, we don't want to close all the menus as a consequence.
+				if (CommitInfo != ETextCommit::Default)
+				{
+					FSlateApplication::Get().DismissAllMenus();
+				}
 			}
 		};
 
diff --git a/Engine/Source/Runtime/Engine/Private/EndUserSettings.cpp b/Engine/Source/Runtime/Engine/Private/EndUserSettings.cpp
--- a/Engine/Source/Runtime/Engine/Private/EndUserSettings.cpp
+++ b/Engine/Source/Runtime/Engine/Private/EndUserSettings.cpp
@@ -63,7 +63,7 @@ FText UEndUserSettings::GetAdditionalInfoUrlLabel() const
 
 void UEndUserSettings::SetSendAnonymousUsageDataToEpic(bool bEnable)
 {
-	if (bSendAnonymousUsageDataToEpic == bEnable)
+	if (bSendAnonymousUsageDataToEpic != bEnable)
 	{
 		bSendAnonymousUsageDataToEpic = bEnable;
 		OnSendAnonymousUsageDataToEpicChanged();
diff --git a/Engine/Source/Runtime/Engine/Private/GameEngine.cpp b/Engine/Source/Runtime/Engine/Private/GameEngine.cpp
--- a/Engine/Source/Runtime/Engine/Private/GameEngine.cpp
+++ b/Engine/Source/Runtime/Engine/Private/GameEngine.cpp
@@ -571,14 +571,18 @@ void UGameEngine::Init(IEngineLoop* InEngineLoop)
 		UGameViewportClient::OnViewportCreated().Broadcast();
 	}
 
-	GameInstance->StartGameInstance();
-
 	UE_LOG(LogInit, Display, TEXT("Game Engine Initialized.") );
 
 	// for IsInitialized()
 	bIsInitialized = true;
 }
 
+void UGameEngine::Start()
+{
+	UE_LOG(LogInit, Display, TEXT("Starting Game."));
+
+	GameInstance->StartGameInstance();
+}
 
 void UGameEngine::PreExit()
 {
