Fix nightly build failure due to trying to run static analysis on Mac.

diff --git a/Engine/Source/Programs/AutomationTool/GUBP/LegacyBranchSetup.cs b/Engine/Source/Programs/AutomationTool/GUBP/LegacyBranchSetup.cs
--- a/Engine/Source/Programs/AutomationTool/GUBP/LegacyBranchSetup.cs
+++ b/Engine/Source/Programs/AutomationTool/GUBP/LegacyBranchSetup.cs
@@ -628,7 +628,7 @@ partial class GUBP
 				BranchConfig.AddNode(new NonUnityTestNode(BranchConfig, HostPlatform));
 			}
 
-			if (!BranchOptions.ExcludePlatformsForEditor.Contains(HostPlatform) && !BranchOptions.ExcludeNodes.Contains("StaticAnalysis"))
+			if (HostPlatform == UnrealTargetPlatform.Win64 && !BranchOptions.ExcludePlatformsForEditor.Contains(HostPlatform) && !BranchOptions.ExcludeNodes.Contains("StaticAnalysis"))
 			{
 				BranchConfig.AddNode(new StaticAnalysisTestNode(BranchConfig, HostPlatform));
 			}
