- Fixed C# warning as error in mono

diff --git a/Engine/Source/Programs/AutomationTool/GUBP/LegacyNodes.cs b/Engine/Source/Programs/AutomationTool/GUBP/LegacyNodes.cs
--- a/Engine/Source/Programs/AutomationTool/GUBP/LegacyNodes.cs
+++ b/Engine/Source/Programs/AutomationTool/GUBP/LegacyNodes.cs
@@ -2537,12 +2537,9 @@ partial class GUBP
 
     public class StaticAnalysisTestNode : TestNode
     {
-		GUBP.GUBPBranchConfig BranchConfig;
-
         public StaticAnalysisTestNode(GUBP.GUBPBranchConfig InBranchConfig, UnrealTargetPlatform InHostPlatform)
             : base(InHostPlatform)
         {
-			BranchConfig = InBranchConfig;
             AddDependency(ToolsForCompileNode.StaticGetFullName(HostPlatform));
             AddPseudodependency(RootEditorNode.StaticGetFullName(HostPlatform));
 			AgentSharingGroup = "TargetPlatforms" + StaticGetHostPlatformSuffix(InHostPlatform);
