Prevent platform QA running creating sample builds from saturating the //UE4 build farm.

diff --git a/Engine/Source/Programs/AutomationTool/GUBP/LegacyNodes.cs b/Engine/Source/Programs/AutomationTool/GUBP/LegacyNodes.cs
--- a/Engine/Source/Programs/AutomationTool/GUBP/LegacyNodes.cs
+++ b/Engine/Source/Programs/AutomationTool/GUBP/LegacyNodes.cs
@@ -1144,9 +1144,16 @@ partial class GUBP
 				bBehindTrigger = true;
                 //AgentSharingGroup = "TemplateMonolithics" + StaticGetHostPlatformSuffix(InHostPlatform);
             }
-			if(!InBranchConfig.BranchOptions.bTargetPlatformsInParallel && !bBehindTrigger)
+			if(!InBranchConfig.BranchOptions.bTargetPlatformsInParallel)
 			{
-				AgentSharingGroup = "TargetPlatforms" + StaticGetHostPlatformSuffix(InHostPlatform);
+				if(bBehindTrigger)
+				{
+					AgentSharingGroup = "TargetPlatformsTest" + StaticGetHostPlatformSuffix(InHostPlatform);
+				}
+				else
+				{
+					AgentSharingGroup = "TargetPlatforms" + StaticGetHostPlatformSuffix(InHostPlatform);
+				}
 			}
         }
 
