Fix failures to build "Editor, Tools & Monolithics" in //UE4/Main build process due to missing nodes in branch setup.

diff --git a/Engine/Source/Programs/AutomationTool/GUBP/LegacyBranchSetup.cs b/Engine/Source/Programs/AutomationTool/GUBP/LegacyBranchSetup.cs
--- a/Engine/Source/Programs/AutomationTool/GUBP/LegacyBranchSetup.cs
+++ b/Engine/Source/Programs/AutomationTool/GUBP/LegacyBranchSetup.cs
@@ -1144,7 +1144,7 @@ partial class GUBP
 			BranchConfig.AddNode(new BuildEngineLocalizationNode(BranchOptions.EngineLocalizationBranchSuffix));
 		}
 
-		if(BranchConfig.BranchName != null && BranchConfig.BranchName.StartsWith("//UE4/Dev-", StringComparison.InvariantCultureIgnoreCase))
+		if(BranchConfig.BranchName != null && (BranchConfig.BranchName.StartsWith("//UE4/Dev-", StringComparison.InvariantCultureIgnoreCase) || ParseParam("WithSingleTargetNodes")))
 		{
 			List<BranchInfo.BranchUProject> CodeProjects = new List<BranchInfo.BranchUProject>();
 			CodeProjects.Add(BranchConfig.Branch.BaseEngineProject);
