GUBP: Don't use UE4Build to compile the APIDocTool solution; it requires a non-default platform (x64). Fixes ToolsForDocumentation compile.

diff --git a/Engine/Source/Programs/AutomationTool/Scripts/BuildDocumentation.Automation.cs b/Engine/Source/Programs/AutomationTool/Scripts/BuildDocumentation.Automation.cs
--- a/Engine/Source/Programs/AutomationTool/Scripts/BuildDocumentation.Automation.cs
+++ b/Engine/Source/Programs/AutomationTool/Scripts/BuildDocumentation.Automation.cs
@@ -8,10 +8,10 @@ using System.Reflection;
 using AutomationTool;
 using UnrealBuildTool;
 
-class ToolsForDocumentationNode : GUBP.CompileNode
+class ToolsForDocumentationNode : GUBP.HostPlatformNode
 {
     public ToolsForDocumentationNode(GUBP.GUBPBranchConfig InBranchConfig, UnrealTargetPlatform InHostPlatform)
-        : base(InBranchConfig, InHostPlatform, false)
+        : base(InHostPlatform)
     {
 		AgentSharingGroup = "Documentation" + StaticGetHostPlatformSuffix(InHostPlatform);
     }
@@ -36,13 +36,14 @@ class ToolsForDocumentationNode : GUBP.CompileNode
 		return base.CISFrequencyQuantumShift(BranchConfig) + 6;
 	}
 
-	public override UE4Build.BuildAgenda GetAgenda(GUBP bp)
-    {
-        var Agenda = new UE4Build.BuildAgenda();
-		Agenda.DotNetSolutions.Add("Engine/Source/Programs/UnrealDocTool/APIDocTool/APIDocTool.sln");
-		Agenda.DotNetSolutions.Add("Engine/Source/Programs/UnrealDocTool/UnrealDocTool/UnrealDocTool.sln");
-		return Agenda;
-    }
+	public override void DoBuild(GUBP bp)
+	{
+		CommandUtils.BuildSolution(CommandUtils.CmdEnv, "Engine/Source/Programs/UnrealDocTool/APIDocTool/APIDocTool.sln", "Development", "x64");
+		CommandUtils.BuildSolution(CommandUtils.CmdEnv, "Engine/Source/Programs/UnrealDocTool/UnrealDocTool/UnrealDocTool.sln", "Development", "Any CPU");
+
+		BuildProducts = new List<string>();
+		SaveRecordOfSuccessAndAddToBuildProducts();
+	}
 }
 	
 class DocumentationNode : GUBP.HostPlatformNode
