Rollback to revision 1. Fixes version files not being updated when building for UGS.

diff --git a/Engine/Source/Programs/AutomationTool/Scripts/BuildForUGS.Automation.cs b/Engine/Source/Programs/AutomationTool/Scripts/BuildForUGS.Automation.cs
--- a/Engine/Source/Programs/AutomationTool/Scripts/BuildForUGS.Automation.cs
+++ b/Engine/Source/Programs/AutomationTool/Scripts/BuildForUGS.Automation.cs
@@ -47,7 +47,7 @@ namespace AutomationTool
 
 			// Build everything
 			UE4Build Builder = new UE4Build(this);
-			Builder.Build(Agenda, InDeleteBuildProducts: false, InUpdateVersionFiles: false, InForceNoXGE: true, InForceNonUnity: true);
+			Builder.Build(Agenda, InDeleteBuildProducts: true, InUpdateVersionFiles: true, InForceNoXGE: true);
 
 			// Include the build products for UAT and UBT if required
 			if(ParseParam("WithUAT"))
