UAT: Remove forced arguments for BuildForUGS script. Games now specify other parameters via command-line arguments in their individual EC jobs.

diff --git a/Engine/Source/Programs/AutomationTool/Scripts/BuildForUGS.Automation.cs b/Engine/Source/Programs/AutomationTool/Scripts/BuildForUGS.Automation.cs
--- a/Engine/Source/Programs/AutomationTool/Scripts/BuildForUGS.Automation.cs
+++ b/Engine/Source/Programs/AutomationTool/Scripts/BuildForUGS.Automation.cs
@@ -47,7 +47,7 @@ namespace AutomationTool
 
 			// Build everything
 			UE4Build Builder = new UE4Build(this);
-			Builder.Build(Agenda, InDeleteBuildProducts: true, InUpdateVersionFiles: true, InForceNoXGE: true);
+			Builder.Build(Agenda, InUpdateVersionFiles: ArchivePath != null);
 
 			// Include the build products for UAT and UBT if required
 			if(ParseParam("WithUAT"))
