Don't fail due to signing errors until we've tried all available timestamp servers AND we've passed our own self-imposed timeout - unreachable servers can be slow to fail as it is.

diff --git a/Engine/Source/Programs/AutomationTool/AutomationUtils/CommandUtils.cs b/Engine/Source/Programs/AutomationTool/AutomationUtils/CommandUtils.cs
--- a/Engine/Source/Programs/AutomationTool/AutomationUtils/CommandUtils.cs
+++ b/Engine/Source/Programs/AutomationTool/AutomationUtils/CommandUtils.cs
@@ -2755,11 +2755,10 @@ namespace AutomationTool
 										 "http://timestamp.comodoca.com/authenticode",
 										 "http://www.startssl.com/timestamp"
 									   };
-			int TimestampServerIndex = 0;
 
 			string SpecificStoreArg = bUseMachineStoreInsteadOfUserStore ? " /sm" : "";	
 			
-			DateTime StartTime = DateTime.Now;
+			Stopwatch Stopwatch = Stopwatch.StartNew();
 
 			int NumTrials = 0;
 			for (; ; )
@@ -2767,7 +2766,7 @@ namespace AutomationTool
 				//@TODO: Verbosity choosing
 				//  /v will spew lots of info
 				//  /q does nothing on success and minimal output on failure
-				string CodeSignArgs = String.Format("sign{0} /a /n \"{1}\" /t {2} /v {3}", SpecificStoreArg, SigningIdentity, TimestampServer[TimestampServerIndex], FilesToSign);
+				string CodeSignArgs = String.Format("sign{0} /a /n \"{1}\" /t {2} /v {3}", SpecificStoreArg, SigningIdentity, TimestampServer[NumTrials % TimestampServer.Length], FilesToSign);
 
 				ProcessResult Result = CommandUtils.Run(SignToolName, CodeSignArgs, null, CommandUtils.ERunOptions.AllowSpew);
 				++NumTrials;
@@ -2783,17 +2782,9 @@ namespace AutomationTool
 				}
 				else
 				{
-					// try another timestamp server on the next iteration
-					TimestampServerIndex++;
-					if (TimestampServerIndex >= TimestampServer.Count())
-					{
-						// loop back to the first timestamp server
-						TimestampServerIndex = 0;
-					}
-					
 					// Keep retrying until we run out of time
-					TimeSpan RunTime = DateTime.Now - StartTime;
-					if (RunTime > CodeSignTimeOut)
+					TimeSpan RunTime = Stopwatch.Elapsed;
+					if (RunTime > CodeSignTimeOut && NumTrials >= TimestampServer.Length)
 					{
 						throw new AutomationException("Failed to sign executables {0} times over a period of {1}", NumTrials, RunTime);
 					}
