Remove hack forcing CrashReportClient to be built with Visual Studio 2013. We don't want to ship this code with base UE4, so modify build scripts for games that ship on 2013 to pass the -2013 argument to UBT instead.

diff --git a/Engine/Source/Programs/AutomationTool/GUBP/LegacyNodes.cs b/Engine/Source/Programs/AutomationTool/GUBP/LegacyNodes.cs
--- a/Engine/Source/Programs/AutomationTool/GUBP/LegacyNodes.cs
+++ b/Engine/Source/Programs/AutomationTool/GUBP/LegacyNodes.cs
@@ -637,11 +637,17 @@ partial class GUBP
                 {
 					if (!BranchConfig.BranchOptions.ExcludeNodes.Contains(ProgramTarget.TargetName))
                     {
+						// Always compile CRC with VS2013 due to requiring the 2015 redist
+						string AddTargetArgs = AddArgs;
+						if(String.Compare(ProgramTarget.TargetName, "CrashReportClient", true) == 0)
+						{
+							AddTargetArgs += " -2013";
+						}
 						foreach (var Plat in ProgramTarget.Rules.GUBP_ToolPlatforms(HostPlatform))
 						{
 							foreach (var Config in ProgramTarget.Rules.GUBP_ToolConfigs(HostPlatform))
 							{
-								Agenda.AddTargets(new string[] { ProgramTarget.TargetName }, Plat, Config, InAddArgs: AddArgs);
+								Agenda.AddTargets(new string[] { ProgramTarget.TargetName }, Plat, Config, InAddArgs: AddTargetArgs);
 							}
 						}
 					}
diff --git a/Engine/Source/Programs/CrashReporter/CrashReportClient/CrashReportClient.Target.cs b/Engine/Source/Programs/CrashReporter/CrashReportClient/CrashReportClient.Target.cs
--- a/Engine/Source/Programs/CrashReporter/CrashReportClient/CrashReportClient.Target.cs
+++ b/Engine/Source/Programs/CrashReporter/CrashReportClient/CrashReportClient.Target.cs
@@ -23,15 +23,6 @@ public class CrashReportClientTarget : TargetRules
 		return true;
 	}
 
-    public override bool ConfigureToolchain(TargetInfo Target)
-    {
-        if (Target.Platform == UnrealTargetPlatform.Win32 || Target.Platform == UnrealTargetPlatform.Win64)
-        {
-            WindowsPlatform.Compiler = WindowsCompiler.VisualStudio2013;
-        }
-        return true;
-    }
-
     public override bool GetSupportedConfigurations(ref List<UnrealTargetConfiguration> OutConfigurations, bool bIncludeTestAndShippingConfigs)
 	{
 		if( base.GetSupportedConfigurations( ref OutConfigurations, bIncludeTestAndShippingConfigs ) )
