Fix SwarmAgent.exe missing from installed UE4 build; non-console applications were not enumerating build products.

diff --git a/Engine/Source/Programs/AutomationTool/BuildGraph/Tasks/CsCompileTask.cs b/Engine/Source/Programs/AutomationTool/BuildGraph/Tasks/CsCompileTask.cs
--- a/Engine/Source/Programs/AutomationTool/BuildGraph/Tasks/CsCompileTask.cs
+++ b/Engine/Source/Programs/AutomationTool/BuildGraph/Tasks/CsCompileTask.cs
@@ -370,6 +370,7 @@ namespace AutomationTool.Tasks
 				switch(OutputType)
 				{
 					case "Exe":
+					case "WinExe":
 						BuildProducts.Add(FileReference.Combine(OutputDir, AssemblyName + ".exe"));
 						BuildProducts.Add(FileReference.Combine(OutputDir, AssemblyName + ".pdb"));
                         FileReference ExeConfig = FileReference.Combine(OutputDir, AssemblyName + ".exe.config");
