Removed extra '#' character from filespec for stripping of UE4Game on IOS
Added validation of tag names to FindOrAddTagSet so that they only contain a single '#' character and any spaces come after the second character

diff --git a/Engine/Build/InstalledEngineBuild.xml b/Engine/Build/InstalledEngineBuild.xml
--- a/Engine/Build/InstalledEngineBuild.xml
+++ b/Engine/Build/InstalledEngineBuild.xml
@@ -191,7 +191,7 @@
 			<Compile Target="UE4Game" Platform="IOS" Configuration="Shipping" Tag="#UE4Game IOS" Arguments="-precompile"/>
 
 			<Tag Files="#UE4Game IOS" Filter="$(IOSStripFilter)" With="#UE4Game IOS Unstripped"/>
-			<Strip BaseDir="$(RootDir)" OutputDir="$(SavedOutput)" Platform="IOS" Files="##UE4Game IOS Unstripped" Tag="#UE4Game IOS Stripped"/>
+			<Strip BaseDir="$(RootDir)" OutputDir="$(SavedOutput)" Platform="IOS" Files="#UE4Game IOS Unstripped" Tag="#UE4Game IOS Stripped"/>
 
 			<!-- Tag the generated includes for this target -->
 			<Tag Files="Engine/Intermediate/Build/IOS/UE4/Inc/...;Engine/Plugins/.../Intermediate/Build/IOS/UE4/Inc/..." With="#UE4Game IOS Includes"/>
diff --git a/Engine/Source/Programs/AutomationTool/BuildGraph/Task.cs b/Engine/Source/Programs/AutomationTool/BuildGraph/Task.cs
--- a/Engine/Source/Programs/AutomationTool/BuildGraph/Task.cs
+++ b/Engine/Source/Programs/AutomationTool/BuildGraph/Task.cs
@@ -252,14 +252,20 @@ namespace AutomationTool
 		/// <summary>
 		/// Finds or adds a set containing files with the given tag
 		/// </summary>
-		/// <param name="Name">The tag name to return a set for. An leading '#' character is optional.</param>
+		/// <param name="TagName">The tag name to return a set for. A leading '#' character is required.</param>
 		/// <returns>Set of files</returns>
 		public static HashSet<FileReference> FindOrAddTagSet(Dictionary<string, HashSet<FileReference>> TagNameToFileSet, string TagName)
 		{
-			// Make sure the tag name contains a leading hash
-			if(!TagName.StartsWith("#"))
+			// Make sure the tag name contains a single leading hash
+			if (TagName.LastIndexOf('#') != 0)
 			{
-				throw new AutomationException("Tag name does not start with a '#' character");
+				throw new AutomationException("Tag name '{0}' is not valid - should contain a single leading '#' character", TagName);
+			}
+
+			// Any spaces should be later than the second char - most likely to be a typo if directly after the # character
+			if (TagName.IndexOf(' ') == 1)
+			{
+				throw new AutomationException("Tag name '{0}' is not valid - spaces should only be used to separate words", TagName);
 			}
 
 			// Find the files which match this tag
