Merging CL 3071691 from //UE4/Dev-Build: Fix UAT failure if uebp_CodeCL is not set

diff --git a/Engine/Source/Programs/AutomationTool/AutomationUtils/P4Environment.cs b/Engine/Source/Programs/AutomationTool/AutomationUtils/P4Environment.cs
--- a/Engine/Source/Programs/AutomationTool/AutomationUtils/P4Environment.cs
+++ b/Engine/Source/Programs/AutomationTool/AutomationUtils/P4Environment.cs
@@ -104,7 +104,7 @@ namespace AutomationTool
 			LabelToSync = CommandUtils.GetEnvVar(EnvVarNames.LabelToSync);
 
 			string CodeChangelistString = CommandUtils.GetEnvVar(EnvVarNames.CodeChangelist);
-			if(CodeChangelistString != null)
+			if(!String.IsNullOrEmpty(CodeChangelistString))
 			{
 				CodeChangelist = Int32.Parse(CodeChangelistString);
 			}
