Copying //UE4/Dev-Build to //UE4/Dev-Main (Source: //UE4/Dev-Build @ 3092544)

==========================
MAJOR FEATURES + CHANGES
==========================

Change 3079316 on 2016/08/05 by [REDACTED_PERSON]

	Better PCH selection in ShaderFormatOpenGL and MetalShaderFormat - make sure Core is the first included header.

Change 3080579 on 2016/08/08 by [REDACTED_PERSON]

	Slate: Move DEBUG_TAB_MANAGEMENT into SDockingTabStack.h to remove circular include dependency with DockingPrivate.h.

Change 3080587 on 2016/08/08 by [REDACTED_PERSON]

	StandaloneRenderer: Move platform includes into a separate header so we can make individual headers self-contained, without having a circular dependency on StandaloneRendererPrivate.h.

Change 3080789 on 2016/08/08 by [REDACTED_PERSON]

	Move BuildGraph tasks for chunking, posting builds, labeling builds, and merging manifests into the MCP project. While we do provide public interfaces for this functionality, we don't currently expect anyone outside Epic to be using them.

Change 3080815 on 2016/08/08 by [REDACTED_PERSON]

	BuildGraph: Add a -GenerateDocs option, which writes out an HTML file constructed from C# XML documentation containing all the task information.

Change 3081374 on 2016/08/08 by [REDACTED_PERSON]

	UBT: Invalidate the makefile if any UHT headers are deleted. Should fix issue where files are moved from one module to another, and the original module no longer contains any generated headers. Its include path needs to be removed from the compile environment.

Change 3083152 on 2016/08/09 by [REDACTED_PERSON]

	PR #2667: Add Intel C++ Compiler support to Windows build (Contributed by [REDACTED_PERSON])

Change 3084039 on 2016/08/10 by [REDACTED_PERSON]

	BuildGraph: Add additional markup for parameter attributes. Also improve some documation.

Change 3084240 on 2016/08/10 by [REDACTED_PERSON]

	Plugins: Allow plugins in the project folder to replace plugins in the engine folder with the same name. Prohibit multiple plugins with the same name at other times.

Change 3084337 on 2016/08/10 by [REDACTED_PERSON]

	UBT: Specify the -precompile option when generating project files for a target, so we include all valid modules for intellisense.

Change 3085594 on 2016/08/11 by [REDACTED_PERSON]

	Change modules which reference a public header for their PCH to use a private PCH instead, even if it just includes the public header for now.

Change 3085999 on 2016/08/11 by [REDACTED_PERSON]

	Add some missing #pragma once directives.

Change 3086146 on 2016/08/11 by [REDACTED_PERSON]

	Core: Move prototype and linkage specifier for ConsoleCommandLibrary_* functions into header matching cpp file.

Change 3086172 on 2016/08/11 by [REDACTED_PERSON]

	Fixup some C-style header guards to use #pragma once instead.

Change 3087289 on 2016/08/12 by [REDACTED_PERSON]

	Split out UPackage and UMetaData into their own headers (they're already implemented in separate CPP files)

Change 3087310 on 2016/08/12 by [REDACTED_PERSON]

	Move method stubs for FNullSlateSoundDevice into a CPP file, since they're exported from the SlateCore module.

Change 3087341 on 2016/08/12 by [REDACTED_PERSON]

	UdpMessaging: Move PCH before #if PLATFORM_DESKTOP; it will only be defined if the definition is included.

Change 3087457 on 2016/08/12 by [REDACTED_PERSON]

	Core: Reorganize the FTransform and FMatrix headers: Transform.h now includes TransformNonVectorized.h or TransformVectorized.h as appropriate, and UnrealMatrix.h is now Matrix.inl (and included from Matrix.h).

Change 3088407 on 2016/08/13 by [REDACTED_PERSON]

	Replace use of Windows SIZE_T define with the regular C++ size_t.

Change 3088416 on 2016/08/13 by [REDACTED_PERSON]

	Include a header from all .generated.cpp files (GeneratedCppIncludes.h) which includes all the basic types required to compile them, rather than assuming that the module PCH will include everything.

	Also include the real declarations of noexport classes in Object.h (now renamed to NoExportTypes.h for clarity) when the CPP macro is defined, so the .generated.deps.h file will automatically have the correct definitions for them at compile time rather than relying on them being in the private PCH.

	Finally, rename UObject.h to Object.h for consistency with the naming convention for all other UObject classes. UObject.h still exists for now, but outputs a deprecated message if included.

Change 3088544 on 2016/08/14 by [REDACTED_PERSON]

	Core: Move the definition of the TEXT() macro into Platform.h, to avoid having to include OS headers to get it.

Change 3088552 on 2016/08/14 by [REDACTED_PERSON]

	Fix compile errors for some modules that don't already include CoreUObject.h.

Change 3088925 on 2016/08/15 by [REDACTED_PERSON]

	Remove circular include dependencies from VulkanRHI.

Change 3088926 on 2016/08/15 by [REDACTED_PERSON]

	Remove duplicate definition for WITH_FIXED_AREA_ENTERING_COST from EngineDefines.h - always uses the definition from DetourNavMeshQuery.h instead.

Change 3088930 on 2016/08/15 by [REDACTED_PERSON]

	Remove circular include dependency from PhysX.

Change 3088935 on 2016/08/15 by [REDACTED_PERSON]

	OnlineSubsystemUtils: Move CPP files out of public header directory.

Change 3088965 on 2016/08/15 by [REDACTED_PERSON]

	Add private PCH to Landscape, MoviePlayer, TaskGraph, XAudio2 and RealtimeProfiler modules.

Change 3088966 on 2016/08/15 by [REDACTED_PERSON]

	Engine: Move CPP files out of public header directories.

Change 3089520 on 2016/08/15 by [REDACTED_PERSON]

	BuildGraph: Change documentation command to output markdown.

Change 3090299 on 2016/08/16 by [REDACTED_PERSON]

	D3D12RHI: Move around some implementations to fix circular header dependencies which are masked by delayed template instantiation.

Change 3090303 on 2016/08/16 by [REDACTED_PERSON]

	Engine: Add a template specialization for TPointerIsConvertibleFromTo<AActor, const volatile UObject> to fix dependency on complete AActor definition for static assert in TWeakPointer<AActor>, which only appears if including Level.h without Actor.h. Delayed template instantiation usually masks this issue.

Change 3091861 on 2016/08/17 by [REDACTED_PERSON]

	Remove circular header dependencies, and fix ambiguous include paths in OSVR.

Change 3092068 on 2016/08/17 by [REDACTED_PERSON]

	Moving VulkanDynamicRHI into its own header.

Change 3093133 on 2016/08/18 by [REDACTED_PERSON]

	EC: Include additional context lines for Clang errors.

Change 3093147 on 2016/08/18 by [REDACTED_PERSON]

	UBT: Add an error message when attempting to do a single-file compile with the wrong target selected.

Change 3093228 on 2016/08/18 by [REDACTED_PERSON]

	Remove redundant setting for remote server name from XML config, and set it to a valid machine in the engine config.

diff --git a/Engine/Build/BuildFarm/PostpExtensions.pl b/Engine/Build/BuildFarm/PostpExtensions.pl
--- a/Engine/Build/BuildFarm/PostpExtensions.pl
+++ b/Engine/Build/BuildFarm/PostpExtensions.pl
@@ -83,7 +83,7 @@ unshift @::gMatchers, (
     {
         id =>               "clangError",
         pattern =>          q{([^:]+):[\d:]+ error:},
-        action =>           q{incValue("errors"); diagnostic($1, "error", backWhile(": In (member )?function"), forwardWhile("^   "))},
+        action =>           q{incValue("errors"); diagnostic($1, "error", backWhile(": In (member )?function|In file included from"), forwardWhile("^   ")) },
     },
     {
         id =>               "clangWarning",
diff --git a/Engine/Plugins/Experimental/AlembicImporter/Source/AlembicLibrary/Public/AlembicLibraryPublicPCH.h b/Engine/Plugins/Experimental/AlembicImporter/Source/AlembicLibrary/Public/AlembicLibraryPublicPCH.h
--- a/Engine/Plugins/Experimental/AlembicImporter/Source/AlembicLibrary/Public/AlembicLibraryPublicPCH.h
+++ b/Engine/Plugins/Experimental/AlembicImporter/Source/AlembicLibrary/Public/AlembicLibraryPublicPCH.h
@@ -1 +1,3 @@
 // Copyright 1998-2016 Epic Games, Inc. All Rights Reserved.
+
+#include "Core.h"
diff --git a/Engine/Plugins/Experimental/OSVR/Source/OSVR/OSVR.Build.cs b/Engine/Plugins/Experimental/OSVR/Source/OSVR/OSVR.Build.cs
--- a/Engine/Plugins/Experimental/OSVR/Source/OSVR/OSVR.Build.cs
+++ b/Engine/Plugins/Experimental/OSVR/Source/OSVR/OSVR.Build.cs
@@ -41,6 +41,7 @@ public class OSVR : ModuleRules
         if(Target.Platform == UnrealTargetPlatform.Win32 || Target.Platform == UnrealTargetPlatform.Win64)
         {
             PrivateDependencyModuleNames.AddRange(new string[] { "D3D11RHI" });
+			PrivateIncludePathModuleNames.Add("DX11");
 
             // Required for some private headers needed for the rendering support.
             var EngineDir = Path.GetFullPath(BuildConfiguration.RelativeEnginePath);
diff --git a/Engine/Plugins/Experimental/OSVR/Source/OSVRInput/OSVRInput.Build.cs b/Engine/Plugins/Experimental/OSVR/Source/OSVRInput/OSVRInput.Build.cs
--- a/Engine/Plugins/Experimental/OSVR/Source/OSVRInput/OSVRInput.Build.cs
+++ b/Engine/Plugins/Experimental/OSVR/Source/OSVRInput/OSVRInput.Build.cs
@@ -41,6 +41,7 @@ public class OSVRInput : ModuleRules
         if (Target.Platform == UnrealTargetPlatform.Win32 || Target.Platform == UnrealTargetPlatform.Win64)
         {
             PrivateDependencyModuleNames.AddRange(new string[] { "D3D11RHI" });
+			PrivateIncludePathModuleNames.Add("DX11");
 
             // Required for some private headers needed for the rendering support.
             var EngineDir = Path.GetFullPath(BuildConfiguration.RelativeEnginePath);
diff --git a/Engine/Plugins/MovieScene/LevelSequenceEditor/Source/LevelSequenceEditor/Private/LevelSequenceEditorPCH.h b/Engine/Plugins/MovieScene/LevelSequenceEditor/Source/LevelSequenceEditor/Private/LevelSequenceEditorPCH.h
--- a/Engine/Plugins/MovieScene/LevelSequenceEditor/Source/LevelSequenceEditor/Private/LevelSequenceEditorPCH.h
+++ b/Engine/Plugins/MovieScene/LevelSequenceEditor/Source/LevelSequenceEditor/Private/LevelSequenceEditorPCH.h
@@ -14,7 +14,7 @@
 #include "Engine.h"
 #include "SlateBasics.h"
 #include "SlateStyle.h"
-
+#include "Editor.h"
 
 /* Private includes
  *****************************************************************************/
diff --git a/Engine/Plugins/Online/OnlineSubsystem/Source/Public/OnlineSubsystem.h b/Engine/Plugins/Online/OnlineSubsystem/Source/Public/OnlineSubsystem.h
--- a/Engine/Plugins/Online/OnlineSubsystem/Source/Public/OnlineSubsystem.h
+++ b/Engine/Plugins/Online/OnlineSubsystem/Source/Public/OnlineSubsystem.h
@@ -177,7 +177,7 @@ public:
 	 * @param SubsystemName - Name of the requested online service
 	 * @return true if instance exists, false otherwise
 	 */
-	static bool DoesInstanceExist(const FName& SubsystemName = NAME_None)
+	FORCENOINLINE static bool DoesInstanceExist(const FName& SubsystemName = NAME_None)
 	{
 		if (FModuleManager::Get().IsModuleLoaded("OnlineSubsystem"))
 		{
diff --git a/Engine/Plugins/Online/OnlineSubsystem/Source/Public/OnlineSubsystemNames.h b/Engine/Plugins/Online/OnlineSubsystem/Source/Public/OnlineSubsystemNames.h
--- a/Engine/Plugins/Online/OnlineSubsystem/Source/Public/OnlineSubsystemNames.h
+++ b/Engine/Plugins/Online/OnlineSubsystem/Source/Public/OnlineSubsystemNames.h
@@ -1,5 +1,7 @@
 // Copyright 1998-2016 Epic Games, Inc. All Rights Reserved.
 
+#pragma once
+
 #ifndef NULL_SUBSYSTEM
 #define NULL_SUBSYSTEM FName(TEXT("NULL"))
 #endif
diff --git a/Engine/Plugins/Online/OnlineSubsystemOculus/Source/Private/OnlineMessageTaskManagerOculus.h b/Engine/Plugins/Online/OnlineSubsystemOculus/Source/Private/OnlineMessageTaskManagerOculus.h
--- a/Engine/Plugins/Online/OnlineSubsystemOculus/Source/Private/OnlineMessageTaskManagerOculus.h
+++ b/Engine/Plugins/Online/OnlineSubsystemOculus/Source/Private/OnlineMessageTaskManagerOculus.h
@@ -2,7 +2,6 @@
 
 #pragma once
 
-#include "OnlineSubsystemOculus.h"
 #include "OVR_Platform.h"
 
 DECLARE_DELEGATE_TwoParams(FOculusMessageOnCompleteDelegate, ovrMessageHandle, bool);
diff --git a/Engine/Plugins/Runtime/CustomMeshComponent/Source/CustomMeshComponent/Private/CustomMeshComponentPluginPrivatePCH.h b/Engine/Plugins/Runtime/CustomMeshComponent/Source/CustomMeshComponent/Private/CustomMeshComponentPluginPrivatePCH.h
--- a/Engine/Plugins/Runtime/CustomMeshComponent/Source/CustomMeshComponent/Private/CustomMeshComponentPluginPrivatePCH.h
+++ b/Engine/Plugins/Runtime/CustomMeshComponent/Source/CustomMeshComponent/Private/CustomMeshComponentPluginPrivatePCH.h
@@ -1,10 +1,8 @@
 // Copyright 1998-2016 Epic Games, Inc. All Rights Reserved.
 
+#pragma once
+
 #include "CoreUObject.h"
 #include "Engine.h"
-
-// You should place include statements to your module's private header files here.  You only need to
-// add includes for headers that are used in most of your module's source files though.
 #include "ICustomMeshComponentPlugin.h"
-
-#include "CustomMeshComponent.h"
\ No newline at end of file
+#include "CustomMeshComponent.h"
diff --git a/Engine/Plugins/Runtime/OculusInput/Source/Private/OculusInputState.h b/Engine/Plugins/Runtime/OculusInput/Source/Private/OculusInputState.h
--- a/Engine/Plugins/Runtime/OculusInput/Source/Private/OculusInputState.h
+++ b/Engine/Plugins/Runtime/OculusInput/Source/Private/OculusInputState.h
@@ -183,10 +183,10 @@ struct FOculusTouchControllerState
 	FVector2D ThumbstickAxes;
 
 	/** Button states */
-	FOculusButtonState Buttons[ EOculusTouchControllerButton::TotalButtonCount ];
+	FOculusButtonState Buttons[ (int32)EOculusTouchControllerButton::TotalButtonCount ];
 
 	/** Capacitive Touch axes */
-	FOculusTouchCapacitiveState CapacitiveAxes[EOculusTouchCapacitiveAxes::TotalAxisCount];
+	FOculusTouchCapacitiveState CapacitiveAxes[(int32)EOculusTouchCapacitiveAxes::TotalAxisCount];
 
 	/** Whether or not we're playing a haptic effect.  If true, force feedback calls will be early-outed in favor of the haptic effect */
 	bool bPlayingHapticEffect;
@@ -239,7 +239,7 @@ struct FOculusTouchControllerState
 struct FOculusRemoteControllerState
 {
 	/** Button states */
-	FOculusButtonState Buttons[EOculusRemoteControllerButton::TotalButtonCount];
+	FOculusButtonState Buttons[(int32)EOculusRemoteControllerButton::TotalButtonCount];
 
 	FOculusRemoteControllerState()
 	{
diff --git a/Engine/Source/Developer/Apple/MetalShaderFormat/Private/MetalBackend.cpp b/Engine/Source/Developer/Apple/MetalShaderFormat/Private/MetalBackend.cpp
--- a/Engine/Source/Developer/Apple/MetalShaderFormat/Private/MetalBackend.cpp
+++ b/Engine/Source/Developer/Apple/MetalShaderFormat/Private/MetalBackend.cpp
@@ -1,8 +1,7 @@
 // Copyright 1998-2016 Epic Games, Inc. All Rights Reserved.
-// 
 
-#include "MetalShaderFormat.h"
 #include "Core.h"
+#include "MetalShaderFormat.h"
 #include "hlslcc.h"
 #include "hlslcc_private.h"
 #include "MetalBackend.h"
diff --git a/Engine/Source/Developer/AudioFormatOpus/AudioFormatOpus.Build.cs b/Engine/Source/Developer/AudioFormatOpus/AudioFormatOpus.Build.cs
--- a/Engine/Source/Developer/AudioFormatOpus/AudioFormatOpus.Build.cs
+++ b/Engine/Source/Developer/AudioFormatOpus/AudioFormatOpus.Build.cs
@@ -26,5 +26,7 @@ public class AudioFormatOpus : ModuleRules
 				"libOpus"
 				);
 		}
+
+		Definitions.Add("WITH_OGGVORBIS=1");
 	}
 }
diff --git a/Engine/Source/Developer/FunctionalTesting/Private/FuncTestManager.cpp b/Engine/Source/Developer/FunctionalTesting/Private/FuncTestManager.cpp
--- a/Engine/Source/Developer/FunctionalTesting/Private/FuncTestManager.cpp
+++ b/Engine/Source/Developer/FunctionalTesting/Private/FuncTestManager.cpp
@@ -2,7 +2,7 @@
 
 #include "FunctionalTestingPrivatePCH.h"
 
-#include "UObject/UObject.h"
+#include "UObject/Object.h"
 #include "Engine/World.h"
 
 DEFINE_LOG_CATEGORY(LogFunctionalTest);
diff --git a/Engine/Source/Developer/GameplayDebugger/Private/GameplayDebugger.cpp b/Engine/Source/Developer/GameplayDebugger/Private/GameplayDebugger.cpp
--- a/Engine/Source/Developer/GameplayDebugger/Private/GameplayDebugger.cpp
+++ b/Engine/Source/Developer/GameplayDebugger/Private/GameplayDebugger.cpp
@@ -4,9 +4,10 @@
 // THIS CLASS IS NOW DEPRECATED AND WILL BE REMOVED IN NEXT VERSION
 // Please check GameplayDebugger.h for details.
 
+#include "GameplayDebuggerPrivatePCH.h"
+
 #if ENABLE_OLD_GAMEPLAY_DEBUGGER
 
-#include "GameplayDebuggerPrivatePCH.h"
 #include "Misc/CoreMisc.h"
 #include "GameplayDebuggingComponent.h"
 #include "GameplayDebuggingHUDComponent.h"
@@ -678,7 +679,6 @@ void FGameplayDebuggerCompat::UseNewGameplayDebugger()
 
 #else // ENABLE_OLD_GAMEPLAY_DEBUGGER
 
-#include "GameplayDebuggerPrivatePCH.h"
 #include "GameplayDebuggingReplicator.h"
 #include "GameplayDebuggerSettings.h"
 
diff --git a/Engine/Source/Developer/MessageLog/Private/MessageLogPrivatePCH.h b/Engine/Source/Developer/MessageLog/Private/MessageLogPrivatePCH.h
--- a/Engine/Source/Developer/MessageLog/Private/MessageLogPrivatePCH.h
+++ b/Engine/Source/Developer/MessageLog/Private/MessageLogPrivatePCH.h
@@ -1,7 +1,6 @@
 // Copyright 1998-2016 Epic Games, Inc. All Rights Reserved.
 
-#ifndef __MessageLogPrivatePCH_h__
-#define __MessageLogPrivatePCH_h__
+#pragma once
 
 #include "Core.h"
 #include "SlateBasics.h"
@@ -9,4 +8,3 @@
 #include "EditorStyle.h"
 #include "MessageLogModule.h"
 
-#endif // __MessageLogPrivatePCH_h__
diff --git a/Engine/Source/Developer/RealtimeProfiler/Private/RealtimeProfiler.cpp b/Engine/Source/Developer/RealtimeProfiler/Private/RealtimeProfiler.cpp
--- a/Engine/Source/Developer/RealtimeProfiler/Private/RealtimeProfiler.cpp
+++ b/Engine/Source/Developer/RealtimeProfiler/Private/RealtimeProfiler.cpp
@@ -1,5 +1,6 @@
 // Copyright 1998-2016 Epic Games, Inc. All Rights Reserved.
 
+#include "RealtimeProfilerPrivatePCH.h"
 #include "SlateBasics.h"
 #include "RealtimeProfiler.h"
 
diff --git a/Engine/Source/Developer/RealtimeProfiler/Private/RealtimeProfilerPrivatePCH.h b/Engine/Source/Developer/RealtimeProfiler/Private/RealtimeProfilerPrivatePCH.h
new file mode 100644
--- /dev/null
+++ b/Engine/Source/Developer/RealtimeProfiler/Private/RealtimeProfilerPrivatePCH.h
@@ -0,0 +1,5 @@
+// Copyright 1998-2016 Epic Games, Inc. All Rights Reserved.
+
+#pragma once
+
+#include "SlateBasics.h"
diff --git a/Engine/Source/Developer/RealtimeProfiler/Private/SRealtimeProfilerFrame.cpp b/Engine/Source/Developer/RealtimeProfiler/Private/SRealtimeProfilerFrame.cpp
--- a/Engine/Source/Developer/RealtimeProfiler/Private/SRealtimeProfilerFrame.cpp
+++ b/Engine/Source/Developer/RealtimeProfiler/Private/SRealtimeProfilerFrame.cpp
@@ -1,5 +1,6 @@
 // Copyright 1998-2016 Epic Games, Inc. All Rights Reserved.
 
+#include "RealtimeProfilerPrivatePCH.h"
 #include "SlateBasics.h"
 #include "RealtimeProfiler.h"
 #include "SRealtimeProfilerFrame.h"
diff --git a/Engine/Source/Developer/ScreenShotComparison/Private/ScreenShotComparisonPrivatePCH.h b/Engine/Source/Developer/ScreenShotComparison/Private/ScreenShotComparisonPrivatePCH.h
--- a/Engine/Source/Developer/ScreenShotComparison/Private/ScreenShotComparisonPrivatePCH.h
+++ b/Engine/Source/Developer/ScreenShotComparison/Private/ScreenShotComparisonPrivatePCH.h
@@ -4,8 +4,7 @@
 	ScreenShotComparisonPrivatePCH.h: Pre-compiled header file for the ScreenShotComparison module.
 =============================================================================*/
 
-#ifndef __ScreenShotComparisonPrivatePCH_h__
-#define __ScreenShotComparisonPrivatePCH_h__
+#pragma once
 
 #include "ScreenShotComparison.h"
 
@@ -33,4 +32,3 @@
 #include "SScreenShotBrowser.h"
 #include "SScreenShotImagePopup.h"
 
-#endif // __ScreenShotComparisonPrivatePCH_h__
diff --git a/Engine/Source/Developer/ScreenShotComparisonTools/Public/ScreenShotComparisonTools.h b/Engine/Source/Developer/ScreenShotComparisonTools/Public/ScreenShotComparisonTools.h
--- a/Engine/Source/Developer/ScreenShotComparisonTools/Public/ScreenShotComparisonTools.h
+++ b/Engine/Source/Developer/ScreenShotComparisonTools/Public/ScreenShotComparisonTools.h
@@ -15,11 +15,7 @@
 #include "IFilter.h"
 #include "TextFilter.h"
 #include "FilterCollection.h"
-
-/**
-* The screen shot filter collection - used for updating the screen shot list
-*/
-typedef TFilterCollection< const TSharedPtr< class IScreenShotData >& > ScreenShotFilterCollection;
+#include "ScreenShotFilterCollection.h"
 
 /* Interfaces
  *****************************************************************************/
diff --git a/Engine/Source/Developer/ScreenShotComparisonTools/Public/ScreenShotFilterCollection.h b/Engine/Source/Developer/ScreenShotComparisonTools/Public/ScreenShotFilterCollection.h
new file mode 100644
--- /dev/null
+++ b/Engine/Source/Developer/ScreenShotComparisonTools/Public/ScreenShotFilterCollection.h
@@ -0,0 +1,8 @@
+// Copyright 1998-2016 Epic Games, Inc. All Rights Reserved.
+
+#pragma once
+
+/**
+* The screen shot filter collection - used for updating the screen shot list
+*/
+typedef TFilterCollection< const TSharedPtr< class IScreenShotData >& > ScreenShotFilterCollection;
diff --git a/Engine/Source/Developer/ShaderCompilerCommon/Private/HlslAST.cpp b/Engine/Source/Developer/ShaderCompilerCommon/Private/HlslAST.cpp
--- a/Engine/Source/Developer/ShaderCompilerCommon/Private/HlslAST.cpp
+++ b/Engine/Source/Developer/ShaderCompilerCommon/Private/HlslAST.cpp
@@ -4,6 +4,7 @@
 	HlslAST.cpp - Abstract Syntax Tree implementation for HLSL.
 =============================================================================*/
 
+#include "ShaderCompilerCommonPrivatePCH.h"
 #include "ShaderCompilerCommon.h"
 #include "HlslLexer.h"
 #include "HlslAST.h"
diff --git a/Engine/Source/Developer/ShaderCompilerCommon/Private/HlslAST.h b/Engine/Source/Developer/ShaderCompilerCommon/Private/HlslAST.h
--- a/Engine/Source/Developer/ShaderCompilerCommon/Private/HlslAST.h
+++ b/Engine/Source/Developer/ShaderCompilerCommon/Private/HlslAST.h
@@ -115,7 +115,7 @@ namespace CrossCompiler
 
 			/* Callers of this ralloc-based new need not call delete. It's
 			* easier to just ralloc_free 'ctx' (or any of its ancestors). */
-			static void* operator new(SIZE_T Size, FLinearAllocator* Allocator)
+			static void* operator new(size_t Size, FLinearAllocator* Allocator)
 			{
 #if USE_UNREAL_ALLOCATOR
 				return FMemory::Malloc(Size);
diff --git a/Engine/Source/Developer/ShaderCompilerCommon/Private/HlslLexer.cpp b/Engine/Source/Developer/ShaderCompilerCommon/Private/HlslLexer.cpp
--- a/Engine/Source/Developer/ShaderCompilerCommon/Private/HlslLexer.cpp
+++ b/Engine/Source/Developer/ShaderCompilerCommon/Private/HlslLexer.cpp
@@ -4,7 +4,7 @@
 	HlslLexer.cpp - Implementation for scanning & tokenizing hlsl
 =============================================================================*/
 
-
+#include "ShaderCompilerCommonPrivatePCH.h"
 #include "ShaderCompilerCommon.h"
 #include "HlslLexer.h"
 
diff --git a/Engine/Source/Developer/ShaderCompilerCommon/Private/HlslParser.cpp b/Engine/Source/Developer/ShaderCompilerCommon/Private/HlslParser.cpp
--- a/Engine/Source/Developer/ShaderCompilerCommon/Private/HlslParser.cpp
+++ b/Engine/Source/Developer/ShaderCompilerCommon/Private/HlslParser.cpp
@@ -4,6 +4,7 @@
 	HlslParser.cpp - Implementation for parsing hlsl.
 =============================================================================*/
 
+#include "ShaderCompilerCommonPrivatePCH.h"
 #include "ShaderCompilerCommon.h"
 #include "HlslParser.h"
 #include "HlslExpressionParser.inl"
diff --git a/Engine/Source/Developer/ShaderCompilerCommon/Private/ShaderCompilerCommonPrivatePCH.h b/Engine/Source/Developer/ShaderCompilerCommon/Private/ShaderCompilerCommonPrivatePCH.h
new file mode 100644
--- /dev/null
+++ b/Engine/Source/Developer/ShaderCompilerCommon/Private/ShaderCompilerCommonPrivatePCH.h
@@ -0,0 +1,5 @@
+// Copyright 1998-2016 Epic Games, Inc. All Rights Reserved.
+
+#pragma once
+
+#include "ShaderCompilerCommon.h"
diff --git a/Engine/Source/Developer/ShaderPreprocessor/Private/PreprocessorFile.cpp b/Engine/Source/Developer/ShaderPreprocessor/Private/PreprocessorFile.cpp
--- a/Engine/Source/Developer/ShaderPreprocessor/Private/PreprocessorFile.cpp
+++ b/Engine/Source/Developer/ShaderPreprocessor/Private/PreprocessorFile.cpp
@@ -1,5 +1,6 @@
 // Copyright 1998-2016 Epic Games, Inc. All Rights Reserved.
 
+#include "ShaderPreprocessorPrivatePCH.h"
 #include "ShaderPreprocessor.h"
 #include "ModuleManager.h"
 #include "PreprocessorPrivate.h"
diff --git a/Engine/Source/Developer/ShaderPreprocessor/Private/ShaderPreprocessor.cpp b/Engine/Source/Developer/ShaderPreprocessor/Private/ShaderPreprocessor.cpp
--- a/Engine/Source/Developer/ShaderPreprocessor/Private/ShaderPreprocessor.cpp
+++ b/Engine/Source/Developer/ShaderPreprocessor/Private/ShaderPreprocessor.cpp
@@ -1,5 +1,6 @@
 // Copyright 1998-2016 Epic Games, Inc. All Rights Reserved.
 
+#include "ShaderPreprocessorPrivatePCH.h"
 #include "ShaderPreprocessor.h"
 #include "ModuleManager.h"
 #include "PreprocessorPrivate.h"
diff --git a/Engine/Source/Developer/ShaderPreprocessor/Private/ShaderPreprocessorPrivatePCH.h b/Engine/Source/Developer/ShaderPreprocessor/Private/ShaderPreprocessorPrivatePCH.h
new file mode 100644
--- /dev/null
+++ b/Engine/Source/Developer/ShaderPreprocessor/Private/ShaderPreprocessorPrivatePCH.h
@@ -0,0 +1,7 @@
+// Copyright 1998-2016 Epic Games, Inc. All Rights Reserved.
+
+#pragma once
+
+#include "ShaderPreprocessor.h"
+#include "ModuleManager.h"
+#include "PreprocessorPrivate.h"
diff --git a/Engine/Source/Developer/StandaloneRenderer/Private/OpenGL/SlateOpenGLRenderingPolicy.h b/Engine/Source/Developer/StandaloneRenderer/Private/OpenGL/SlateOpenGLRenderingPolicy.h
--- a/Engine/Source/Developer/StandaloneRenderer/Private/OpenGL/SlateOpenGLRenderingPolicy.h
+++ b/Engine/Source/Developer/StandaloneRenderer/Private/OpenGL/SlateOpenGLRenderingPolicy.h
@@ -4,6 +4,7 @@
 #pragma once
 
 class FSlateOpenGLTextureCache;
+class FSlateOpenGLTextureManager;
 
 class FSlateOpenGLRenderingPolicy : public FSlateRenderingPolicy
 {
diff --git a/Engine/Source/Developer/StandaloneRenderer/Private/StandaloneRendererPrivate.h b/Engine/Source/Developer/StandaloneRenderer/Private/StandaloneRendererPrivate.h
--- a/Engine/Source/Developer/StandaloneRenderer/Private/StandaloneRendererPrivate.h
+++ b/Engine/Source/Developer/StandaloneRenderer/Private/StandaloneRendererPrivate.h
@@ -1,35 +1,13 @@
 // Copyright 1998-2016 Epic Games, Inc. All Rights Reserved.
 
-#ifndef UNREAL_STANDALONERENDERER_PRIVATEPCH_H
-#define UNREAL_STANDALONERENDERER_PRIVATEPCH_H
+#pragma once
 
 #include "Core.h"
 #include "SlateCore.h"
 
 DECLARE_LOG_CATEGORY_EXTERN(LogStandaloneRenderer, Log, All);
 
-#if PLATFORM_WINDOWS
-#include "AllowWindowsPlatformTypes.h"
-
-// Need to include before OpenGL headers
-#include <windows.h>
-
-// D3D headers.
-
-
-// Disable macro redefinition warning for compatibility with Windows SDK 8+
-#pragma warning(push)
-#pragma warning(disable : 4005)	// macro redefinition
-
-#include <d3d11.h>
-#include <d3dx11.h>
-#include <D3DX10.h>
-#include <D3Dcompiler.h>
-
-#pragma warning(pop)
-
-#include "HideWindowsPlatformTypes.h"
-#endif		// PLATFORM_WINDOWS
+#include "StandaloneRendererPlatformHeaders.h"
 
 #if PLATFORM_WINDOWS
 #include "Windows/D3D/SlateD3DTextures.h"
@@ -39,41 +17,6 @@ DECLARE_LOG_CATEGORY_EXTERN(LogStandaloneRenderer, Log, All);
 #include "Windows/D3D/SlateD3DVertexBuffer.h"
 #endif
 
-#if PLATFORM_WINDOWS
-#include "AllowWindowsPlatformTypes.h"
-#include <GL/gl.h>
-#include <GL/glext.h>
-#include <GL/wglext.h>
-#include "HideWindowsPlatformTypes.h"
-#elif PLATFORM_MAC
-#include <OpenGL/gl.h>
-#include <OpenGL/glext.h>
-#elif PLATFORM_IOS
-
-#include <OpenGLES/ES2/gl.h>
-#include <OpenGLES/ES2/glext.h>
-
-// renamed keywords
-#define glMapBuffer		glMapBufferOES
-#define glUnmapBuffer	glUnmapBufferOES
-#define GL_WRITE_ONLY	GL_WRITE_ONLY_OES
-#define GL_CLAMP		GL_CLAMP_TO_EDGE
-#define GL_RGBA8		GL_RGBA
-#define GL_UNSIGNED_INT_8_8_8_8_REV	GL_UNSIGNED_BYTE
-
-// unused functions
-#define glTexEnvf(...)
-#define glTexEnvi(...)
-#define glAlphaFunc(...)
-
-#elif PLATFORM_LINUX
-//#define GLCOREARB_PROTOTYPES 
-#include <GL/glcorearb.h>
-//#include <GL/glext.h>
-#include "SDL.h"
-
-#endif // PLATFORM_LINUX
-
 #include "OpenGL/SlateOpenGLExtensions.h"
 #include "OpenGL/SlateOpenGLTextures.h"
 #include "OpenGL/SlateOpenGLShaders.h"
@@ -81,6 +24,3 @@ DECLARE_LOG_CATEGORY_EXTERN(LogStandaloneRenderer, Log, All);
 #include "OpenGL/SlateOpenGLVertexBuffer.h"
 
 #include "StandaloneRenderer.h"
-
-
-#endif //UNREAL_STANDALONERENDERER_PRIVATEPCH_H
\ No newline at end of file
diff --git a/Engine/Source/Developer/StandaloneRenderer/Private/Windows/D3D/SlateD3DConstantBuffer.h b/Engine/Source/Developer/StandaloneRenderer/Private/Windows/D3D/SlateD3DConstantBuffer.h
--- a/Engine/Source/Developer/StandaloneRenderer/Private/Windows/D3D/SlateD3DConstantBuffer.h
+++ b/Engine/Source/Developer/StandaloneRenderer/Private/Windows/D3D/SlateD3DConstantBuffer.h
@@ -1,9 +1,6 @@
 // Copyright 1998-2016 Epic Games, Inc. All Rights Reserved.
 
-
-#ifndef __D3DCONSTANT_BUFFER_H__
-#define __D3DCONSTANT_BUFFER_H__
-
+#pragma once
 
 template<typename BufferType>
 class FSlateD3DConstantBuffer
@@ -48,4 +45,3 @@ private:
 	BufferType BufferData;
 };
 
-#endif
diff --git a/Engine/Source/Developer/StandaloneRenderer/Private/Windows/D3D/SlateD3DIndexBuffer.h b/Engine/Source/Developer/StandaloneRenderer/Private/Windows/D3D/SlateD3DIndexBuffer.h
--- a/Engine/Source/Developer/StandaloneRenderer/Private/Windows/D3D/SlateD3DIndexBuffer.h
+++ b/Engine/Source/Developer/StandaloneRenderer/Private/Windows/D3D/SlateD3DIndexBuffer.h
@@ -1,8 +1,6 @@
 // Copyright 1998-2016 Epic Games, Inc. All Rights Reserved.
 
-
-#ifndef __D3DINDEXBUFFER_H__
-#define __D3DINDEXBUFFER_H__
+#pragma once
 
 class FSlateD3DIndexBuffer 
 {
@@ -37,5 +35,3 @@ private:
 	void operator=(const FSlateD3DIndexBuffer& );
 
 };
-
-#endif
\ No newline at end of file
diff --git a/Engine/Source/Developer/StandaloneRenderer/Private/Windows/D3D/SlateD3DRenderer.h b/Engine/Source/Developer/StandaloneRenderer/Private/Windows/D3D/SlateD3DRenderer.h
--- a/Engine/Source/Developer/StandaloneRenderer/Private/Windows/D3D/SlateD3DRenderer.h
+++ b/Engine/Source/Developer/StandaloneRenderer/Private/Windows/D3D/SlateD3DRenderer.h
@@ -1,9 +1,6 @@
 // Copyright 1998-2016 Epic Games, Inc. All Rights Reserved.
 
-
-#ifndef _SLATERENDERRD3D_H__
-#define _SLATERENDERRD3D_H__
-
+#pragma once
 
 class FSlateD3DTextureManager;
 class FSlateD3D11RenderingPolicy;
@@ -78,5 +75,3 @@ private:
 	TSharedPtr<FSlateD3D11RenderingPolicy> RenderingPolicy;
 	TArray<TSharedPtr<FSlateDynamicImageBrush>> DynamicBrushesToRemove;
 };
-
-#endif
\ No newline at end of file
diff --git a/Engine/Source/Developer/StandaloneRenderer/Private/Windows/D3D/SlateD3DShaders.h b/Engine/Source/Developer/StandaloneRenderer/Private/Windows/D3D/SlateD3DShaders.h
--- a/Engine/Source/Developer/StandaloneRenderer/Private/Windows/D3D/SlateD3DShaders.h
+++ b/Engine/Source/Developer/StandaloneRenderer/Private/Windows/D3D/SlateD3DShaders.h
@@ -3,6 +3,8 @@
 
 #pragma once
 
+#include "SlateD3DConstantBuffer.h"
+
 class FSlateD3DShaderParameter
 {
 public:
diff --git a/Engine/Source/Developer/StandaloneRenderer/Private/Windows/D3D/SlateD3DVertexBuffer.h b/Engine/Source/Developer/StandaloneRenderer/Private/Windows/D3D/SlateD3DVertexBuffer.h
--- a/Engine/Source/Developer/StandaloneRenderer/Private/Windows/D3D/SlateD3DVertexBuffer.h
+++ b/Engine/Source/Developer/StandaloneRenderer/Private/Windows/D3D/SlateD3DVertexBuffer.h
@@ -1,8 +1,6 @@
 // Copyright 1998-2016 Epic Games, Inc. All Rights Reserved.
 
-
-#ifndef __SLATED3DVERTEXBUFFER_H__
-#define __SLATED3DVERTEXBUFFER_H__
+#pragma once
 
 /** 
  * Vertex buffer containing all Slate vertices
@@ -38,5 +36,3 @@ private:
 	FSlateD3DVertexBuffer( const FSlateD3DVertexBuffer& );
 	void operator=(const FSlateD3DVertexBuffer& );
 };
-
-#endif //__SLATED3DVERTEXBUFFER_H__
\ No newline at end of file
diff --git a/Engine/Source/Developer/TaskGraph/Private/SBarVisualizer.cpp b/Engine/Source/Developer/TaskGraph/Private/SBarVisualizer.cpp
--- a/Engine/Source/Developer/TaskGraph/Private/SBarVisualizer.cpp
+++ b/Engine/Source/Developer/TaskGraph/Private/SBarVisualizer.cpp
@@ -1,5 +1,6 @@
 // Copyright 1998-2016 Epic Games, Inc. All Rights Reserved.
 
+#include "TaskGraphPrivatePCH.h"
 #include "SlateBasics.h"
 #include "TaskGraphInterfaces.h"
 #include "VisualizerEvents.h"
diff --git a/Engine/Source/Developer/TaskGraph/Private/SEventsTree.cpp b/Engine/Source/Developer/TaskGraph/Private/SEventsTree.cpp
--- a/Engine/Source/Developer/TaskGraph/Private/SEventsTree.cpp
+++ b/Engine/Source/Developer/TaskGraph/Private/SEventsTree.cpp
@@ -1,5 +1,6 @@
 // Copyright 1998-2016 Epic Games, Inc. All Rights Reserved.
 
+#include "TaskGraphPrivatePCH.h"
 #include "SlateBasics.h"
 #include "TaskGraphInterfaces.h"
 #include "VisualizerEvents.h"
diff --git a/Engine/Source/Developer/TaskGraph/Private/TaskGraphPrivatePCH.h b/Engine/Source/Developer/TaskGraph/Private/TaskGraphPrivatePCH.h
new file mode 100644
--- /dev/null
+++ b/Engine/Source/Developer/TaskGraph/Private/TaskGraphPrivatePCH.h
@@ -0,0 +1,5 @@
+// Copyright 1998-2016 Epic Games, Inc. All Rights Reserved.
+
+#pragma once
+
+#include "SlateBasics.h"
diff --git a/Engine/Source/Developer/TaskGraph/TaskGraph.Build.cs b/Engine/Source/Developer/TaskGraph/TaskGraph.Build.cs
--- a/Engine/Source/Developer/TaskGraph/TaskGraph.Build.cs
+++ b/Engine/Source/Developer/TaskGraph/TaskGraph.Build.cs
@@ -9,6 +9,7 @@ public class TaskGraph : ModuleRules
 		PrivateDependencyModuleNames.AddRange(
 			new string[] {
 				"Core",
+				"CoreUObject",
 				"Slate",
 				"SlateCore",
                 "EditorStyle",
diff --git a/Engine/Source/Developer/iOS/IOSTargetPlatform/Private/Mac/IOSDeviceHelperMac.cpp b/Engine/Source/Developer/iOS/IOSTargetPlatform/Private/Mac/IOSDeviceHelperMac.cpp
--- a/Engine/Source/Developer/iOS/IOSTargetPlatform/Private/Mac/IOSDeviceHelperMac.cpp
+++ b/Engine/Source/Developer/iOS/IOSTargetPlatform/Private/Mac/IOSDeviceHelperMac.cpp
@@ -1,6 +1,6 @@
 // Copyright 1998-2016 Epic Games, Inc. All Rights Reserved
 
-#include "../IOSTargetPlatformPrivatePCH.h"
+#include "IOSTargetPlatformPrivatePCH.h"
 
 #include <CoreFoundation/CoreFoundation.h>
 
diff --git a/Engine/Source/Developer/iOS/IOSTargetPlatform/Private/Windows/IOSDeviceHelperWindows.cpp b/Engine/Source/Developer/iOS/IOSTargetPlatform/Private/Windows/IOSDeviceHelperWindows.cpp
--- a/Engine/Source/Developer/iOS/IOSTargetPlatform/Private/Windows/IOSDeviceHelperWindows.cpp
+++ b/Engine/Source/Developer/iOS/IOSTargetPlatform/Private/Windows/IOSDeviceHelperWindows.cpp
@@ -1,6 +1,6 @@
 // Copyright 1998-2016 Epic Games, Inc. All Rights Reserved
 
-#include "../IOSTargetPlatformPrivatePCH.h"
+#include "IOSTargetPlatformPrivatePCH.h"
 
 #define EABLE_IOS_DEVICE_DETECT 0
 
diff --git a/Engine/Source/Editor/BlueprintGraph/Classes/K2Node_GetClassDefaults.h b/Engine/Source/Editor/BlueprintGraph/Classes/K2Node_GetClassDefaults.h
--- a/Engine/Source/Editor/BlueprintGraph/Classes/K2Node_GetClassDefaults.h
+++ b/Engine/Source/Editor/BlueprintGraph/Classes/K2Node_GetClassDefaults.h
@@ -52,14 +52,14 @@ protected:
 	 *
 	 * @param FromPins	A list of pins to search.
 	 */
-	UEdGraphPin* FindClassPin(const TArray<UEdGraphPin*>& FromPins) const;
+	BLUEPRINTGRAPH_API UEdGraphPin* FindClassPin(const TArray<UEdGraphPin*>& FromPins) const;
 
 	/**
 	* Determines the input class type from the given pin.
 	*
 	* @param FromPin	Input class pin.
 	*/
-	UClass* GetInputClass(const UEdGraphPin* FromPin) const;
+	BLUEPRINTGRAPH_API UClass* GetInputClass(const UEdGraphPin* FromPin) const;
 
 	/**
 	 * Creates the full set of output pins (properties) from the given input class.
diff --git a/Engine/Source/Editor/GameplayAbilitiesEditor/Private/SGameplayCueEditor.cpp b/Engine/Source/Editor/GameplayAbilitiesEditor/Private/SGameplayCueEditor.cpp
--- a/Engine/Source/Editor/GameplayAbilitiesEditor/Private/SGameplayCueEditor.cpp
+++ b/Engine/Source/Editor/GameplayAbilitiesEditor/Private/SGameplayCueEditor.cpp
@@ -481,181 +481,180 @@ private:
 
 	typedef STreeView< TSharedPtr< FTreeItem > > SGameplayCueTreeView;
 
-
-	/** Builds widget for rows in the GameplayCue Editor tab */
-	TSharedRef<ITableRow> OnGenerateWidgetForGameplayCueListView(TSharedPtr< FTreeItem > InItem, const TSharedRef<STableViewBase>& OwnerTable)
+	class SCueItemWidget : public SMultiColumnTableRow< TSharedPtr< FCueItem > >
 	{
-		class SCueItemWidget : public SMultiColumnTableRow< TSharedPtr< FCueItem > >
-		{
-		public:
-			SLATE_BEGIN_ARGS(SCueItemWidget){}
-			SLATE_END_ARGS()
+	public:
+		SLATE_BEGIN_ARGS(SCueItemWidget) {}
+		SLATE_END_ARGS()
 
-			void Construct(const FArguments& InArgs, const TSharedRef<SGameplayCueTreeView>& InOwnerTable, FSimpleDelegate InRefreshDelegate,  TSharedPtr<FCueItem> InListItem, FText InSearchTerm )
-			{
-				Item = InListItem;
-				SearchTerm = InSearchTerm;
-				RefreshDelegate = InRefreshDelegate;
-				SMultiColumnTableRow< TSharedPtr< FCueItem > >::Construct(FSuperRowType::FArguments(), InOwnerTable);
-			}
-		private:
+			void Construct(const FArguments& InArgs, const TSharedRef<SGameplayCueTreeView>& InOwnerTable, FSimpleDelegate InRefreshDelegate, TSharedPtr<FCueItem> InListItem, FText InSearchTerm)
+		{
+			Item = InListItem;
+			SearchTerm = InSearchTerm;
+			RefreshDelegate = InRefreshDelegate;
+			SMultiColumnTableRow< TSharedPtr< FCueItem > >::Construct(FSuperRowType::FArguments(), InOwnerTable);
+		}
+	private:
 
-			virtual TSharedRef<SWidget> GenerateWidgetForColumn(const FName& ColumnName) override
+		virtual TSharedRef<SWidget> GenerateWidgetForColumn(const FName& ColumnName) override
+		{
+			if (ColumnName == CueTagColumnName)
 			{
-				if (ColumnName == CueTagColumnName)
-				{
-					return
+				return
 					SNew(SHorizontalBox)
 					+ SHorizontalBox::Slot()
 					.AutoWidth()
 					[
 						SNew(SExpanderArrow, SharedThis(this))
 					]
-					+ SHorizontalBox::Slot()
+				+ SHorizontalBox::Slot()
 					.FillWidth(1)
 					.VAlign(VAlign_Center)
 					[
 						SNew(STextBlock)
 						.HighlightText(SearchTerm)
-						.Text(FText::FromString(Item->GameplayCueTagName.ToString()))
+					.Text(FText::FromString(Item->GameplayCueTagName.ToString()))
 					];
-				}
-				else if (ColumnName == CueHandlerColumnName)
+			}
+			else if (ColumnName == CueHandlerColumnName)
+			{
+				bool HasNotifies = false;
+				for (auto& HandlerItem : Item->HandlerItems)
 				{
-					bool HasNotifies = false;
-					for (auto& HandlerItem : Item->HandlerItems)
+					if (HandlerItem->GameplayCueNotifyObj.IsValid())
 					{
-						if (HandlerItem->GameplayCueNotifyObj.IsValid())
-						{
-							HasNotifies = true;
-							break;
-						}
+						HasNotifies = true;
+						break;
 					}
+				}
 
-					return
+				return
 					SNew(SBox)
-					.Padding( FMargin(2.0f, 0.0f) )
+					.Padding(FMargin(2.0f, 0.0f))
 					.HAlign(HAlign_Left)
 					.VAlign(VAlign_Center)
 					.Visibility(HasNotifies ? EVisibility::Collapsed : EVisibility::Visible)
 					[
 						SNew(SButton)
 						.Text(LOCTEXT("AddNew", "Add New"))
-						.OnClicked(this, &SCueItemWidget::OnAddNewClicked)
+					.OnClicked(this, &SCueItemWidget::OnAddNewClicked)
 					];
-				}
-				else
-				{
-					return SNew(STextBlock).Text(LOCTEXT("UnknownColumn", "Unknown Column"));
-				}
-
 			}
-
-			/** Create new GameplayCueNotify: brings up dialogue to pick class, then creates it via the content browser. */
-			FReply OnAddNewClicked()
+			else
 			{
-				SGameplayCueEditor::CreateNewGameplayCueNotifyDialogue(Item->GameplayCueTagName.ToString());
-				RefreshDelegate.ExecuteIfBound();
-				return FReply::Handled();
+				return SNew(STextBlock).Text(LOCTEXT("UnknownColumn", "Unknown Column"));
 			}
 
-			TSharedPtr< FCueItem > Item;
-			FText SearchTerm;
-			FSimpleDelegate RefreshDelegate;
-		};
+		}
 
-		class SCueHandlerItemWidget : public SMultiColumnTableRow < TSharedPtr< FCueHandlerItem > >
+		/** Create new GameplayCueNotify: brings up dialogue to pick class, then creates it via the content browser. */
+		FReply OnAddNewClicked()
 		{
-		public:
-			SLATE_BEGIN_ARGS(SCueHandlerItemWidget){}
-			SLATE_END_ARGS()
+			SGameplayCueEditor::CreateNewGameplayCueNotifyDialogue(Item->GameplayCueTagName.ToString());
+			RefreshDelegate.ExecuteIfBound();
+			return FReply::Handled();
+		}
+
+		TSharedPtr< FCueItem > Item;
+		FText SearchTerm;
+		FSimpleDelegate RefreshDelegate;
+	};
+
+	class SCueHandlerItemWidget : public SMultiColumnTableRow < TSharedPtr< FCueHandlerItem > >
+	{
+	public:
+		SLATE_BEGIN_ARGS(SCueHandlerItemWidget) {}
+		SLATE_END_ARGS()
 
 			void Construct(const FArguments& InArgs, const TSharedRef<STableViewBase>& InOwnerTable, TSharedPtr<FCueHandlerItem> InListItem)
+		{
+			CueHandlerItem = InListItem;
+			SMultiColumnTableRow< TSharedPtr< FCueHandlerItem > >::Construct(FSuperRowType::FArguments(), InOwnerTable);
+		}
+	private:
+
+		virtual TSharedRef<SWidget> GenerateWidgetForColumn(const FName& ColumnName) override
+		{
+			if (ColumnName == CueTagColumnName)
 			{
-				CueHandlerItem = InListItem;
-				SMultiColumnTableRow< TSharedPtr< FCueHandlerItem > >::Construct(FSuperRowType::FArguments(), InOwnerTable);
+				return SNew(SSpacer);
 			}
-		private:
-
-			virtual TSharedRef<SWidget> GenerateWidgetForColumn(const FName& ColumnName) override
+			else if (ColumnName == CueHandlerColumnName)
 			{
-				if (ColumnName == CueTagColumnName)
+				if (CueHandlerItem->GameplayCueNotifyObj.ToString().IsEmpty() == false)
 				{
-					return SNew(SSpacer);
-				}
-				else if (ColumnName == CueHandlerColumnName)
-				{
-					if (CueHandlerItem->GameplayCueNotifyObj.ToString().IsEmpty() == false)
-					{
-						FString ObjName = CueHandlerItem->GameplayCueNotifyObj.ToString();
+					FString ObjName = CueHandlerItem->GameplayCueNotifyObj.ToString();
 
-						int32 idx;
-						if (ObjName.FindLastChar(TEXT('.'), idx))
+					int32 idx;
+					if (ObjName.FindLastChar(TEXT('.'), idx))
+					{
+						ObjName = ObjName.RightChop(idx + 1);
+						if (ObjName.FindLastChar(TEXT('_'), idx))
 						{
-							ObjName = ObjName.RightChop(idx + 1);
-							if (ObjName.FindLastChar(TEXT('_'), idx))
-							{
-								ObjName = ObjName.Left(idx);
-							}
+							ObjName = ObjName.Left(idx);
 						}
+					}
 
-						return
+					return
 						SNew(SBox)
 						.HAlign(HAlign_Left)
 						[
 							SNew(SHyperlink)
 							.Style(FEditorStyle::Get(), "Common.GotoBlueprintHyperlink")
-							.Text(FText::FromString(ObjName))
-							.OnNavigate(this, &SCueHandlerItemWidget::NavigateToHandler)
+						.Text(FText::FromString(ObjName))
+						.OnNavigate(this, &SCueHandlerItemWidget::NavigateToHandler)
 						];
-					}
-					else if (CueHandlerItem->FunctionPtr.IsValid())
+				}
+				else if (CueHandlerItem->FunctionPtr.IsValid())
+				{
+					FString ObjName;
+					UFunction* Func = CueHandlerItem->FunctionPtr.Get();
+					UClass* OuterClass = Cast<UClass>(Func->GetOuter());
+					if (OuterClass)
 					{
-						FString ObjName;
-						UFunction* Func = CueHandlerItem->FunctionPtr.Get();
-						UClass* OuterClass = Cast<UClass>(Func->GetOuter());
-						if (OuterClass)
-						{
-							ObjName = OuterClass->GetName();
-							ObjName.RemoveFromEnd(TEXT("_c"));
-						}
+						ObjName = OuterClass->GetName();
+						ObjName.RemoveFromEnd(TEXT("_c"));
+					}
 
-						return
+					return
 						SNew(SBox)
 						.HAlign(HAlign_Left)
 						[
 							SNew(SHyperlink)
 							.Text(FText::FromString(ObjName))
-							.OnNavigate(this, &SCueHandlerItemWidget::NavigateToHandler)
+						.OnNavigate(this, &SCueHandlerItemWidget::NavigateToHandler)
 						];
-					}
-					else
-					{
-						return SNew(STextBlock).Text(LOCTEXT("UnknownHandler", "Unknown HandlerType"));
-					}
 				}
 				else
 				{
-					return SNew(STextBlock).Text(LOCTEXT("UnknownColumn", "Unknown Column"));
+					return SNew(STextBlock).Text(LOCTEXT("UnknownHandler", "Unknown HandlerType"));
 				}
 			}
+			else
+			{
+				return SNew(STextBlock).Text(LOCTEXT("UnknownColumn", "Unknown Column"));
+			}
+		}
 
-			void NavigateToHandler()
+		void NavigateToHandler()
+		{
+			if (CueHandlerItem->GameplayCueNotifyObj.IsValid())
 			{
-				if (CueHandlerItem->GameplayCueNotifyObj.IsValid())
-				{
-					SGameplayCueEditor::OpenEditorForNotify(CueHandlerItem->GameplayCueNotifyObj.ToString());
-					
-				}
-				else if (CueHandlerItem->FunctionPtr.IsValid())
-				{
-					SGameplayCueEditor::OpenEditorForNotify(CueHandlerItem->FunctionPtr->GetOuter()->GetPathName());
-				}
+				SGameplayCueEditor::OpenEditorForNotify(CueHandlerItem->GameplayCueNotifyObj.ToString());
+
 			}
+			else if (CueHandlerItem->FunctionPtr.IsValid())
+			{
+				SGameplayCueEditor::OpenEditorForNotify(CueHandlerItem->FunctionPtr->GetOuter()->GetPathName());
+			}
+		}
 
-			TSharedPtr<FCueHandlerItem> CueHandlerItem;
-		};
+		TSharedPtr<FCueHandlerItem> CueHandlerItem;
+	};
 
+	/** Builds widget for rows in the GameplayCue Editor tab */
+	TSharedRef<ITableRow> OnGenerateWidgetForGameplayCueListView(TSharedPtr< FTreeItem > InItem, const TSharedRef<STableViewBase>& OwnerTable)
+	{
 		TSharedPtr<FCueItem> CueItem = InItem->AsCueItem();
 		TSharedPtr<FCueHandlerItem> CueHandlerItem = InItem->AsCueHandlerItem();
 
diff --git a/Engine/Source/Editor/UnrealEd/Classes/Editor/Transactor.h b/Engine/Source/Editor/UnrealEd/Classes/Editor/Transactor.h
--- a/Engine/Source/Editor/UnrealEd/Classes/Editor/Transactor.h
+++ b/Engine/Source/Editor/UnrealEd/Classes/Editor/Transactor.h
@@ -434,7 +434,7 @@ public:
 };
 
 UCLASS(transient)
-class UTransactor : public UObject
+class UNREALED_API UTransactor : public UObject
 {
     GENERATED_UCLASS_BODY()
 	/**
diff --git a/Engine/Source/Editor/UnrealEd/Public/Toolkits/AssetEditorToolkit.h b/Engine/Source/Editor/UnrealEd/Public/Toolkits/AssetEditorToolkit.h
--- a/Engine/Source/Editor/UnrealEd/Public/Toolkits/AssetEditorToolkit.h
+++ b/Engine/Source/Editor/UnrealEd/Public/Toolkits/AssetEditorToolkit.h
@@ -279,7 +279,7 @@ private:
 	 * so that we don't break compatibility with all the asset editors out there that individually 
 	 * implement FGCObject.
 	 */
-	class FGCEditingObjects : public FGCObject
+	class UNREALED_API FGCEditingObjects : public FGCObject
 	{
 	public:
 		FGCEditingObjects(FAssetEditorToolkit& InOwnerToolkit) : OwnerToolkit(InOwnerToolkit) {}
diff --git a/Engine/Source/Editor/UnrealEd/Public/UnrealWidget.h b/Engine/Source/Editor/UnrealEd/Public/UnrealWidget.h
--- a/Engine/Source/Editor/UnrealEd/Public/UnrealWidget.h
+++ b/Engine/Source/Editor/UnrealEd/Public/UnrealWidget.h
@@ -2,8 +2,6 @@
 
 #pragma once
 
-#include "Editor.h"
-
 
 class FPrimitiveDrawInterface;
 class FEditorModeTools;
diff --git a/Engine/Source/Editor/ViewportInteraction/Public/ViewportInteraction.h b/Engine/Source/Editor/ViewportInteraction/Public/ViewportInteraction.h
--- a/Engine/Source/Editor/ViewportInteraction/Public/ViewportInteraction.h
+++ b/Engine/Source/Editor/ViewportInteraction/Public/ViewportInteraction.h
@@ -5,5 +5,4 @@
 #include "ViewportInteractableInterface.h"
 #include "IViewportInteractionModule.h"
 #include "ViewportWorldInteractionInterface.h"
-#include "ViewportInteraction.h"
 #include "ViewportInteractor.h"
\ No newline at end of file
diff --git a/Engine/Source/Editor/WorkspaceMenuStructure/Private/WorkspaceMenuStructureModule.cpp b/Engine/Source/Editor/WorkspaceMenuStructure/Private/WorkspaceMenuStructureModule.cpp
--- a/Engine/Source/Editor/WorkspaceMenuStructure/Private/WorkspaceMenuStructureModule.cpp
+++ b/Engine/Source/Editor/WorkspaceMenuStructure/Private/WorkspaceMenuStructureModule.cpp
@@ -1,5 +1,6 @@
 // Copyright 1998-2016 Epic Games, Inc. All Rights Reserved.
 
+#include "Core.h"
 #include "ModuleManager.h"
 
 #include "WorkspaceMenuStructureModule.h"
diff --git a/Engine/Source/Programs/AutomationTool/BuildGraph/BuildGraph.Automation.csproj b/Engine/Source/Programs/AutomationTool/BuildGraph/BuildGraph.Automation.csproj
--- a/Engine/Source/Programs/AutomationTool/BuildGraph/BuildGraph.Automation.csproj
+++ b/Engine/Source/Programs/AutomationTool/BuildGraph/BuildGraph.Automation.csproj
@@ -32,6 +32,7 @@
     <WarningLevel>4</WarningLevel>
     <PlatformTarget>AnyCPU</PlatformTarget>
     <TreatWarningsAsErrors>true</TreatWarningsAsErrors>
+    <DocumentationFile>..\..\..\..\Binaries\DotNET\AutomationScripts\BuildGraph.Automation.xml</DocumentationFile>
   </PropertyGroup>
   <ItemGroup>
     <Reference Include="Ionic.Zip.Reduced">
@@ -73,14 +74,10 @@
     <Compile Include="Condition.cs" />
     <Compile Include="Schema.cs" />
     <Compile Include="Script.cs" />
-    <Compile Include="Tasks\ChunkTask.cs" />
     <Compile Include="Tasks\CommandletTask.cs" />
     <Compile Include="Tasks\CsCompileTask.cs" />
     <Compile Include="Tasks\TagReceiptTask.cs" />
-    <Compile Include="Tasks\LabelBuildTask.cs" />
-    <Compile Include="Tasks\MergeTask.cs" />
     <Compile Include="Tasks\PakFileTask.cs" />
-    <Compile Include="Tasks\PostBuildTask.cs" />
     <Compile Include="Tasks\RenameTask.cs" />
     <Compile Include="Tasks\SpawnTask.cs" />
     <Compile Include="Tasks\LogTask.cs" />
diff --git a/Engine/Source/Programs/AutomationTool/BuildGraph/BuildGraph.cs b/Engine/Source/Programs/AutomationTool/BuildGraph/BuildGraph.cs
--- a/Engine/Source/Programs/AutomationTool/BuildGraph/BuildGraph.cs
+++ b/Engine/Source/Programs/AutomationTool/BuildGraph/BuildGraph.cs
@@ -8,11 +8,12 @@ using BuildGraph;
 using System.Reflection;
 using System.Collections;
 using System.IO;
+using System.Xml;
 
 namespace AutomationTool
 {
 	/// <summary>
-	/// Tool to execute build processes for UE4 projects, which can be run locally or in parallel across a build farm (assuming synchronization and resource allocation implemented by a separate system).
+	/// Tool to execute build automation scripts for UE4 projects, which can be run locally or in parallel across a build farm (assuming synchronization and resource allocation implemented by a separate system).
 	///
 	/// Build graphs are declared using an XML script using syntax similar to MSBuild, ANT or NAnt, and consist of the following components:
 	///
@@ -23,14 +24,11 @@ namespace AutomationTool
 	/// - Triggers:     Container for agents which should only be executed when explicitly triggered (using the -Trigger=... or -SkipTriggers command line argument). Declared with the 'Trigger' element.
 	/// - Notifiers:    Specifies email recipients for failures in one or more nodes, whether they should receive notifications on warnings, and so on.
 	/// 
-	/// Properties can be passed in to a script on the command line, or set procedurally with the &ltProperty Name="Foo" Value="Bar"/&gt; syntax. Properties referenced with the $(Property Name) notation are valid within 
-	/// all strings, and will be expanded as macros when the script is read. If a property name is not set explicitly, it defaults to the contents of an environment variable with the same name. 
-	/// Local properties, which only affect the scope of the containing XML element (node, agent, etc...) are declared with the &lt;Local Name="Foo" Value="Bar"/&gt; element, and will override a similarly named global 
-	/// property for the local property's scope.
+	/// Scripts may set properties with the &lt;Property Name="Foo" Value="Bar"/&gt; syntax. Properties referenced with the $(Property Name) notation are valid within all strings, and will be expanded as macros when the 
+	/// script is read. If a property name is not set explicitly, it defaults to the contents of an environment variable with the same name. Properties may be sourced from environment variables or the command line using
+	/// the &lt;EnvVar&gt; and &lt;Option&gt; elements respectively.
 	///
-	/// Any elements can be conditionally defined via the "If" attribute, which follows a syntax similar to MSBuild. Literals in conditions may be quoted with single (') or double (") quotes, or an unquoted sequence of 
-	/// letters, digits and underscore characters. All literals are considered identical regardless of how they are declared, and are considered case-insensitive for comparisons (so true equals 'True', equals "TRUE"). 
-	/// Available operators are "==", "!=", "And", "Or", "!", "(...)", "Exists(...)" and "HasTrailingSlash(...)". A full grammar is written up in Condition.cs.
+	/// Any elements can be conditionally defined via the "If" attribute. A full grammar for conditions is written up in Condition.cs.
 	/// 
 	/// File manipulation is done using wildcards and tags. Any attribute that accepts a list of files may consist of: a Perforce-style wildcard (matching any number of "...", "*" and "?" patterns in any location), a 
 	/// full path name, or a reference to a tagged collection of files, denoted by prefixing with a '#' character. Files may be added to a tag set using the &lt;Tag&gt; Task, which also allows performing set union/difference 
@@ -73,13 +71,8 @@ namespace AutomationTool
 		{
 			// Parse the command line parameters
 			string ScriptFileName = ParseParamValue("Script", null);
-			if(ScriptFileName == null)
-			{
-				LogError("Missing -Script= parameter for BuildGraph");
-				return ExitCode.Error_Unknown;
-			}
-
 			string TargetNames = ParseParamValue("Target", null);
+			string DocumentationFileName = ParseParamValue("Documentation", null);
 			string SchemaFileName = ParseParamValue("Schema", null);
 			string ExportFileName = ParseParamValue("Export", null);
 			string PreprocessedFileName = ParseParamValue("Preprocess", null);
@@ -160,11 +153,28 @@ namespace AutomationTool
 				return ExitCode.Error_Unknown;
 			}
 
+			// Generate documentation
+			if(DocumentationFileName != null)
+			{
+				GenerateDocumentation(NameToTask, new FileReference(DocumentationFileName));
+				return ExitCode.Success;
+			}
+
 			// Create a schema for the given tasks
 			ScriptSchema Schema = new ScriptSchema(NameToTask);
 			if(SchemaFileName != null)
 			{
-				Schema.Export(new FileReference(SchemaFileName));
+				FileReference FullSchemaFileName = new FileReference(SchemaFileName);
+				Log("Writing schema to {0}...", FullSchemaFileName.FullName);
+				Schema.Export(FullSchemaFileName);
+				return ExitCode.Success;
+			}
+
+			// Check there was a script specified
+			if(ScriptFileName == null)
+			{
+				LogError("Missing -Script= parameter for BuildGraph");
+				return ExitCode.Error_Unknown;
 			}
 
 			// Read the script from disk
@@ -365,7 +375,7 @@ namespace AutomationTool
 		/// <summary>
 		/// Find all the tasks which are available from the loaded assemblies
 		/// </summary>
-		/// <param name="TaskNameToReflectionInfo">Mapping from task name to information about how to serialize it</param>
+		/// <param name="NameToTask">Mapping from task name to information about how to serialize it</param>
 		/// <param name="bPublicTasksOnly">Whether to include just public tasks, or all the tasks in any loaded assemblies</param>
 		static bool FindAvailableTasks(Dictionary<string, ScriptTask> NameToTask, bool bPublicTasksOnly)
 		{
@@ -503,6 +513,7 @@ namespace AutomationTool
 		/// </summary>
 		/// <param name="Job">Information about the current job</param>
 		/// <param name="Graph">The graph instance</param>
+		/// <param name="Storage">The temp storage backend which stores the shared state</param>
 		/// <returns>True if everything built successfully</returns>
 		bool BuildAllNodes(JobContext Job, Graph Graph, TempStorage Storage)
 		{
@@ -545,6 +556,8 @@ namespace AutomationTool
 		/// <param name="Job">Information about the current job</param>
 		/// <param name="Graph">The graph to which the node belongs. Used to determine which outputs need to be transferred to temp storage.</param>
 		/// <param name="Node">The node to build</param>
+		/// <param name="Storage">The temp storage backend which stores the shared state</param>
+		/// <param name="bWithBanner">Whether to write a banner before and after this node's log output</param>
 		/// <returns>True if the node built successfully, false otherwise.</returns>
 		bool BuildNode(JobContext Job, Graph Graph, Node Node, TempStorage Storage, bool bWithBanner)
 		{
@@ -695,6 +708,165 @@ namespace AutomationTool
 			Storage.MarkAsComplete(Node.Name);
 			return true;
 		}
+
+		/// <summary>
+		/// Generate HTML documentation for all the tasks
+		/// </summary>
+		/// <param name="NameToTask">Map of task name to implementation</param>
+		/// <param name="OutputFile">Output file</param>
+		static void GenerateDocumentation(Dictionary<string, ScriptTask> NameToTask, FileReference OutputFile)
+		{
+			// Find all the assemblies containing tasks
+			Assembly[] TaskAssemblies = NameToTask.Values.Select(x => x.ParametersClass.Assembly).Distinct().ToArray();
+
+			// Read documentation for each of them
+			Dictionary<string, XmlElement> MemberNameToElement = new Dictionary<string, XmlElement>();
+			foreach(Assembly TaskAssembly in NameToTask.Values.Select(x => x.ParametersClass.Assembly).Distinct())
+			{
+				string XmlFileName = Path.ChangeExtension(TaskAssembly.Location, ".xml");
+				if(File.Exists(XmlFileName))
+				{
+					// Read the document
+					XmlDocument Document = new XmlDocument();
+					Document.Load(XmlFileName);
+
+					// Parse all the members, and add them to the map
+					foreach(XmlElement Element in Document.SelectNodes("/doc/members/member"))
+					{
+						string Name = Element.GetAttribute("name");
+						MemberNameToElement.Add(Name, Element);
+					}
+				}
+			}
+
+			// Create the output directory
+			OutputFile.Directory.CreateDirectory();
+			Log("Writing {0}...", OutputFile);
+
+			// Parse the engine version
+			BuildVersion Version;
+			if(!BuildVersion.TryRead(out Version))
+			{
+				throw new AutomationException("Couldn't read Build.version");
+			}
+
+			// Write the output file
+			using (StreamWriter Writer = new StreamWriter(OutputFile.FullName))
+			{
+				Writer.WriteLine("Availability: NoPublish");
+				Writer.WriteLine("Title: BuildGraph Predefined Tasks");
+				Writer.WriteLine("Crumbs: %ROOT%, Programming, Programming/Development, Programming/Development/BuildGraph, Programming/Development/BuildGraph/BuildGraphScriptTasks");
+				Writer.WriteLine("Description: This is a procedurally generated markdown page.");
+				Writer.WriteLine("version: {0}.{1}", Version.MajorVersion, Version.MinorVersion);
+				Writer.WriteLine("parent:Programming/Development/BuildGraph/BuildGraphScriptTasks");
+				Writer.WriteLine();
+				foreach(string TaskName in NameToTask.Keys.OrderBy(x => x))
+				{
+					// Get the task object
+					ScriptTask Task = NameToTask[TaskName];
+
+					// Get the documentation for this task
+					XmlElement TaskElement;
+					if(MemberNameToElement.TryGetValue("T:" + Task.TaskClass.FullName, out TaskElement))
+					{
+						// Write the task heading
+						Writer.WriteLine("### {0}", TaskName);
+						Writer.WriteLine();
+						Writer.WriteLine(ConvertToMarkdown(TaskElement.SelectSingleNode("summary")));
+						Writer.WriteLine();
+
+						// Document the parameters
+						List<string[]> Rows = new List<string[]>();
+						foreach(string ParameterName in Task.NameToParameter.Keys)
+						{
+							// Get the parameter data
+							ScriptTaskParameter Parameter = Task.NameToParameter[ParameterName];
+
+							// Get the documentation for this parameter
+							XmlElement ParameterElement;
+							if(MemberNameToElement.TryGetValue("F:" + Parameter.FieldInfo.DeclaringType.FullName + "." + Parameter.Name, out ParameterElement))
+							{
+								string TypeName = Parameter.FieldInfo.FieldType.Name;
+								if(Parameter.ValidationType != TaskParameterValidationType.Default)
+								{
+									StringBuilder NewTypeName = new StringBuilder(Parameter.ValidationType.ToString());
+									for(int Idx = 1; Idx < NewTypeName.Length; Idx++)
+									{
+										if(Char.IsLower(NewTypeName[Idx - 1]) && Char.IsUpper(NewTypeName[Idx]))
+										{
+											NewTypeName.Insert(Idx, ' ');
+										}
+									}
+									TypeName = NewTypeName.ToString();
+								}
+
+								string[] Columns = new string[4];
+								Columns[0] = ParameterName;
+								Columns[1] = TypeName;
+								Columns[2] = Parameter.bOptional? "Optional" : "Required";
+								Columns[3] = ConvertToMarkdown(ParameterElement.SelectSingleNode("summary"));
+								Rows.Add(Columns);
+							}
+						}
+
+						// Always include the "If" attribute
+						string[] IfColumns = new string[4];
+						IfColumns[0] = "If";
+						IfColumns[1] = "Condition";
+						IfColumns[2] = "Optional";
+						IfColumns[3] = "Whether to execute this task. It is ignored if this condition evaluates to false.";
+						Rows.Add(IfColumns);
+
+						// Get the width of each column
+						int[] Widths = new int[4];
+						for(int Idx = 0; Idx < 4; Idx++)
+						{
+							Widths[Idx] = Rows.Max(x => x[Idx].Length);
+						}
+
+						// Format the markdown table
+						string Format = String.Format("| {{0,-{0}}} | {{1,-{1}}} | {{2,-{2}}} | {{3,-{3}}} |", Widths[0], Widths[1], Widths[2], Widths[3]);
+						Writer.WriteLine(Format, "", "", "", "");
+						Writer.WriteLine(Format, new string('-', Widths[0]), new string('-', Widths[1]), new string('-', Widths[2]), new string('-', Widths[3]));
+						for(int Idx = 0; Idx < Rows.Count; Idx++)
+						{
+							Writer.WriteLine(Format, Rows[Idx][0], Rows[Idx][1], Rows[Idx][2], Rows[Idx][3]);
+						}
+
+						// Blank line before next task
+						Writer.WriteLine();
+					}
+				}
+			}
+		}
+
+		/// <summary>
+		/// Converts an XML documentation node to markdown
+		/// </summary>
+		/// <param name="Node">The node to read</param>
+		/// <returns>Text in markdown format</returns>
+		static string ConvertToMarkdown(XmlNode Node)
+		{
+			string Text = Node.InnerXml;
+
+			StringBuilder Result = new StringBuilder();
+			for(int Idx = 0; Idx < Text.Length; Idx++)
+			{
+				if(Char.IsWhiteSpace(Text[Idx]))
+				{
+					Result.Append(' ');
+					while(Idx + 1 < Text.Length && Char.IsWhiteSpace(Text[Idx + 1]))
+					{
+						Idx++;
+					}
+				}
+				else
+				{
+					Result.Append(Text[Idx]);
+				}
+			}
+			return Result.ToString().Trim();
+		}
 	}
 
 	/// <summary>
diff --git a/Engine/Source/Programs/AutomationTool/BuildGraph/Graph.cs b/Engine/Source/Programs/AutomationTool/BuildGraph/Graph.cs
--- a/Engine/Source/Programs/AutomationTool/BuildGraph/Graph.cs
+++ b/Engine/Source/Programs/AutomationTool/BuildGraph/Graph.cs
@@ -90,6 +90,8 @@ namespace AutomationTool
 		/// Constructor
 		/// </summary>
 		/// <param name="Name">The name of this option</param>
+		/// <param name="Description">Description of the option, for display on help pages</param>
+		/// <param name="DefaultValue">Default value for the option</param>
 		public GraphOption(string Name, string Description, string DefaultValue)
 		{
 			this.Name = Name;
@@ -332,6 +334,7 @@ namespace AutomationTool
 		/// Writes a preprocessed build graph to a script file
 		/// </summary>
 		/// <param name="File">The file to load</param>
+		/// <param name="SchemaFile">Schema file for validation</param>
 		public void Write(FileReference File, FileReference SchemaFile)
 		{
 			XmlWriterSettings Settings = new XmlWriterSettings();
@@ -549,7 +552,7 @@ namespace AutomationTool
 		/// <summary>
 		/// Print the contents of the graph
 		/// </summary>
-		/// <param name="NodeToState">Mapping of node to its current state</param>
+		/// <param name="CompletedNodes">Set of nodes which are already complete</param>
 		/// <param name="PrintOptions">Options for how to print the graph</param>
 		public void Print(HashSet<Node> CompletedNodes, GraphPrintOptions PrintOptions)
 		{
diff --git a/Engine/Source/Programs/AutomationTool/BuildGraph/Node.cs b/Engine/Source/Programs/AutomationTool/BuildGraph/Node.cs
--- a/Engine/Source/Programs/AutomationTool/BuildGraph/Node.cs
+++ b/Engine/Source/Programs/AutomationTool/BuildGraph/Node.cs
@@ -143,7 +143,6 @@ namespace AutomationTool
 		/// Build all the tasks for this node
 		/// </summary>
 		/// <param name="Job">Information about the current job</param>
-		/// <param name="BuildProducts">Set of build products produced by this node.</param>
 		/// <param name="TagNameToFileSet">Mapping from tag names to the set of files they include. Should be set to contain the node inputs on entry.</param>
 		/// <returns>Whether the task succeeded or not. Exiting with an exception will be caught and treated as a failure.</returns>
 		public bool Build(JobContext Job, Dictionary<string, HashSet<FileReference>> TagNameToFileSet)
diff --git a/Engine/Source/Programs/AutomationTool/BuildGraph/Schema.cs b/Engine/Source/Programs/AutomationTool/BuildGraph/Schema.cs
--- a/Engine/Source/Programs/AutomationTool/BuildGraph/Schema.cs
+++ b/Engine/Source/Programs/AutomationTool/BuildGraph/Schema.cs
@@ -224,7 +224,7 @@ namespace AutomationTool
 		/// <summary>
 		/// Constructor
 		/// </summary>
-		/// <param name="InTaskNameToReflectionInfo">Mapping of task name to information about how to construct it</param>
+		/// <param name="InNameToTask">Mapping of task name to information about how to construct it</param>
 		public ScriptSchema(Dictionary<string, ScriptTask> InNameToTask)
 		{
 			NameToTask = InNameToTask;
@@ -256,15 +256,11 @@ namespace AutomationTool
 				TaskType.Name = Task.Name + "TaskType";
 				foreach(ScriptTaskParameter Parameter in Task.NameToParameter.Values)
 				{
-					XmlQualifiedName SchemaTypeName;
-					if(Parameter.ValidationType == TaskParameterValidationType.Default)
+					XmlQualifiedName SchemaTypeName = GetQualifiedTypeName(Parameter.ValidationType);
+					if(SchemaTypeName == null)
 					{
 						SchemaTypeName = TypeToSchemaTypeName[Parameter.FieldInfo.FieldType];
 					}
-					else
-					{
-						SchemaTypeName = GetQualifiedTypeName(Parameter.ValidationType);
-					}
 					TaskType.Attributes.Add(CreateSchemaAttribute(Parameter.Name, SchemaTypeName, Parameter.bOptional? XmlSchemaUse.Optional : XmlSchemaUse.Required));
 				}
 				TaskType.Attributes.Add(CreateSchemaAttribute("If", ScriptSchemaStandardType.BalancedString, XmlSchemaUse.Optional));
@@ -386,9 +382,9 @@ namespace AutomationTool
 					return GetQualifiedTypeName(ScriptSchemaStandardType.Tag);
 				case TaskParameterValidationType.TagList:
 					return GetQualifiedTypeName(ScriptSchemaStandardType.TagList);
-				case TaskParameterValidationType.NameOrTag:
+				case TaskParameterValidationType.Target:
 					return GetQualifiedTypeName(ScriptSchemaStandardType.NameOrTag);
-				case TaskParameterValidationType.NameOrTagList:
+				case TaskParameterValidationType.TargetList:
 					return GetQualifiedTypeName(ScriptSchemaStandardType.NameOrTagList);
 			}
 			return null;
diff --git a/Engine/Source/Programs/AutomationTool/BuildGraph/Script.cs b/Engine/Source/Programs/AutomationTool/BuildGraph/Script.cs
--- a/Engine/Source/Programs/AutomationTool/BuildGraph/Script.cs
+++ b/Engine/Source/Programs/AutomationTool/BuildGraph/Script.cs
@@ -475,6 +475,7 @@ namespace AutomationTool
 		/// Reads the body of a trigger element
 		/// </summary>
 		/// <param name="Element">Xml element to read the definition from</param>
+		/// <param name="Trigger">The enclosing trigger definition</param>
 		void ReadTriggerBody(XmlElement Element, ManualTrigger Trigger)
 		{
 			EnterScope();
@@ -1220,7 +1221,6 @@ namespace AutomationTool
 		/// </summary>
 		/// <param name="Element">Xml element to read from</param>
 		/// <param name="Name">Name of the alias</param>
-		/// <param name="Nodes">Array of nodes that this name should resolve to</param>
 		/// <returns>True if the name was registered correctly, false otherwise.</returns>
 		bool CheckNameIsUnique(ScriptElement Element, string Name)
 		{
@@ -1410,7 +1410,7 @@ namespace AutomationTool
 		/// </summary>
 		/// <param name="Element">Element to read the attribute from</param>
 		/// <param name="Name">Name of the attribute</param>
-		/// <param name="DefaultValue">Default value if the attribute is missing</param>
+		/// <param name="bDefaultValue">Default value if the attribute is missing</param>
 		/// <returns>The value of the attribute field</returns>
 		bool ReadBooleanAttribute(ScriptElement Element, string Name, bool bDefaultValue)
 		{
diff --git a/Engine/Source/Programs/AutomationTool/BuildGraph/Task.cs b/Engine/Source/Programs/AutomationTool/BuildGraph/Task.cs
--- a/Engine/Source/Programs/AutomationTool/BuildGraph/Task.cs
+++ b/Engine/Source/Programs/AutomationTool/BuildGraph/Task.cs
@@ -41,14 +41,29 @@ namespace AutomationTool
 		TagList,
 
 		/// <summary>
-		/// A standard name or tag name
+		/// A standard node/aggregate/agent name or tag name
 		/// </summary>
-		NameOrTag,
+		Target,
 
 		/// <summary>
 		/// A list of standard name or tag names separated by semicolons
 		/// </summary>
-		NameOrTagList,
+		TargetList,
+
+		/// <summary>
+		/// A file specification, which may contain tags and wildcards.
+		/// </summary>
+		FileSpec,
+
+		/// <summary>
+		/// A single file name
+		/// </summary>
+		FileName,
+
+		/// <summary>
+		/// A single directory name
+		/// </summary>
+		DirectoryName,
 	}
 
 	/// <summary>
@@ -134,8 +149,8 @@ namespace AutomationTool
 		/// <summary>
 		/// Writes this task to an XML writer, using the given parameters object.
 		/// </summary>
-		/// <param name="Writer"></param>
-		/// <param name="Task"></param>
+		/// <param name="Writer">Writer for the XML schema</param>
+		/// <param name="Parameters">Parameters object that this task is constructed with</param>
 		protected void Write(XmlWriter Writer, object Parameters)
 		{
 			TaskElementAttribute Element = GetType().GetCustomAttribute<TaskElementAttribute>();
@@ -198,7 +213,7 @@ namespace AutomationTool
 		/// <summary>
 		/// Enumerates tag names from a list
 		/// </summary>
-		/// <param name="Filespec">A filespec, as can be passed to ResolveFilespec</param>
+		/// <param name="TagList">List of tags separated by semicolons</param>
 		/// <returns>Tag names from this filespec</returns>
 		protected IEnumerable<string> FindTagNamesFromList(string TagList)
 		{
@@ -252,6 +267,7 @@ namespace AutomationTool
 		/// <summary>
 		/// Finds or adds a set containing files with the given tag
 		/// </summary>
+		/// <param name="TagNameToFileSet">Map of tag names to the set of files they contain</param>
 		/// <param name="TagName">The tag name to return a set for. A leading '#' character is required.</param>
 		/// <returns>Set of files</returns>
 		public static HashSet<FileReference> FindOrAddTagSet(Dictionary<string, HashSet<FileReference>> TagNameToFileSet, string TagName)
@@ -323,7 +339,7 @@ namespace AutomationTool
 		/// Resolve a list of files, tag names or file specifications as above, but preserves any directory references for further processing.
 		/// </summary>
 		/// <param name="DefaultDirectory">The default directory to resolve relative paths to</param>
-		/// <param name="DelimitedPatterns">List of files, tag names, or file specifications to include separated by semicolons.</param>
+		/// <param name="FilePatterns">List of files, tag names, or file specifications to include separated by semicolons.</param>
 		/// <param name="ExcludePatterns">Set of patterns to apply to directory searches. This can greatly speed up enumeration by earlying out of recursive directory searches if large directories are excluded (eg. .../Intermediate/...).</param>
 		/// <param name="TagNameToFileSet">Mapping of tag name to fileset, as passed to the Execute() method</param>
 		/// <returns>Set of matching files.</returns>
diff --git a/Engine/Source/Programs/AutomationTool/BuildGraph/Tasks/CommandTask.cs b/Engine/Source/Programs/AutomationTool/BuildGraph/Tasks/CommandTask.cs
--- a/Engine/Source/Programs/AutomationTool/BuildGraph/Tasks/CommandTask.cs
+++ b/Engine/Source/Programs/AutomationTool/BuildGraph/Tasks/CommandTask.cs
@@ -35,7 +35,7 @@ namespace BuildGraph.Tasks
 	}
 
 	/// <summary>
-	/// Implements a task which calls another UAT command
+	/// Invokes an AutomationTool child process to run the given command.
 	/// </summary>
 	[TaskElement("Command", typeof(CommandTaskParameters))]
 	public class CommandTask : CustomTask
diff --git a/Engine/Source/Programs/AutomationTool/BuildGraph/Tasks/CommandletTask.cs b/Engine/Source/Programs/AutomationTool/BuildGraph/Tasks/CommandletTask.cs
--- a/Engine/Source/Programs/AutomationTool/BuildGraph/Tasks/CommandletTask.cs
+++ b/Engine/Source/Programs/AutomationTool/BuildGraph/Tasks/CommandletTask.cs
@@ -22,7 +22,7 @@ namespace AutomationTool.Tasks
 		/// <summary>
 		/// The project to run the editor with
 		/// </summary>
-		[TaskParameter(Optional = true)]
+		[TaskParameter(Optional = true, ValidationType = TaskParameterValidationType.FileSpec)]
 		public string Project;
 
 		/// <summary>
@@ -34,12 +34,12 @@ namespace AutomationTool.Tasks
 		/// <summary>
 		/// The editor executable to use. Defaults to the development UE4Editor executable for the current platform.
 		/// </summary>
-		[TaskParameter(Optional = true)]
+		[TaskParameter(Optional = true, ValidationType = TaskParameterValidationType.FileName)]
 		public string EditorExe;
 	}
 
 	/// <summary>
-	/// Implements a task which calls another UAT command
+	/// Spawns the editor to run a commandlet.
 	/// </summary>
 	[TaskElement("Commandlet", typeof(CommandletTaskParameters))]
 	public class CommandletTask : CustomTask
diff --git a/Engine/Source/Programs/AutomationTool/BuildGraph/Tasks/CompileTask.cs b/Engine/Source/Programs/AutomationTool/BuildGraph/Tasks/CompileTask.cs
--- a/Engine/Source/Programs/AutomationTool/BuildGraph/Tasks/CompileTask.cs
+++ b/Engine/Source/Programs/AutomationTool/BuildGraph/Tasks/CompileTask.cs
@@ -46,7 +46,7 @@ namespace AutomationTool
 	}
 
 	/// <summary>
-	/// Task which compiles a target with UnrealBuildTool
+	/// Compiles a target with UnrealBuildTool.
 	/// </summary>
 	[TaskElement("Compile", typeof(CompileTaskParameters))]
 	public class CompileTask : CustomTask
@@ -64,7 +64,7 @@ namespace AutomationTool
 		/// <summary>
 		/// Construct a compile task
 		/// </summary>
-		/// <param name="InParameters">Parameters for this task</param>
+		/// <param name="Parameters">Parameters for this task</param>
 		public CompileTask(CompileTaskParameters Parameters)
 		{
 			UE4Build.BuildTarget Target = new UE4Build.BuildTarget { TargetName = Parameters.Target, Platform = Parameters.Platform, Config = Parameters.Configuration, UBTArgs = "-nobuilduht " + (Parameters.Arguments ?? "") };
diff --git a/Engine/Source/Programs/AutomationTool/BuildGraph/Tasks/CopyTask.cs b/Engine/Source/Programs/AutomationTool/BuildGraph/Tasks/CopyTask.cs
--- a/Engine/Source/Programs/AutomationTool/BuildGraph/Tasks/CopyTask.cs
+++ b/Engine/Source/Programs/AutomationTool/BuildGraph/Tasks/CopyTask.cs
@@ -17,19 +17,19 @@ namespace BuildGraph.Tasks
 		/// <summary>
 		/// List of file specifications separated by semicolons (eg. *.cpp;Engine/.../*.bat), or the name of a tag set. Relative paths are based at FromDir.
 		/// </summary>
-		[TaskParameter]
+		[TaskParameter(ValidationType = TaskParameterValidationType.FileSpec)]
 		public string Files;
 
 		/// <summary>
 		/// The base directory to copy from. 
 		/// </summary>
-		[TaskParameter]
+		[TaskParameter(ValidationType = TaskParameterValidationType.DirectoryName)]
 		public string FromDir;
 
 		/// <summary>
 		/// The directory to copy to
 		/// </summary>
-		[TaskParameter]
+		[TaskParameter(ValidationType = TaskParameterValidationType.DirectoryName)]
 		public string ToDir;
 
 		/// <summary>
@@ -40,7 +40,7 @@ namespace BuildGraph.Tasks
 	}
 
 	/// <summary>
-	/// Task which copies files from one directory to another
+	/// Copies files from one directory to another.
 	/// </summary>
 	[TaskElement("Copy", typeof(CopyTaskParameters))]
 	public class CopyTask : CustomTask
diff --git a/Engine/Source/Programs/AutomationTool/BuildGraph/Tasks/CsCompileTask.cs b/Engine/Source/Programs/AutomationTool/BuildGraph/Tasks/CsCompileTask.cs
--- a/Engine/Source/Programs/AutomationTool/BuildGraph/Tasks/CsCompileTask.cs
+++ b/Engine/Source/Programs/AutomationTool/BuildGraph/Tasks/CsCompileTask.cs
@@ -37,18 +37,6 @@ namespace AutomationTool.Tasks
 		[TaskParameter(Optional = true)]
 		public string Arguments;
 
-		/// <summary>
-		/// Directory containing output files
-		/// </summary>
-		[TaskParameter(Optional = true)]
-		public string OutputDir;
-
-		/// <summary>
-		/// Patterns for output files
-		/// </summary>
-		[TaskParameter(Optional = true)]
-		public string OutputFiles;
-
 		/// <summary>
 		/// Only enumerate build products; do not actually compile the projects.
 		/// </summary>
@@ -70,7 +58,7 @@ namespace AutomationTool.Tasks
 	}
 
 	/// <summary>
-	/// Compile a C# project file
+	/// Compiles C# project files, and their dependencies.
 	/// </summary>
 	[TaskElement("CsCompile", typeof(CsCompileTaskParameters))]
 	public class CsCompileTask : CustomTask
@@ -487,7 +475,7 @@ namespace AutomationTool.Tasks
 		/// </summary>
 		/// <param name="BaseDirectory">Directory to resolve relative paths against</param>
 		/// <param name="ParentElement">The parent 'Reference' element</param>
-		/// <param name="ProjectReferences">Dictionary of project files to a bool indicating whether the assembly should be copied locally to the referencing project.</param>
+		/// <param name="References">Dictionary of project files to a bool indicating whether the assembly should be copied locally to the referencing project.</param>
 		static void ParseReference(DirectoryReference BaseDirectory, XmlElement ParentElement, Dictionary<FileReference, bool> References)
 		{
 			string HintPath = GetChildElementString(ParentElement, "HintPath", null);
diff --git a/Engine/Source/Programs/AutomationTool/BuildGraph/Tasks/DeleteTask.cs b/Engine/Source/Programs/AutomationTool/BuildGraph/Tasks/DeleteTask.cs
--- a/Engine/Source/Programs/AutomationTool/BuildGraph/Tasks/DeleteTask.cs
+++ b/Engine/Source/Programs/AutomationTool/BuildGraph/Tasks/DeleteTask.cs
@@ -18,18 +18,18 @@ namespace BuildGraph.Tasks
 		/// <summary>
 		/// List of file specifications separated by semicolons (eg. *.cpp;Engine/.../*.bat), or the name of a tag set
 		/// </summary>
-		[TaskParameter]
+		[TaskParameter(ValidationType = TaskParameterValidationType.FileSpec)]
 		public string Files;
 
 		/// <summary>
-		/// Whether to delete empty directories after deleting the files
+		/// Whether to delete empty directories after deleting the files. Defaults to true.
 		/// </summary>
 		[TaskParameter(Optional = true)]
 		public bool DeleteEmptyDirectories = true;
 	}
 
 	/// <summary>
-	/// Task which copies files from one directory to another
+	/// Delete a set of files.
 	/// </summary>
 	[TaskElement("Delete", typeof(DeleteTaskParameters))]
 	public class DeleteTask : CustomTask
@@ -92,7 +92,7 @@ namespace BuildGraph.Tasks
 		/// <summary>
 		/// Deletes a directory, if it's empty
 		/// </summary>
-		/// <param name="Directory">The directory to check</param>
+		/// <param name="CandidateDirectory">The directory to check</param>
 		/// <returns>True if the directory was deleted, false if not</returns>
 		static bool TryDeleteEmptyDirectory(DirectoryReference CandidateDirectory)
 		{
diff --git a/Engine/Source/Programs/AutomationTool/BuildGraph/Tasks/LogTask.cs b/Engine/Source/Programs/AutomationTool/BuildGraph/Tasks/LogTask.cs
--- a/Engine/Source/Programs/AutomationTool/BuildGraph/Tasks/LogTask.cs
+++ b/Engine/Source/Programs/AutomationTool/BuildGraph/Tasks/LogTask.cs
@@ -24,7 +24,7 @@ namespace BuildGraph.Tasks
 		/// <summary>
 		/// If specified, causes the given list of files to be printed after the given message.
 		/// </summary>
-		[TaskParameter(Optional = true)]
+		[TaskParameter(Optional = true, ValidationType = TaskParameterValidationType.FileSpec)]
 		public string Files;
 
 		/// <summary>
diff --git a/Engine/Source/Programs/AutomationTool/BuildGraph/Tasks/PakFileTask.cs b/Engine/Source/Programs/AutomationTool/BuildGraph/Tasks/PakFileTask.cs
--- a/Engine/Source/Programs/AutomationTool/BuildGraph/Tasks/PakFileTask.cs
+++ b/Engine/Source/Programs/AutomationTool/BuildGraph/Tasks/PakFileTask.cs
@@ -18,31 +18,31 @@ namespace BuildGraph.Tasks
 		/// <summary>
 		/// List of files, wildcards and tag sets to add to the pak file, separated by ';' characters.
 		/// </summary>
-		[TaskParameter]
+		[TaskParameter(ValidationType = TaskParameterValidationType.FileSpec)]
 		public string Files;
 
 		/// <summary>
 		/// PAK file to output
 		/// </summary>
-		[TaskParameter]
+		[TaskParameter(ValidationType = TaskParameterValidationType.FileName)]
 		public string Output;
 
 		/// <summary>
 		/// Path to a Response File that contains a list of files to add to the pak file, instead of specifying them individually
 		/// </summary>
-		[TaskParameter(Optional = true)]
+		[TaskParameter(Optional = true, ValidationType = TaskParameterValidationType.FileName)]
 		public string ResponseFile;
 
 		/// <summary>
 		/// Directories to rebase the files relative to. If specified, the shortest path under a listed directory will be used for each file.
 		/// </summary>
-		[TaskParameter(Optional = true)]
+		[TaskParameter(Optional = true, ValidationType = TaskParameterValidationType.DirectoryName)]
 		public string RebaseDir;
 
 		/// <summary>
 		/// Script which gives the order of files
 		/// </summary>
-		[TaskParameter(Optional = true)]
+		[TaskParameter(Optional = true, ValidationType = TaskParameterValidationType.FileName)]
 		public string Order;
 
 		/// <summary>
@@ -71,7 +71,7 @@ namespace BuildGraph.Tasks
 	}
 
 	/// <summary>
-	/// Cook a selection of maps for a certain platform
+	/// Creates a PAK file from a given set of files.
 	/// </summary>
 	[TaskElement("PakFile", typeof(PakFileTaskParameters))]
 	public class PakFileTask : CustomTask
diff --git a/Engine/Source/Programs/AutomationTool/BuildGraph/Tasks/RenameTask.cs b/Engine/Source/Programs/AutomationTool/BuildGraph/Tasks/RenameTask.cs
--- a/Engine/Source/Programs/AutomationTool/BuildGraph/Tasks/RenameTask.cs
+++ b/Engine/Source/Programs/AutomationTool/BuildGraph/Tasks/RenameTask.cs
@@ -17,32 +17,32 @@ namespace BuildGraph.Tasks
 	public class RenameTaskParameters
 	{
 		/// <summary>
-		/// The file to rename
+		/// The file or files to rename
 		/// </summary>
-		[TaskParameter]
+		[TaskParameter(ValidationType = TaskParameterValidationType.FileSpec)]
 		public string Files;
 
 		/// <summary>
-		/// The new name of the file. Should not include any path information.
+		/// The current file name, or pattern to match (eg. *.txt). Should not include any path separators.
 		/// </summary>
-		[TaskParameter(Optional = true)]
+		[TaskParameter(Optional = true, ValidationType = TaskParameterValidationType.DirectoryName)]
 		public string From;
 
 		/// <summary>
-		/// The new name of the file. Should not include any path information.
+		/// The new name for the file(s). Should not include any path separators.
 		/// </summary>
-		[TaskParameter]
+		[TaskParameter(ValidationType = TaskParameterValidationType.DirectoryName)]
 		public string To;
 
 		/// <summary>
-		/// Tag to be applied to build products of this task
+		/// Tag to be applied to the renamed files
 		/// </summary>
 		[TaskParameter(Optional = true, ValidationType = TaskParameterValidationType.TagList)]
 		public string Tag;
 	}
 
 	/// <summary>
-	/// Task which renames a file, or group of files
+	/// Renames a file, or group of files.
 	/// </summary>
 	[TaskElement("Rename", typeof(RenameTaskParameters))]
 	public class RenameTask : CustomTask
diff --git a/Engine/Source/Programs/AutomationTool/BuildGraph/Tasks/SetVersionTask.cs b/Engine/Source/Programs/AutomationTool/BuildGraph/Tasks/SetVersionTask.cs
--- a/Engine/Source/Programs/AutomationTool/BuildGraph/Tasks/SetVersionTask.cs
+++ b/Engine/Source/Programs/AutomationTool/BuildGraph/Tasks/SetVersionTask.cs
@@ -57,7 +57,7 @@ namespace AutomationTool.Tasks
 	}
 
 	/// <summary>
-	/// Task which updates the version files in the current branch
+	/// Updates the local version files (Engine/Source/Runtime/Launch/Resources/Version.h, Engine/Build/Build.version, and Engine/Source/Programs/DotNETCommon/Metadata.cs) with the given version information.
 	/// </summary>
 	[TaskElement("SetVersion", typeof(SetVersionTaskParameters))]
 	public class SetVersionTask : CustomTask
diff --git a/Engine/Source/Programs/AutomationTool/BuildGraph/Tasks/SignTask.cs b/Engine/Source/Programs/AutomationTool/BuildGraph/Tasks/SignTask.cs
--- a/Engine/Source/Programs/AutomationTool/BuildGraph/Tasks/SignTask.cs
+++ b/Engine/Source/Programs/AutomationTool/BuildGraph/Tasks/SignTask.cs
@@ -17,7 +17,7 @@ namespace BuildGraph.Tasks
 		/// <summary>
 		/// List of file specifications separated by semicolons (eg. *.cpp;Engine/.../*.bat), or the name of a tag set
 		/// </summary>
-		[TaskParameter]
+		[TaskParameter(ValidationType = TaskParameterValidationType.FileSpec)]
 		public string Files;
 
 		/// <summary>
@@ -28,7 +28,7 @@ namespace BuildGraph.Tasks
 	}
 
 	/// <summary>
-	/// Task which strips symbols from a set of files
+	/// Signs a set of executable files with an installed certificate.
 	/// </summary>
 	[TaskElement("Sign", typeof(SignTaskParameters))]
 	public class SignTask : CustomTask
@@ -41,7 +41,7 @@ namespace BuildGraph.Tasks
 		/// <summary>
 		/// Construct a spawn task
 		/// </summary>
-		/// <param name="Parameters">Parameters for the task</param>
+		/// <param name="InParameters">Parameters for the task</param>
 		public SignTask(SignTaskParameters InParameters)
 		{
 			Parameters = InParameters;
diff --git a/Engine/Source/Programs/AutomationTool/BuildGraph/Tasks/SpawnTask.cs b/Engine/Source/Programs/AutomationTool/BuildGraph/Tasks/SpawnTask.cs
--- a/Engine/Source/Programs/AutomationTool/BuildGraph/Tasks/SpawnTask.cs
+++ b/Engine/Source/Programs/AutomationTool/BuildGraph/Tasks/SpawnTask.cs
@@ -17,7 +17,7 @@ namespace BuildGraph.Tasks
 		/// <summary>
 		/// Executable to spawn
 		/// </summary>
-		[TaskParameter]
+		[TaskParameter(ValidationType = TaskParameterValidationType.FileName)]
 		public string Exe;
 
 		/// <summary>
@@ -34,7 +34,7 @@ namespace BuildGraph.Tasks
 	}
 
 	/// <summary>
-	/// Task which spawns an external executable
+	/// Spawns an external executable and waits for it to complete.
 	/// </summary>
 	[TaskElement("Spawn", typeof(SpawnTaskParameters))]
 	public class SpawnTask : CustomTask
@@ -47,7 +47,7 @@ namespace BuildGraph.Tasks
 		/// <summary>
 		/// Construct a spawn task
 		/// </summary>
-		/// <param name="Parameters">Parameters for the task</param>
+		/// <param name="InParameters">Parameters for the task</param>
 		public SpawnTask(SpawnTaskParameters InParameters)
 		{
 			Parameters = InParameters;
diff --git a/Engine/Source/Programs/AutomationTool/BuildGraph/Tasks/StageTask.cs b/Engine/Source/Programs/AutomationTool/BuildGraph/Tasks/StageTask.cs
--- a/Engine/Source/Programs/AutomationTool/BuildGraph/Tasks/StageTask.cs
+++ b/Engine/Source/Programs/AutomationTool/BuildGraph/Tasks/StageTask.cs
@@ -17,7 +17,7 @@ namespace BuildGraph.Tasks
 		/// <summary>
 		/// The project that this target belongs to
 		/// </summary>
-		[TaskParameter(Optional = true)]
+		[TaskParameter(Optional = true, ValidationType = TaskParameterValidationType.FileName)]
 		public string Project;
 
 		/// <summary>
@@ -44,18 +44,6 @@ namespace BuildGraph.Tasks
 		[TaskParameter(Optional = true)]
 		public string Architecture;
 
-		/// <summary>
-		/// Subset of files which should be staged. This filter is run against the source directory.
-		/// </summary>
-		[TaskParameter(Optional = true)]
-		public string Files;
-
-		/// <summary>
-		/// Files which should not be staged. This filter is run against the source directory.
-		/// </summary>
-		[TaskParameter(Optional = true)]
-		public string Exclude;
-
 		/// <summary>
 		/// Directory the receipt files should be staged to
 		/// </summary>
@@ -76,7 +64,7 @@ namespace BuildGraph.Tasks
 	}
 
 	/// <summary>
-	/// Stages files from a receipt to an output directory
+	/// Stages files listed in a build receipt to an output directory.
 	/// </summary>
 	[TaskElement("Stage", typeof(StageTaskParameters))]
 	public class StageTask : CustomTask
@@ -201,14 +189,7 @@ namespace BuildGraph.Tasks
 		/// <returns>The tag names which are read by this task</returns>
 		public override IEnumerable<string> FindConsumedTagNames()
 		{
-			foreach(string TagName in FindTagNamesFromFilespec(Parameters.Files))
-			{
-				yield return TagName;
-			}
-			foreach(string TagName in FindTagNamesFromFilespec(Parameters.Exclude))
-			{
-				yield return TagName;
-			}
+			yield break;
 		}
 
 		/// <summary>
diff --git a/Engine/Source/Programs/AutomationTool/BuildGraph/Tasks/StripTask.cs b/Engine/Source/Programs/AutomationTool/BuildGraph/Tasks/StripTask.cs
--- a/Engine/Source/Programs/AutomationTool/BuildGraph/Tasks/StripTask.cs
+++ b/Engine/Source/Programs/AutomationTool/BuildGraph/Tasks/StripTask.cs
@@ -23,19 +23,19 @@ namespace BuildGraph.Tasks
 		/// <summary>
 		/// The directory to find files in
 		/// </summary>
-		[TaskParameter(Optional = true)]
+		[TaskParameter(Optional = true, ValidationType = TaskParameterValidationType.DirectoryName)]
 		public string BaseDir;
 
 		/// <summary>
-		/// List of file specifications separated by semicolons (eg. *.cpp;Engine/.../*.bat), or the name of a tag set
+		/// List of file specifications separated by semicolons (eg. Engine/.../*.pdb), or the name of a tag set
 		/// </summary>
-		[TaskParameter]
+		[TaskParameter(ValidationType = TaskParameterValidationType.FileSpec)]
 		public string Files;
 
 		/// <summary>
 		/// Output directory for the stripped files. Defaults to the input path (overwriting the input files).
 		/// </summary>
-		[TaskParameter]
+		[TaskParameter(ValidationType = TaskParameterValidationType.DirectoryName)]
 		public string OutputDir;
 
 		/// <summary>
@@ -46,7 +46,7 @@ namespace BuildGraph.Tasks
 	}
 
 	/// <summary>
-	/// Task which strips symbols from a set of files
+	/// Strips debugging information from a set of files.
 	/// </summary>
 	[TaskElement("Strip", typeof(StripTaskParameters))]
 	public class StripTask : CustomTask
@@ -59,7 +59,7 @@ namespace BuildGraph.Tasks
 		/// <summary>
 		/// Construct a spawn task
 		/// </summary>
-		/// <param name="Parameters">Parameters for the task</param>
+		/// <param name="InParameters">Parameters for the task</param>
 		public StripTask(StripTaskParameters InParameters)
 		{
 			Parameters = InParameters;
diff --git a/Engine/Source/Programs/AutomationTool/BuildGraph/Tasks/SubmitTask.cs b/Engine/Source/Programs/AutomationTool/BuildGraph/Tasks/SubmitTask.cs
--- a/Engine/Source/Programs/AutomationTool/BuildGraph/Tasks/SubmitTask.cs
+++ b/Engine/Source/Programs/AutomationTool/BuildGraph/Tasks/SubmitTask.cs
@@ -22,23 +22,23 @@ namespace AutomationTool.Tasks
 		/// <summary>
 		/// The files to submit
 		/// </summary>
-		[TaskParameter]
+		[TaskParameter(ValidationType = TaskParameterValidationType.FileSpec)]
 		public string Files;
 
 		/// <summary>
-		/// The filetype for the submitted files
+		/// The Perforce file type for the submitted files (eg. binary+FS32)
 		/// </summary>
 		[TaskParameter(Optional = true)]
 		public string FileType;
 
 		/// <summary>
-		/// The workspace name. If specified, a new workspace will be created using the given stream and root directory to submit the files.
+		/// The workspace name. If specified, a new workspace will be created using the given stream and root directory to submit the files. If not, the current workspace will be used.
 		/// </summary>
 		[TaskParameter(Optional=true)]
 		public string Workspace;
 
 		/// <summary>
-		/// The stream for the workspace. If not specified, defaults to the current stream.
+		/// The stream for the workspace; defaults to the current stream. Ignored unless If the Workspace attribute is also specified.
 		/// </summary>
 		[TaskParameter(Optional = true)]
 		public string Stream;
@@ -46,12 +46,12 @@ namespace AutomationTool.Tasks
 		/// <summary>
 		/// Root directory for the stream. If not specified, defaults to the current root directory.
 		/// </summary>
-		[TaskParameter(Optional = true)]
+		[TaskParameter(Optional = true, ValidationType = TaskParameterValidationType.DirectoryName)]
 		public string RootDir;
 	}
 
 	/// <summary>
-	/// Task which submits a the version files in the current branch
+	/// Creates a new changelist and submits a set of files to a Perforce stream.
 	/// </summary>
 	[TaskElement("Submit", typeof(SubmitTaskParameters))]
 	public class SubmitTask : CustomTask
diff --git a/Engine/Source/Programs/AutomationTool/BuildGraph/Tasks/TagReceiptTask.cs b/Engine/Source/Programs/AutomationTool/BuildGraph/Tasks/TagReceiptTask.cs
--- a/Engine/Source/Programs/AutomationTool/BuildGraph/Tasks/TagReceiptTask.cs
+++ b/Engine/Source/Programs/AutomationTool/BuildGraph/Tasks/TagReceiptTask.cs
@@ -15,23 +15,21 @@ namespace AutomationTool.Tasks
 	public class TagReceiptTaskParameters
 	{
 		/// <summary>
-		/// Set of *.target files to read from, including wildcards and tag names, separated by semicolons. Resolved relative to the branch root directory.
+		/// Set of receipt files (*.target) to read, including wildcards and tag names, separated by semicolons.
 		/// </summary>
-		[TaskParameter]
+		[TaskParameter(ValidationType = TaskParameterValidationType.FileSpec)]
 		public string Files;
 
 		/// <summary>
-		/// Directory relative to the root directory that the Engine data is in - used to replace the $(EngineDir) property used in *.target files
-		/// If unspecified will default to Root/Engine
+		/// Path to the Engine folder, used to expand $(EngineDir) properties in receipt files. Defaults to the Engine directory for the current workspace.
 		/// </summary>
-		[TaskParameter(Optional = true)]
+		[TaskParameter(Optional = true, ValidationType = TaskParameterValidationType.DirectoryName)]
 		public string EngineDir;
 
 		/// <summary>
-		/// Directory relative to the root directory that the Project data is in - used to replace the $(ProjectDir) property used in *.target files
-		/// If unspecified will default to Root/Engine
+		/// Path to the project folder, used to expand $(ProjectDir) properties in receipt files. Defaults to the Engine directory for the current workspace.
 		/// </summary>
-		[TaskParameter(Optional = true)]
+		[TaskParameter(Optional = true, ValidationType = TaskParameterValidationType.DirectoryName)]
 		public string ProjectDir;
 
 		/// <summary>
@@ -79,8 +77,6 @@ namespace AutomationTool.Tasks
 
 	/// <summary>
 	/// Task which tags build products and/or runtime dependencies by reading from *.target files.
-	/// Filtering is performed using the following method:
-	/// * A set of files is enumerated from the tags and file specifications given by the "Files" parameter
 	/// </summary>
 	[TaskElement("TagReceipt", typeof(TagReceiptTaskParameters))]
 	class TagReceiptTask : CustomTask
diff --git a/Engine/Source/Programs/AutomationTool/BuildGraph/Tasks/TagTask.cs b/Engine/Source/Programs/AutomationTool/BuildGraph/Tasks/TagTask.cs
--- a/Engine/Source/Programs/AutomationTool/BuildGraph/Tasks/TagTask.cs
+++ b/Engine/Source/Programs/AutomationTool/BuildGraph/Tasks/TagTask.cs
@@ -17,25 +17,25 @@ namespace AutomationTool.Tasks
 		/// <summary>
 		/// Set the base directory to resolve relative paths and patterns against. If set, any absolute patterns (eg. /Engine/Build/...) are taken to be relative to this path. If not, they are taken to be truly absolute.
 		/// </summary>
-		[TaskParameter(Optional = true)]
+		[TaskParameter(Optional = true, ValidationType = TaskParameterValidationType.DirectoryName)]
 		public string BaseDir;
 
 		/// <summary>
 		/// Set of files to work from, including wildcards and tag names, separated by semicolons. Resolved relative to BaseDir if set, otherwise to the branch root directory.
 		/// </summary>
-		[TaskParameter]
+		[TaskParameter(ValidationType = TaskParameterValidationType.FileSpec)]
 		public string Files;
 
 		/// <summary>
 		/// Patterns to filter the list of files by, including tag names or wildcards. May include patterns that apply to the base directory if set. Defaults to all files if not specified.
 		/// </summary>
-		[TaskParameter(Optional = true)]
+		[TaskParameter(Optional = true, ValidationType = TaskParameterValidationType.FileSpec)]
 		public string Filter;
 
 		/// <summary>
 		/// Set of patterns to exclude from the matched list. May include tag names of patterns that apply to the base directory.
 		/// </summary>
-		[TaskParameter(Optional = true)]
+		[TaskParameter(Optional = true, ValidationType = TaskParameterValidationType.FileSpec)]
 		public string Except;
 
 		/// <summary>
@@ -46,10 +46,8 @@ namespace AutomationTool.Tasks
 	}
 
 	/// <summary>
-	/// Task which applies a tag to a given set of files. Filtering is performed using the following method:
-	/// * A set of files is enumerated from the tags and file specifications given by the "Files" parameter
-	/// * Any files not matched by the "Filter" parameter are removed.
-	/// * Any files matched by the "Except" parameter are removed.
+	/// Applies a tag to a given set of files. The list of files is found by enumerating the tags and file specifications given by the 'Files' 
+	/// parameter. From this list, any files not matched by the 'Filter' parameter are removed, followed by any files matched by the 'Except' parameter.
 	/// </summary>
 	[TaskElement("Tag", typeof(TagTaskParameters))]
 	class TagTask : CustomTask
@@ -63,7 +61,6 @@ namespace AutomationTool.Tasks
 		/// Constructor
 		/// </summary>
 		/// <param name="InParameters">Parameters to select which files to match</param>
-		/// <param name="Type"></param>
 		public TagTask(TagTaskParameters InParameters)
 		{
 			Parameters = InParameters;
diff --git a/Engine/Source/Programs/AutomationTool/BuildGraph/Tasks/UnzipTask.cs b/Engine/Source/Programs/AutomationTool/BuildGraph/Tasks/UnzipTask.cs
--- a/Engine/Source/Programs/AutomationTool/BuildGraph/Tasks/UnzipTask.cs
+++ b/Engine/Source/Programs/AutomationTool/BuildGraph/Tasks/UnzipTask.cs
@@ -15,26 +15,26 @@ namespace BuildGraph.Tasks
 	public class UnzipTaskParameters
 	{
 		/// <summary>
-		/// The directory to copy from
+		/// Path to the zip file to extract
 		/// </summary>
-		[TaskParameter]
+		[TaskParameter(ValidationType = TaskParameterValidationType.FileName)]
 		public string ZipFile;
 
 		/// <summary>
 		/// Output directory for the extracted files
 		/// </summary>
-		[TaskParameter]
+		[TaskParameter(ValidationType = TaskParameterValidationType.DirectoryName)]
 		public string ToDir;
 
 		/// <summary>
-		/// Tag to be applied to build products of this task
+		/// Tag to be applied to the extracted files
 		/// </summary>
 		[TaskParameter(Optional = true, ValidationType = TaskParameterValidationType.TagList)]
 		public string Tag;
 	}
 
 	/// <summary>
-	/// Task which creates a zip archive
+	/// Extract files from a zip archive.
 	/// </summary>
 	[TaskElement("Unzip", typeof(UnzipTaskParameters))]
 	public class UnzipTask : CustomTask
diff --git a/Engine/Source/Programs/AutomationTool/BuildGraph/Tasks/ZipTask.cs b/Engine/Source/Programs/AutomationTool/BuildGraph/Tasks/ZipTask.cs
--- a/Engine/Source/Programs/AutomationTool/BuildGraph/Tasks/ZipTask.cs
+++ b/Engine/Source/Programs/AutomationTool/BuildGraph/Tasks/ZipTask.cs
@@ -15,32 +15,32 @@ namespace BuildGraph.Tasks
 	public class ZipTaskParameters
 	{
 		/// <summary>
-		/// The directory to copy from
+		/// The directory to read compressed files from
 		/// </summary>
-		[TaskParameter]
+		[TaskParameter(ValidationType = TaskParameterValidationType.DirectoryName)]
 		public string FromDir;
 
 		/// <summary>
-		/// List of file specifications separated by semicolons (eg. *.cpp;Engine/.../*.bat), or the name of a tag set
+		/// List of file specifications separated by semicolons (eg. *.cpp;Engine/.../*.bat), or the name of a tag set. Relative paths are taken from FromDir.
 		/// </summary>
-		[TaskParameter(Optional = true)]
+		[TaskParameter(Optional = true, ValidationType = TaskParameterValidationType.FileSpec)]
 		public string Files;
 
 		/// <summary>
 		/// The zip file to create
 		/// </summary>
-		[TaskParameter]
+		[TaskParameter(ValidationType = TaskParameterValidationType.FileName)]
 		public string ZipFile;
 
 		/// <summary>
-		/// Tag to be applied to build products of this task
+		/// Tag to be applied to the created zip file
 		/// </summary>
 		[TaskParameter(Optional = true, ValidationType = TaskParameterValidationType.TagList)]
 		public string Tag;
 	}
 
 	/// <summary>
-	/// Task which creates a zip archive
+	/// Compresses files into a zip archive.
 	/// </summary>
 	[TaskElement("Zip", typeof(ZipTaskParameters))]
 	public class ZipTask : CustomTask
diff --git a/Engine/Source/Programs/AutomationTool/BuildGraph/TempStorage.cs b/Engine/Source/Programs/AutomationTool/BuildGraph/TempStorage.cs
--- a/Engine/Source/Programs/AutomationTool/BuildGraph/TempStorage.cs
+++ b/Engine/Source/Programs/AutomationTool/BuildGraph/TempStorage.cs
@@ -352,7 +352,7 @@ namespace AutomationTool
 		/// </summary>
 		/// <param name="InFiles">List of full file paths</param>
 		/// <param name="RootDir">Root folder for all the files. All files must be relative to this RootDir.</param>
-		/// <param name="InStorageBlocks">Referenced storage blocks required for these files</param>
+		/// <param name="InBlocks">Referenced storage blocks required for these files</param>
 		public TempStorageFileList(IEnumerable<FileReference> InFiles, DirectoryReference RootDir, IEnumerable<TempStorageBlock> InBlocks)
 		{
 			List<string> NewLocalFiles = new List<string>();
@@ -450,6 +450,7 @@ namespace AutomationTool
 		/// <param name="InRootDir">Root directory for this branch</param>
 		/// <param name="InLocalDir">The local temp storage directory.</param>
 		/// <param name="InSharedDir">The shared temp storage directory. May be null.</param>
+		/// <param name="bInWriteToSharedStorage">Whether to write to shared storage, or only permit reads from it</param>
 		public TempStorage(DirectoryReference InRootDir, DirectoryReference InLocalDir, DirectoryReference InSharedDir, bool bInWriteToSharedStorage)
 		{
 			RootDir = InRootDir;
@@ -665,6 +666,8 @@ namespace AutomationTool
 		/// </summary>
 		/// <param name="NodeName">Name of the node which produced the tag set</param>
 		/// <param name="TagName">Name of the tag, with a '#' prefix</param>
+		/// <param name="Files">List of files in this set</param>
+		/// <param name="Blocks">List of referenced storage blocks</param>
 		/// <returns>The set of files</returns>
 		public void WriteFileList(string NodeName, string TagName, IEnumerable<FileReference> Files, IEnumerable<TempStorageBlock> Blocks)
 		{
@@ -744,7 +747,7 @@ namespace AutomationTool
 		/// Retrieve an output of the given node. Fetches and decompresses the files from shared storage if necessary, or validates the local files.
 		/// </summary>
 		/// <param name="NodeName">The node which created the storage block</param>
-		/// <param name="BlockName">Name of the block to retrieve. May be null or empty.</param>
+		/// <param name="OutputName">Name of the block to retrieve. May be null or empty.</param>
 		/// <returns>Manifest of the files retrieved</returns>
 		public TempStorageManifest Retreive(string NodeName, string OutputName)
 		{
@@ -823,7 +826,7 @@ namespace AutomationTool
 		/// <summary>
 		/// Zips a set of files (that must be rooted at the given RootDir) to a set of zip files in the given OutputDir. The files will be prefixed with the given basename.
 		/// </summary>
-		/// <param name="Files">Fully qualified list of files to zip (must be rooted at RootDir).</param>
+		/// <param name="InputFiles">Fully qualified list of files to zip (must be rooted at RootDir).</param>
 		/// <param name="RootDir">Root Directory where all files will be extracted.</param>
 		/// <param name="OutputDir">Location to place the set of zip files created.</param>
 		/// <param name="StagingDir">Location to create zip files before copying them to the OutputDir. If the OutputDir is on a remote file share, staging may be more efficient. Use null to avoid using a staging copy.</param>
@@ -946,9 +949,8 @@ namespace AutomationTool
 		/// <summary>
 		/// Unzips a set of zip files with a given basename in a given folder to a given RootDir.
 		/// </summary>
+		/// <param name="ZipFiles">Files to extract</param>
 		/// <param name="RootDir">Root Directory where all files will be extracted.</param>
-		/// <param name="FolderWithZipFiles">Folder containing the zip files to unzip. None of the zips should have the same file path in them.</param>
-		/// <param name="ZipBasename">The basename of the set of zip files to unzip.</param>
 		/// <returns>Some metrics about the unzip process.</returns>
 		/// <remarks>
 		/// The code is expected to be the used as the symmetrical inverse of <see cref="ParallelZipFiles"/>, but could be used independently, as long as the files in the zip do not overlap.
@@ -1021,7 +1023,7 @@ namespace AutomationTool
 		/// </summary>
 		/// <param name="BaseDir">A local or shared temp storage root directory.</param>
 		/// <param name="NodeName">Name of the node to get the file for</param>
-		/// <param name="OutputName">Name of the output to get the manifest for</param>
+		/// <param name="BlockName">Name of the output block to get the manifest for</param>
 		static FileReference GetManifestLocation(DirectoryReference BaseDir, string NodeName, string BlockName)
 		{
 			return FileReference.Combine(BaseDir, NodeName, String.IsNullOrEmpty(BlockName)? "Manifest.xml" : String.Format("Manifest-{0}.xml", BlockName));
@@ -1032,7 +1034,7 @@ namespace AutomationTool
 		/// </summary>
 		/// <param name="BaseDir">A local or shared temp storage root directory.</param>
 		/// <param name="NodeName">Name of the node to get the file for</param>
-		/// <param name="OutputName">Name of the output to get the manifest for</param>
+		/// <param name="TagName">Name of the tag to get the manifest for</param>
 		static FileReference GetTaggedFileListLocation(DirectoryReference BaseDir, string NodeName, string TagName)
 		{
 			Debug.Assert(TagName.StartsWith("#"));
