Fix Mono warning (which is treated as an error) causing all dependent automation assemblies to fail to compile.

diff --git a/Engine/Source/Programs/AutomationTool/BuildGraph/BuildGraph.cs b/Engine/Source/Programs/AutomationTool/BuildGraph/BuildGraph.cs
--- a/Engine/Source/Programs/AutomationTool/BuildGraph/BuildGraph.cs
+++ b/Engine/Source/Programs/AutomationTool/BuildGraph/BuildGraph.cs
@@ -734,7 +734,7 @@ namespace AutomationTool
 
 			// Read documentation for each of them
 			Dictionary<string, XmlElement> MemberNameToElement = new Dictionary<string, XmlElement>();
-			foreach(Assembly TaskAssembly in NameToTask.Values.Select(x => x.ParametersClass.Assembly).Distinct())
+			foreach(Assembly TaskAssembly in TaskAssemblies)
 			{
 				string XmlFileName = Path.ChangeExtension(TaskAssembly.Location, ".xml");
 				if(File.Exists(XmlFileName))
