Merging fix CL 3097602 (by [REDACTED_PERSON]) from //UE4/Release-4.13 to Main (//UE4/Main) to remove AutomationTool compiler warnings caused by mis-matched processor architectures in its dependencies.

diff --git a/Engine/Source/Programs/AutomationTool/IOS/IOS.Automation.csproj b/Engine/Source/Programs/AutomationTool/IOS/IOS.Automation.csproj
--- a/Engine/Source/Programs/AutomationTool/IOS/IOS.Automation.csproj
+++ b/Engine/Source/Programs/AutomationTool/IOS/IOS.Automation.csproj
@@ -64,11 +64,6 @@
     <Compile Include="Properties\Resources.Designer.cs" />
   </ItemGroup>
   <ItemGroup>
-    <ProjectReference Include="..\..\IOS\DeploymentInterface\DeploymentInterface.csproj">
-      <Project>{2da5397a-f93f-41e0-a1ad-b220d40fd739}</Project>
-      <Name>DeploymentInterface</Name>
-      <Private>True</Private>
-    </ProjectReference>
     <ProjectReference Include="..\..\UnrealBuildTool\UnrealBuildTool.csproj">
       <Project>{fd7c5e1a-cfe4-4fd5-a525-1eb1599a39ac}</Project>
       <Name>UnrealBuildTool</Name>
diff --git a/Engine/Source/Programs/AutomationTool/TVOS/TVOS.Automation.csproj b/Engine/Source/Programs/AutomationTool/TVOS/TVOS.Automation.csproj
--- a/Engine/Source/Programs/AutomationTool/TVOS/TVOS.Automation.csproj
+++ b/Engine/Source/Programs/AutomationTool/TVOS/TVOS.Automation.csproj
@@ -59,11 +59,6 @@
     <Compile Include="Properties\AssemblyInfo.cs" />
   </ItemGroup>
   <ItemGroup>
-    <ProjectReference Include="..\..\IOS\DeploymentInterface\DeploymentInterface.csproj">
-      <Project>{2da5397a-f93f-41e0-a1ad-b220d40fd739}</Project>
-      <Name>DeploymentInterface</Name>
-      <Private>True</Private>
-    </ProjectReference>
     <ProjectReference Include="..\..\UnrealBuildTool\UnrealBuildTool.csproj">
       <Project>{fd7c5e1a-cfe4-4fd5-a525-1eb1599a39ac}</Project>
       <Name>UnrealBuildTool</Name>
