Added path to MobileDeviceInterface for IOS.Automation project

diff --git a/Engine/Source/Programs/AutomationTool/IOS/IOS.Automation.csproj b/Engine/Source/Programs/AutomationTool/IOS/IOS.Automation.csproj
--- a/Engine/Source/Programs/AutomationTool/IOS/IOS.Automation.csproj
+++ b/Engine/Source/Programs/AutomationTool/IOS/IOS.Automation.csproj
@@ -96,7 +96,7 @@
       <Name>AutomationUtils.Automation</Name>
       <Private>False</Private>
     </ProjectReference>
-    <ProjectReference Include="MobileDeviceInterface">
+    <ProjectReference Include="..\..\IOS\MobileDeviceInterface\MobileDeviceInterface.csproj">
       <Project>{A5CA2153-90B2-42E4-9547-341C28A7726A}</Project>
       <Name>MobileDeviceInterface</Name>
     </ProjectReference>
