+ Standarize on AnyCPU for all platforms
 + Remove other non AnyCPU platforms
 + Ensure all AnyCPU targets actually use the AnyCPU architecture
 + Clean up inconsistent references and paths

diff --git a/Engine/Source/Programs/AutomationTool/AllDesktop/AllDesktop.Automation.csproj b/Engine/Source/Programs/AutomationTool/AllDesktop/AllDesktop.Automation.csproj
--- a/Engine/Source/Programs/AutomationTool/AllDesktop/AllDesktop.Automation.csproj
+++ b/Engine/Source/Programs/AutomationTool/AllDesktop/AllDesktop.Automation.csproj
@@ -2,7 +2,7 @@
 <Project ToolsVersion="4.0" DefaultTargets="Build" xmlns="http://schemas.microsoft.com/developer/msbuild/2003">
   <Import Project="$(MSBuildExtensionsPath)\$(MSBuildToolsVersion)\Microsoft.Common.props" Condition="Exists('$(MSBuildExtensionsPath)\$(MSBuildToolsVersion)\Microsoft.Common.props')" />
   <PropertyGroup>
-    <Configuration Condition=" '$(Configuration)' == '' ">Debug</Configuration>
+    <Configuration Condition=" '$(Configuration)' == '' ">Development</Configuration>
     <Platform Condition=" '$(Platform)' == '' ">AnyCPU</Platform>
     <ProjectGuid>{A9CD1A06-764A-4C53-86CE-2F34B94F0283}</ProjectGuid>
     <OutputType>Library</OutputType>
@@ -33,25 +33,6 @@
     <PlatformTarget>AnyCPU</PlatformTarget>
     <TreatWarningsAsErrors>true</TreatWarningsAsErrors>
   </PropertyGroup>
-  <PropertyGroup Condition="'$(Configuration)|$(Platform)' == 'Debug|x64'">
-    <DebugSymbols>true</DebugSymbols>
-    <OutputPath>..\..\..\..\Binaries\DotNET\AutomationScripts\AllDesktop\</OutputPath>
-    <DefineConstants>DEBUG;TRACE</DefineConstants>
-    <TreatWarningsAsErrors>true</TreatWarningsAsErrors>
-    <DebugType>full</DebugType>
-    <PlatformTarget>x64</PlatformTarget>
-    <ErrorReport>prompt</ErrorReport>
-    <CodeAnalysisRuleSet>MinimumRecommendedRules.ruleset</CodeAnalysisRuleSet>
-  </PropertyGroup>
-  <PropertyGroup Condition="'$(Configuration)|$(Platform)' == 'Development|x64'">
-    <OutputPath>..\..\..\..\Binaries\DotNET\AutomationScripts\AllDesktop\</OutputPath>
-    <DefineConstants>TRACE</DefineConstants>
-    <TreatWarningsAsErrors>true</TreatWarningsAsErrors>
-    <DebugType>pdbonly</DebugType>
-    <PlatformTarget>x64</PlatformTarget>
-    <ErrorReport>prompt</ErrorReport>
-    <CodeAnalysisRuleSet>MinimumRecommendedRules.ruleset</CodeAnalysisRuleSet>
-  </PropertyGroup>
   <ItemGroup>
     <Reference Include="System" />
     <Reference Include="System.Core" />
@@ -88,4 +69,4 @@
     </ProjectReference>
   </ItemGroup>
   <Import Project="$(MSBuildToolsPath)\Microsoft.CSharp.targets" />
-</Project>
\ No newline at end of file
+</Project>
diff --git a/Engine/Source/Programs/AutomationTool/Android/Android.Automation.csproj b/Engine/Source/Programs/AutomationTool/Android/Android.Automation.csproj
--- a/Engine/Source/Programs/AutomationTool/Android/Android.Automation.csproj
+++ b/Engine/Source/Programs/AutomationTool/Android/Android.Automation.csproj
@@ -2,7 +2,7 @@
 <Project ToolsVersion="4.0" DefaultTargets="Build" xmlns="http://schemas.microsoft.com/developer/msbuild/2003">
   <Import Project="$(MSBuildExtensionsPath)\$(MSBuildToolsVersion)\Microsoft.Common.props" Condition="Exists('$(MSBuildExtensionsPath)\$(MSBuildToolsVersion)\Microsoft.Common.props')" />
   <PropertyGroup>
-    <Configuration Condition=" '$(Configuration)' == '' ">Debug</Configuration>
+    <Configuration Condition=" '$(Configuration)' == '' ">Development</Configuration>
     <Platform Condition=" '$(Platform)' == '' ">AnyCPU</Platform>
     <ProjectGuid>{32E3A47E-359A-4F78-AE6C-703B87DE0B97}</ProjectGuid>
     <OutputType>Library</OutputType>
@@ -33,25 +33,6 @@
     <PlatformTarget>AnyCPU</PlatformTarget>
     <TreatWarningsAsErrors>true</TreatWarningsAsErrors>
   </PropertyGroup>
-  <PropertyGroup Condition="'$(Configuration)|$(Platform)' == 'Debug|x64'">
-    <DebugSymbols>true</DebugSymbols>
-    <OutputPath>..\..\..\..\Binaries\DotNET\AutomationScripts\Android\</OutputPath>
-    <DefineConstants>DEBUG;TRACE</DefineConstants>
-    <TreatWarningsAsErrors>true</TreatWarningsAsErrors>
-    <DebugType>full</DebugType>
-    <PlatformTarget>x64</PlatformTarget>
-    <ErrorReport>prompt</ErrorReport>
-    <CodeAnalysisRuleSet>MinimumRecommendedRules.ruleset</CodeAnalysisRuleSet>
-  </PropertyGroup>
-  <PropertyGroup Condition="'$(Configuration)|$(Platform)' == 'Development|x64'">
-    <OutputPath>..\..\..\..\Binaries\DotNET\AutomationScripts\Android\</OutputPath>
-    <DefineConstants>TRACE</DefineConstants>
-    <TreatWarningsAsErrors>true</TreatWarningsAsErrors>
-    <DebugType>pdbonly</DebugType>
-    <PlatformTarget>x64</PlatformTarget>
-    <ErrorReport>prompt</ErrorReport>
-    <CodeAnalysisRuleSet>MinimumRecommendedRules.ruleset</CodeAnalysisRuleSet>
-  </PropertyGroup>
   <ItemGroup>
     <Reference Include="OneSky, Version=1.0.0.0, Culture=neutral, processorArchitecture=MSIL">
       <HintPath>..\..\..\..\Binaries\DotNET\OneSky.dll</HintPath>
@@ -97,4 +78,4 @@
     </ProjectReference>
   </ItemGroup>
   <Import Project="$(MSBuildToolsPath)\Microsoft.CSharp.targets" />
-</Project>
\ No newline at end of file
+</Project>
diff --git a/Engine/Source/Programs/AutomationTool/AutomationTool.csproj b/Engine/Source/Programs/AutomationTool/AutomationTool.csproj
--- a/Engine/Source/Programs/AutomationTool/AutomationTool.csproj
+++ b/Engine/Source/Programs/AutomationTool/AutomationTool.csproj
@@ -1,8 +1,8 @@
 ﻿<?xml version="1.0" encoding="utf-8"?>
 <Project ToolsVersion="4.0" DefaultTargets="Build" xmlns="http://schemas.microsoft.com/developer/msbuild/2003">
   <PropertyGroup>
-    <Configuration Condition=" '$(Configuration)' == '' ">Debug</Configuration>
-    <Platform Condition=" '$(Platform)' == '' ">x86</Platform>
+    <Configuration Condition=" '$(Configuration)' == '' ">Development</Configuration>
+    <Platform Condition=" '$(Platform)' == '' ">AnyCPU</Platform>
     <ProductVersion>8.0.30703</ProductVersion>
     <SchemaVersion>2.0</SchemaVersion>
     <ProjectGuid>{A35B0261-D868-475E-A098-EB14CE6DC89F}</ProjectGuid>
@@ -15,44 +15,6 @@
     </TargetFrameworkProfile>
     <FileAlignment>512</FileAlignment>
   </PropertyGroup>
-  <PropertyGroup Condition=" '$(Configuration)|$(Platform)' == 'Debug|x86' ">
-    <PlatformTarget>AnyCPU</PlatformTarget>
-    <DebugSymbols>true</DebugSymbols>
-    <DebugType>full</DebugType>
-    <Optimize>false</Optimize>
-    <OutputPath>..\..\..\Intermediate\Build\AutomationTool\Debug\</OutputPath>
-    <DefineConstants>DEBUG;TRACE</DefineConstants>
-    <ErrorReport>prompt</ErrorReport>
-    <WarningLevel>4</WarningLevel>
-    <DocumentationFile>..\..\..\Intermediate\Build\UnrealAutomationTool\Debug\AutomationTool.XML</DocumentationFile>
-  </PropertyGroup>
-  <PropertyGroup Condition=" '$(Configuration)|$(Platform)' == 'Development|x86' ">
-    <PlatformTarget>AnyCPU</PlatformTarget>
-    <DebugType>pdbonly</DebugType>
-    <Optimize>true</Optimize>
-    <OutputPath>bin\Release\</OutputPath>
-    <DefineConstants>TRACE</DefineConstants>
-    <ErrorReport>prompt</ErrorReport>
-    <WarningLevel>4</WarningLevel>
-  </PropertyGroup>
-  <PropertyGroup Condition="'$(Configuration)|$(Platform)' == 'Debug|x64'">
-    <DebugSymbols>true</DebugSymbols>
-    <OutputPath>..\..\..\Binaries\DotNET\</OutputPath>
-    <DefineConstants>DEBUG;TRACE</DefineConstants>
-    <DebugType>full</DebugType>
-    <PlatformTarget>x64</PlatformTarget>
-    <ErrorReport>prompt</ErrorReport>
-    <CodeAnalysisRuleSet>MinimumRecommendedRules.ruleset</CodeAnalysisRuleSet>
-  </PropertyGroup>
-  <PropertyGroup Condition=" '$(Configuration)|$(Platform)' == 'Development|x64' ">
-    <OutputPath>..\..\..\Binaries\DotNET\</OutputPath>
-    <DefineConstants>TRACE</DefineConstants>
-    <Optimize>true</Optimize>
-    <DebugType>pdbonly</DebugType>
-    <PlatformTarget>x64</PlatformTarget>
-    <ErrorReport>prompt</ErrorReport>
-    <CodeAnalysisRuleSet>MinimumRecommendedRules.ruleset</CodeAnalysisRuleSet>
-  </PropertyGroup>
   <PropertyGroup Condition="'$(Configuration)|$(Platform)' == 'Debug|AnyCPU'">
     <DebugSymbols>true</DebugSymbols>
     <OutputPath>..\..\..\Binaries\DotNET\</OutputPath>
@@ -60,7 +22,7 @@
     <DocumentationFile>
     </DocumentationFile>
     <DebugType>full</DebugType>
-    <PlatformTarget>x64</PlatformTarget>
+    <PlatformTarget>AnyCPU</PlatformTarget>
     <ErrorReport>prompt</ErrorReport>
     <CodeAnalysisRuleSet>MinimumRecommendedRules.ruleset</CodeAnalysisRuleSet>
     <Prefer32Bit>false</Prefer32Bit>
@@ -70,7 +32,7 @@
     <DefineConstants>TRACE</DefineConstants>
     <Optimize>false</Optimize>
     <DebugType>pdbonly</DebugType>
-    <PlatformTarget>x64</PlatformTarget>
+    <PlatformTarget>AnyCPU</PlatformTarget>
     <ErrorReport>prompt</ErrorReport>
     <CodeAnalysisRuleSet>MinimumRecommendedRules.ruleset</CodeAnalysisRuleSet>
     <DocumentationFile>
@@ -140,4 +102,4 @@
   <Target Name="AfterBuild">
   </Target>
   -->
-</Project>
\ No newline at end of file
+</Project>
diff --git a/Engine/Source/Programs/AutomationTool/AutomationUtils/AutomationUtils.Automation.csproj b/Engine/Source/Programs/AutomationTool/AutomationUtils/AutomationUtils.Automation.csproj
--- a/Engine/Source/Programs/AutomationTool/AutomationUtils/AutomationUtils.Automation.csproj
+++ b/Engine/Source/Programs/AutomationTool/AutomationUtils/AutomationUtils.Automation.csproj
@@ -2,7 +2,7 @@
 <Project ToolsVersion="4.0" DefaultTargets="Build" xmlns="http://schemas.microsoft.com/developer/msbuild/2003">
   <Import Project="$(MSBuildExtensionsPath)\$(MSBuildToolsVersion)\Microsoft.Common.props" Condition="Exists('$(MSBuildExtensionsPath)\$(MSBuildToolsVersion)\Microsoft.Common.props')" />
   <PropertyGroup>
-    <Configuration Condition=" '$(Configuration)' == '' ">Debug</Configuration>
+    <Configuration Condition=" '$(Configuration)' == '' ">Development</Configuration>
     <Platform Condition=" '$(Platform)' == '' ">AnyCPU</Platform>
     <ProjectGuid>{2C96A7F2-B1A3-4258-8E0A-E588FF41A53E}</ProjectGuid>
     <OutputType>Library</OutputType>
@@ -33,25 +33,6 @@
     <PlatformTarget>AnyCPU</PlatformTarget>
     <TreatWarningsAsErrors>true</TreatWarningsAsErrors>
   </PropertyGroup>
-  <PropertyGroup Condition="'$(Configuration)|$(Platform)' == 'Debug|x64'">
-    <DebugSymbols>true</DebugSymbols>
-    <OutputPath>..\..\..\..\Binaries\DotNET\</OutputPath>
-    <DefineConstants>DEBUG;TRACE</DefineConstants>
-    <TreatWarningsAsErrors>true</TreatWarningsAsErrors>
-    <DebugType>full</DebugType>
-    <PlatformTarget>x64</PlatformTarget>
-    <ErrorReport>prompt</ErrorReport>
-    <CodeAnalysisRuleSet>MinimumRecommendedRules.ruleset</CodeAnalysisRuleSet>
-  </PropertyGroup>
-  <PropertyGroup Condition="'$(Configuration)|$(Platform)' == 'Development|x64'">
-    <OutputPath>..\..\..\..\Binaries\DotNET\</OutputPath>
-    <DefineConstants>TRACE</DefineConstants>
-    <TreatWarningsAsErrors>true</TreatWarningsAsErrors>
-    <DebugType>pdbonly</DebugType>
-    <PlatformTarget>x64</PlatformTarget>
-    <ErrorReport>prompt</ErrorReport>
-    <CodeAnalysisRuleSet>MinimumRecommendedRules.ruleset</CodeAnalysisRuleSet>
-  </PropertyGroup>
   <ItemGroup>
     <Reference Include="Ionic.Zip.Reduced">
       <HintPath>..\..\..\..\Binaries\DotNET\Ionic.Zip.Reduced.dll</HintPath>
@@ -118,4 +99,4 @@
     </ProjectReference>
   </ItemGroup>
   <Import Project="$(MSBuildToolsPath)\Microsoft.CSharp.targets" />
-</Project>
\ No newline at end of file
+</Project>
diff --git a/Engine/Source/Programs/AutomationTool/BuildGraph/BuildGraph.Automation.csproj b/Engine/Source/Programs/AutomationTool/BuildGraph/BuildGraph.Automation.csproj
--- a/Engine/Source/Programs/AutomationTool/BuildGraph/BuildGraph.Automation.csproj
+++ b/Engine/Source/Programs/AutomationTool/BuildGraph/BuildGraph.Automation.csproj
@@ -2,7 +2,7 @@
 <Project ToolsVersion="4.0" DefaultTargets="Build" xmlns="http://schemas.microsoft.com/developer/msbuild/2003">
   <Import Project="$(MSBuildExtensionsPath)\$(MSBuildToolsVersion)\Microsoft.Common.props" Condition="Exists('$(MSBuildExtensionsPath)\$(MSBuildToolsVersion)\Microsoft.Common.props')" />
   <PropertyGroup>
-    <Configuration Condition=" '$(Configuration)' == '' ">Debug</Configuration>
+    <Configuration Condition=" '$(Configuration)' == '' ">Development</Configuration>
     <Platform Condition=" '$(Platform)' == '' ">AnyCPU</Platform>
     <ProjectGuid>{A002361E-37F9-4124-AF82-CF0D393CB928}</ProjectGuid>
     <OutputType>Library</OutputType>
@@ -34,26 +34,6 @@
     <TreatWarningsAsErrors>true</TreatWarningsAsErrors>
     <DocumentationFile>..\..\..\..\Binaries\DotNET\AutomationScripts\BuildGraph.Automation.xml</DocumentationFile>
   </PropertyGroup>
-  <PropertyGroup Condition="'$(Configuration)|$(Platform)' == 'Debug|x64'">
-    <DebugSymbols>true</DebugSymbols>
-    <OutputPath>..\..\..\..\Binaries\DotNET\AutomationScripts\</OutputPath>
-    <DefineConstants>DEBUG;TRACE</DefineConstants>
-    <TreatWarningsAsErrors>true</TreatWarningsAsErrors>
-    <DebugType>full</DebugType>
-    <PlatformTarget>x64</PlatformTarget>
-    <ErrorReport>prompt</ErrorReport>
-    <CodeAnalysisRuleSet>MinimumRecommendedRules.ruleset</CodeAnalysisRuleSet>
-  </PropertyGroup>
-  <PropertyGroup Condition="'$(Configuration)|$(Platform)' == 'Development|x64'">
-    <OutputPath>..\..\..\..\Binaries\DotNET\AutomationScripts\</OutputPath>
-    <DefineConstants>TRACE</DefineConstants>
-    <DocumentationFile>..\..\..\..\Binaries\DotNET\AutomationScripts\BuildGraph.Automation.xml</DocumentationFile>
-    <TreatWarningsAsErrors>true</TreatWarningsAsErrors>
-    <DebugType>pdbonly</DebugType>
-    <PlatformTarget>x64</PlatformTarget>
-    <ErrorReport>prompt</ErrorReport>
-    <CodeAnalysisRuleSet>MinimumRecommendedRules.ruleset</CodeAnalysisRuleSet>
-  </PropertyGroup>
   <ItemGroup>
     <Reference Include="Ionic.Zip.Reduced">
       <HintPath>..\..\..\..\Binaries\DotNET\Ionic.Zip.Reduced.dll</HintPath>
@@ -133,4 +113,4 @@
   </ItemGroup>
   <ItemGroup />
   <Import Project="$(MSBuildToolsPath)\Microsoft.CSharp.targets" />
-</Project>
\ No newline at end of file
+</Project>
diff --git a/Engine/Source/Programs/AutomationTool/GUBP/GUBP.Automation.csproj b/Engine/Source/Programs/AutomationTool/GUBP/GUBP.Automation.csproj
--- a/Engine/Source/Programs/AutomationTool/GUBP/GUBP.Automation.csproj
+++ b/Engine/Source/Programs/AutomationTool/GUBP/GUBP.Automation.csproj
@@ -2,7 +2,7 @@
 <Project ToolsVersion="4.0" DefaultTargets="Build" xmlns="http://schemas.microsoft.com/developer/msbuild/2003">
   <Import Project="$(MSBuildExtensionsPath)\$(MSBuildToolsVersion)\Microsoft.Common.props" Condition="Exists('$(MSBuildExtensionsPath)\$(MSBuildToolsVersion)\Microsoft.Common.props')" />
   <PropertyGroup>
-    <Configuration Condition=" '$(Configuration)' == '' ">Debug</Configuration>
+    <Configuration Condition=" '$(Configuration)' == '' ">Development</Configuration>
     <Platform Condition=" '$(Platform)' == '' ">AnyCPU</Platform>
     <ProjectGuid>{9A1A118B-37FC-4F33-BBF4-B431B6E3B3AC}</ProjectGuid>
     <OutputType>Library</OutputType>
@@ -33,25 +33,6 @@
     <PlatformTarget>AnyCPU</PlatformTarget>
     <TreatWarningsAsErrors>true</TreatWarningsAsErrors>
   </PropertyGroup>
-  <PropertyGroup Condition="'$(Configuration)|$(Platform)' == 'Debug|x64'">
-    <DebugSymbols>true</DebugSymbols>
-    <OutputPath>..\..\..\..\Binaries\DotNET\AutomationScripts\</OutputPath>
-    <DefineConstants>DEBUG;TRACE</DefineConstants>
-    <TreatWarningsAsErrors>true</TreatWarningsAsErrors>
-    <DebugType>full</DebugType>
-    <PlatformTarget>x64</PlatformTarget>
-    <ErrorReport>prompt</ErrorReport>
-    <CodeAnalysisRuleSet>MinimumRecommendedRules.ruleset</CodeAnalysisRuleSet>
-  </PropertyGroup>
-  <PropertyGroup Condition="'$(Configuration)|$(Platform)' == 'Development|x64'">
-    <OutputPath>..\..\..\..\Binaries\DotNET\AutomationScripts\</OutputPath>
-    <DefineConstants>TRACE</DefineConstants>
-    <TreatWarningsAsErrors>true</TreatWarningsAsErrors>
-    <DebugType>pdbonly</DebugType>
-    <PlatformTarget>x64</PlatformTarget>
-    <ErrorReport>prompt</ErrorReport>
-    <CodeAnalysisRuleSet>MinimumRecommendedRules.ruleset</CodeAnalysisRuleSet>
-  </PropertyGroup>
   <ItemGroup>
     <Reference Include="Ionic.Zip.Reduced">
       <HintPath>..\..\..\..\Binaries\DotNET\Ionic.Zip.Reduced.dll</HintPath>
@@ -116,4 +97,4 @@
   </ItemGroup>
   <ItemGroup />
   <Import Project="$(MSBuildToolsPath)\Microsoft.CSharp.targets" />
-</Project>
\ No newline at end of file
+</Project>
diff --git a/Engine/Source/Programs/AutomationTool/HTML5/HTML5.Automation.csproj b/Engine/Source/Programs/AutomationTool/HTML5/HTML5.Automation.csproj
--- a/Engine/Source/Programs/AutomationTool/HTML5/HTML5.Automation.csproj
+++ b/Engine/Source/Programs/AutomationTool/HTML5/HTML5.Automation.csproj
@@ -2,7 +2,7 @@
 <Project ToolsVersion="4.0" DefaultTargets="Build" xmlns="http://schemas.microsoft.com/developer/msbuild/2003">
   <Import Project="$(MSBuildExtensionsPath)\$(MSBuildToolsVersion)\Microsoft.Common.props" Condition="Exists('$(MSBuildExtensionsPath)\$(MSBuildToolsVersion)\Microsoft.Common.props')" />
   <PropertyGroup>
-    <Configuration Condition=" '$(Configuration)' == '' ">Debug</Configuration>
+    <Configuration Condition=" '$(Configuration)' == '' ">Development</Configuration>
     <Platform Condition=" '$(Platform)' == '' ">AnyCPU</Platform>
     <ProjectGuid>{90FF4220-1CAE-4CCA-BBF6-7AAAED069F03}</ProjectGuid>
     <OutputType>Library</OutputType>
@@ -33,25 +33,6 @@
     <PlatformTarget>AnyCPU</PlatformTarget>
     <TreatWarningsAsErrors>true</TreatWarningsAsErrors>
   </PropertyGroup>
-  <PropertyGroup Condition="'$(Configuration)|$(Platform)' == 'Debug|x64'">
-    <DebugSymbols>true</DebugSymbols>
-    <OutputPath>..\..\..\..\Binaries\DotNET\AutomationScripts\HTML5\</OutputPath>
-    <DefineConstants>DEBUG;TRACE</DefineConstants>
-    <TreatWarningsAsErrors>true</TreatWarningsAsErrors>
-    <DebugType>full</DebugType>
-    <PlatformTarget>x64</PlatformTarget>
-    <ErrorReport>prompt</ErrorReport>
-    <CodeAnalysisRuleSet>MinimumRecommendedRules.ruleset</CodeAnalysisRuleSet>
-  </PropertyGroup>
-  <PropertyGroup Condition="'$(Configuration)|$(Platform)' == 'Development|x64'">
-    <OutputPath>..\..\..\..\Binaries\DotNET\AutomationScripts\HTML5\</OutputPath>
-    <DefineConstants>TRACE</DefineConstants>
-    <TreatWarningsAsErrors>true</TreatWarningsAsErrors>
-    <DebugType>pdbonly</DebugType>
-    <PlatformTarget>x64</PlatformTarget>
-    <ErrorReport>prompt</ErrorReport>
-    <CodeAnalysisRuleSet>MinimumRecommendedRules.ruleset</CodeAnalysisRuleSet>
-  </PropertyGroup>
   <ItemGroup>
     <Reference Include="Ionic.Zip.Reduced">
       <HintPath>..\..\..\..\Binaries\DotNET\Ionic.Zip.Reduced.dll</HintPath>
@@ -91,4 +72,4 @@
     </ProjectReference>
   </ItemGroup>
   <Import Project="$(MSBuildToolsPath)\Microsoft.CSharp.targets" />
-</Project>
\ No newline at end of file
+</Project>
diff --git a/Engine/Source/Programs/AutomationTool/IOS/IOS.Automation.csproj b/Engine/Source/Programs/AutomationTool/IOS/IOS.Automation.csproj
--- a/Engine/Source/Programs/AutomationTool/IOS/IOS.Automation.csproj
+++ b/Engine/Source/Programs/AutomationTool/IOS/IOS.Automation.csproj
@@ -1,8 +1,8 @@
-<?xml version="1.0" encoding="utf-8"?>
+﻿<?xml version="1.0" encoding="utf-8"?>
 <Project ToolsVersion="4.0" DefaultTargets="Build" xmlns="http://schemas.microsoft.com/developer/msbuild/2003">
   <Import Project="$(MSBuildExtensionsPath)\$(MSBuildToolsVersion)\Microsoft.Common.props" Condition="Exists('$(MSBuildExtensionsPath)\$(MSBuildToolsVersion)\Microsoft.Common.props')" />
   <PropertyGroup>
-    <Configuration Condition=" '$(Configuration)' == '' ">Debug</Configuration>
+    <Configuration Condition=" '$(Configuration)' == '' ">Development</Configuration>
     <Platform Condition=" '$(Platform)' == '' ">AnyCPU</Platform>
     <ProjectGuid>{E1F1261F-DA95-41BD-80C1-60FBE0D37D22}</ProjectGuid>
     <OutputType>Library</OutputType>
@@ -33,25 +33,6 @@
     <PlatformTarget>AnyCPU</PlatformTarget>
     <TreatWarningsAsErrors>true</TreatWarningsAsErrors>
   </PropertyGroup>
-  <PropertyGroup Condition="'$(Configuration)|$(Platform)' == 'Debug|x64'">
-    <DebugSymbols>true</DebugSymbols>
-    <OutputPath>..\..\..\..\Binaries\DotNET\AutomationScripts\IOS\</OutputPath>
-    <DefineConstants>DEBUG;TRACE</DefineConstants>
-    <TreatWarningsAsErrors>true</TreatWarningsAsErrors>
-    <DebugType>full</DebugType>
-    <PlatformTarget>x64</PlatformTarget>
-    <ErrorReport>prompt</ErrorReport>
-    <CodeAnalysisRuleSet>MinimumRecommendedRules.ruleset</CodeAnalysisRuleSet>
-  </PropertyGroup>
-  <PropertyGroup Condition="'$(Configuration)|$(Platform)' == 'Development|x64'">
-    <OutputPath>..\..\..\..\Binaries\DotNET\AutomationScripts\IOS\</OutputPath>
-    <DefineConstants>TRACE</DefineConstants>
-    <TreatWarningsAsErrors>true</TreatWarningsAsErrors>
-    <DebugType>pdbonly</DebugType>
-    <PlatformTarget>x64</PlatformTarget>
-    <ErrorReport>prompt</ErrorReport>
-    <CodeAnalysisRuleSet>MinimumRecommendedRules.ruleset</CodeAnalysisRuleSet>
-  </PropertyGroup>
   <ItemGroup>
     <Reference Include="BouncyCastle.Crypto">
       <HintPath>..\..\..\..\Binaries\ThirdParty\IOS\BouncyCastle.Crypto.dll</HintPath>
@@ -75,6 +56,7 @@
     <Reference Include="System.Xml" />
     <Reference Include="MobileDeviceInterface">
       <HintPath>..\..\..\..\Binaries\DotNET\IOS\MobileDeviceInterface.dll</HintPath>
+      <Private>False</Private>
     </Reference>
   </ItemGroup>
   <ItemGroup>
@@ -99,6 +81,7 @@
     <ProjectReference Include="..\..\IOS\MobileDeviceInterface\MobileDeviceInterface.csproj">
       <Project>{A5CA2153-90B2-42E4-9547-341C28A7726A}</Project>
       <Name>MobileDeviceInterface</Name>
+      <Private>False</Private>
     </ProjectReference>
   </ItemGroup>
   <Import Project="$(MSBuildToolsPath)\Microsoft.CSharp.targets" />
diff --git a/Engine/Source/Programs/AutomationTool/Linux/Linux.Automation.csproj b/Engine/Source/Programs/AutomationTool/Linux/Linux.Automation.csproj
--- a/Engine/Source/Programs/AutomationTool/Linux/Linux.Automation.csproj
+++ b/Engine/Source/Programs/AutomationTool/Linux/Linux.Automation.csproj
@@ -2,7 +2,7 @@
 <Project ToolsVersion="4.0" DefaultTargets="Build" xmlns="http://schemas.microsoft.com/developer/msbuild/2003">
   <Import Project="$(MSBuildExtensionsPath)\$(MSBuildToolsVersion)\Microsoft.Common.props" Condition="Exists('$(MSBuildExtensionsPath)\$(MSBuildToolsVersion)\Microsoft.Common.props')" />
   <PropertyGroup>
-    <Configuration Condition=" '$(Configuration)' == '' ">Debug</Configuration>
+    <Configuration Condition=" '$(Configuration)' == '' ">Developement</Configuration>
     <Platform Condition=" '$(Platform)' == '' ">AnyCPU</Platform>
     <ProjectGuid>{EC395254-405B-4030-92CE-E346CC50848D}</ProjectGuid>
     <OutputType>Library</OutputType>
@@ -33,25 +33,6 @@
     <PlatformTarget>AnyCPU</PlatformTarget>
     <TreatWarningsAsErrors>true</TreatWarningsAsErrors>
   </PropertyGroup>
-  <PropertyGroup Condition="'$(Configuration)|$(Platform)' == 'Debug|x64'">
-    <DebugSymbols>true</DebugSymbols>
-    <OutputPath>..\..\..\..\Binaries\DotNET\AutomationScripts\Linux\</OutputPath>
-    <DefineConstants>DEBUG;TRACE</DefineConstants>
-    <TreatWarningsAsErrors>true</TreatWarningsAsErrors>
-    <DebugType>full</DebugType>
-    <PlatformTarget>x64</PlatformTarget>
-    <ErrorReport>prompt</ErrorReport>
-    <CodeAnalysisRuleSet>MinimumRecommendedRules.ruleset</CodeAnalysisRuleSet>
-  </PropertyGroup>
-  <PropertyGroup Condition="'$(Configuration)|$(Platform)' == 'Development|x64'">
-    <OutputPath>..\..\..\..\Binaries\DotNET\AutomationScripts\Linux\</OutputPath>
-    <DefineConstants>TRACE</DefineConstants>
-    <TreatWarningsAsErrors>true</TreatWarningsAsErrors>
-    <DebugType>pdbonly</DebugType>
-    <PlatformTarget>x64</PlatformTarget>
-    <ErrorReport>prompt</ErrorReport>
-    <CodeAnalysisRuleSet>MinimumRecommendedRules.ruleset</CodeAnalysisRuleSet>
-  </PropertyGroup>
   <ItemGroup>
     <Reference Include="System" />
     <Reference Include="System.Core" />
@@ -84,4 +65,4 @@
     </ProjectReference>
   </ItemGroup>
   <Import Project="$(MSBuildToolsPath)\Microsoft.CSharp.targets" />
-</Project>
\ No newline at end of file
+</Project>
diff --git a/Engine/Source/Programs/AutomationTool/Localization/Localization.Automation.csproj b/Engine/Source/Programs/AutomationTool/Localization/Localization.Automation.csproj
--- a/Engine/Source/Programs/AutomationTool/Localization/Localization.Automation.csproj
+++ b/Engine/Source/Programs/AutomationTool/Localization/Localization.Automation.csproj
@@ -2,7 +2,7 @@
 <Project ToolsVersion="4.0" DefaultTargets="Build" xmlns="http://schemas.microsoft.com/developer/msbuild/2003">
   <Import Project="$(MSBuildExtensionsPath)\$(MSBuildToolsVersion)\Microsoft.Common.props" Condition="Exists('$(MSBuildExtensionsPath)\$(MSBuildToolsVersion)\Microsoft.Common.props')" />
   <PropertyGroup>
-    <Configuration Condition=" '$(Configuration)' == '' ">Debug</Configuration>
+    <Configuration Condition=" '$(Configuration)' == '' ">Development</Configuration>
     <Platform Condition=" '$(Platform)' == '' ">AnyCPU</Platform>
     <ProjectGuid>{C2F959B1-1727-4598-9C32-3D9ACCE72750}</ProjectGuid>
     <OutputType>Library</OutputType>
@@ -33,25 +33,6 @@
     <PlatformTarget>AnyCPU</PlatformTarget>
     <TreatWarningsAsErrors>true</TreatWarningsAsErrors>
   </PropertyGroup>
-  <PropertyGroup Condition="'$(Configuration)|$(Platform)' == 'Debug|x64'">
-    <DebugSymbols>true</DebugSymbols>
-    <OutputPath>..\..\..\..\Binaries\DotNET\AutomationScripts\</OutputPath>
-    <DefineConstants>DEBUG;TRACE</DefineConstants>
-    <TreatWarningsAsErrors>true</TreatWarningsAsErrors>
-    <DebugType>full</DebugType>
-    <PlatformTarget>x64</PlatformTarget>
-    <ErrorReport>prompt</ErrorReport>
-    <CodeAnalysisRuleSet>MinimumRecommendedRules.ruleset</CodeAnalysisRuleSet>
-  </PropertyGroup>
-  <PropertyGroup Condition="'$(Configuration)|$(Platform)' == 'Development|x64'">
-    <OutputPath>..\..\..\..\Binaries\DotNET\AutomationScripts\</OutputPath>
-    <DefineConstants>TRACE</DefineConstants>
-    <TreatWarningsAsErrors>true</TreatWarningsAsErrors>
-    <DebugType>pdbonly</DebugType>
-    <PlatformTarget>x64</PlatformTarget>
-    <ErrorReport>prompt</ErrorReport>
-    <CodeAnalysisRuleSet>MinimumRecommendedRules.ruleset</CodeAnalysisRuleSet>
-  </PropertyGroup>
   <ItemGroup>
     <Reference Include="System" />
     <Reference Include="System.Core" />
@@ -83,4 +64,4 @@
   </ItemGroup>
   <ItemGroup />
   <Import Project="$(MSBuildToolsPath)\Microsoft.CSharp.targets" />
-</Project>
\ No newline at end of file
+</Project>
diff --git a/Engine/Source/Programs/AutomationTool/Mac/Mac.Automation.csproj b/Engine/Source/Programs/AutomationTool/Mac/Mac.Automation.csproj
--- a/Engine/Source/Programs/AutomationTool/Mac/Mac.Automation.csproj
+++ b/Engine/Source/Programs/AutomationTool/Mac/Mac.Automation.csproj
@@ -2,7 +2,7 @@
 <Project ToolsVersion="4.0" DefaultTargets="Build" xmlns="http://schemas.microsoft.com/developer/msbuild/2003">
   <Import Project="$(MSBuildExtensionsPath)\$(MSBuildToolsVersion)\Microsoft.Common.props" Condition="Exists('$(MSBuildExtensionsPath)\$(MSBuildToolsVersion)\Microsoft.Common.props')" />
   <PropertyGroup>
-    <Configuration Condition=" '$(Configuration)' == '' ">Debug</Configuration>
+    <Configuration Condition=" '$(Configuration)' == '' ">Development</Configuration>
     <Platform Condition=" '$(Platform)' == '' ">AnyCPU</Platform>
     <ProjectGuid>{16778681-DAB8-46A4-89F6-927796467995}</ProjectGuid>
     <OutputType>Library</OutputType>
@@ -33,25 +33,6 @@
     <PlatformTarget>AnyCPU</PlatformTarget>
     <TreatWarningsAsErrors>true</TreatWarningsAsErrors>
   </PropertyGroup>
-  <PropertyGroup Condition="'$(Configuration)|$(Platform)' == 'Debug|x64'">
-    <DebugSymbols>true</DebugSymbols>
-    <OutputPath>..\..\..\..\Binaries\DotNET\AutomationScripts\Mac\</OutputPath>
-    <DefineConstants>DEBUG;TRACE</DefineConstants>
-    <TreatWarningsAsErrors>true</TreatWarningsAsErrors>
-    <DebugType>full</DebugType>
-    <PlatformTarget>x64</PlatformTarget>
-    <ErrorReport>prompt</ErrorReport>
-    <CodeAnalysisRuleSet>MinimumRecommendedRules.ruleset</CodeAnalysisRuleSet>
-  </PropertyGroup>
-  <PropertyGroup Condition="'$(Configuration)|$(Platform)' == 'Development|x64'">
-    <OutputPath>..\..\..\..\Binaries\DotNET\AutomationScripts\Mac\</OutputPath>
-    <DefineConstants>TRACE</DefineConstants>
-    <TreatWarningsAsErrors>true</TreatWarningsAsErrors>
-    <DebugType>pdbonly</DebugType>
-    <PlatformTarget>x64</PlatformTarget>
-    <ErrorReport>prompt</ErrorReport>
-    <CodeAnalysisRuleSet>MinimumRecommendedRules.ruleset</CodeAnalysisRuleSet>
-  </PropertyGroup>
   <ItemGroup>
     <Reference Include="System" />
     <Reference Include="System.Core" />
@@ -84,4 +65,4 @@
     </ProjectReference>
   </ItemGroup>
   <Import Project="$(MSBuildToolsPath)\Microsoft.CSharp.targets" />
-</Project>
\ No newline at end of file
+</Project>
diff --git a/Engine/Source/Programs/AutomationTool/OneSkyLocalization/OneSkyLocalization.Automation.csproj b/Engine/Source/Programs/AutomationTool/OneSkyLocalization/OneSkyLocalization.Automation.csproj
--- a/Engine/Source/Programs/AutomationTool/OneSkyLocalization/OneSkyLocalization.Automation.csproj
+++ b/Engine/Source/Programs/AutomationTool/OneSkyLocalization/OneSkyLocalization.Automation.csproj
@@ -2,7 +2,7 @@
 <Project ToolsVersion="4.0" DefaultTargets="Build" xmlns="http://schemas.microsoft.com/developer/msbuild/2003">
   <Import Project="$(MSBuildExtensionsPath)\$(MSBuildToolsVersion)\Microsoft.Common.props" Condition="Exists('$(MSBuildExtensionsPath)\$(MSBuildToolsVersion)\Microsoft.Common.props')" />
   <PropertyGroup>
-    <Configuration Condition=" '$(Configuration)' == '' ">Debug</Configuration>
+    <Configuration Condition=" '$(Configuration)' == '' ">Development</Configuration>
     <Platform Condition=" '$(Platform)' == '' ">AnyCPU</Platform>
     <ProjectGuid>{39369368-5603-46DA-BAAE-766BB3CB2B51}</ProjectGuid>
     <OutputType>Library</OutputType>
@@ -33,25 +33,6 @@
     <PlatformTarget>AnyCPU</PlatformTarget>
     <TreatWarningsAsErrors>true</TreatWarningsAsErrors>
   </PropertyGroup>
-  <PropertyGroup Condition="'$(Configuration)|$(Platform)' == 'Debug|x64'">
-    <DebugSymbols>true</DebugSymbols>
-    <OutputPath>..\..\..\..\Binaries\DotNET\AutomationScripts\</OutputPath>
-    <DefineConstants>DEBUG;TRACE</DefineConstants>
-    <TreatWarningsAsErrors>true</TreatWarningsAsErrors>
-    <DebugType>full</DebugType>
-    <PlatformTarget>x64</PlatformTarget>
-    <ErrorReport>prompt</ErrorReport>
-    <CodeAnalysisRuleSet>MinimumRecommendedRules.ruleset</CodeAnalysisRuleSet>
-  </PropertyGroup>
-  <PropertyGroup Condition="'$(Configuration)|$(Platform)' == 'Development|x64'">
-    <OutputPath>..\..\..\..\Binaries\DotNET\AutomationScripts\</OutputPath>
-    <DefineConstants>TRACE</DefineConstants>
-    <TreatWarningsAsErrors>true</TreatWarningsAsErrors>
-    <DebugType>pdbonly</DebugType>
-    <PlatformTarget>x64</PlatformTarget>
-    <ErrorReport>prompt</ErrorReport>
-    <CodeAnalysisRuleSet>MinimumRecommendedRules.ruleset</CodeAnalysisRuleSet>
-  </PropertyGroup>
   <ItemGroup>
     <Reference Include="OneSky, Version=1.0.0.0, Culture=neutral, processorArchitecture=MSIL">
       <HintPath>..\..\..\..\Binaries\DotNET\OneSky.dll</HintPath>
@@ -93,4 +74,4 @@
   </ItemGroup>
   <ItemGroup />
   <Import Project="$(MSBuildToolsPath)\Microsoft.CSharp.targets" />
-</Project>
\ No newline at end of file
+</Project>
diff --git a/Engine/Source/Programs/AutomationTool/Scripts/AutomationScripts.Automation.csproj b/Engine/Source/Programs/AutomationTool/Scripts/AutomationScripts.Automation.csproj
--- a/Engine/Source/Programs/AutomationTool/Scripts/AutomationScripts.Automation.csproj
+++ b/Engine/Source/Programs/AutomationTool/Scripts/AutomationScripts.Automation.csproj
@@ -1,8 +1,8 @@
-<?xml version="1.0" encoding="utf-8"?>
+﻿<?xml version="1.0" encoding="utf-8"?>
 <Project ToolsVersion="4.0" DefaultTargets="Build" xmlns="http://schemas.microsoft.com/developer/msbuild/2003">
   <Import Project="$(MSBuildExtensionsPath)\$(MSBuildToolsVersion)\Microsoft.Common.props" Condition="Exists('$(MSBuildExtensionsPath)\$(MSBuildToolsVersion)\Microsoft.Common.props')" />
   <PropertyGroup>
-    <Configuration Condition=" '$(Configuration)' == '' ">Debug</Configuration>
+    <Configuration Condition=" '$(Configuration)' == '' ">Development</Configuration>
     <Platform Condition=" '$(Platform)' == '' ">AnyCPU</Platform>
     <ProjectGuid>{8AA00D65-0954-4A27-AC0D-FB8B1106120F}</ProjectGuid>
     <OutputType>Library</OutputType>
@@ -33,31 +33,12 @@
     <PlatformTarget>AnyCPU</PlatformTarget>
     <TreatWarningsAsErrors>true</TreatWarningsAsErrors>
   </PropertyGroup>
-  <PropertyGroup Condition="'$(Configuration)|$(Platform)' == 'Debug|x64'">
-    <DebugSymbols>true</DebugSymbols>
-    <OutputPath>..\..\..\..\Binaries\DotNET\AutomationScripts\</OutputPath>
-    <DefineConstants>DEBUG;TRACE</DefineConstants>
-    <TreatWarningsAsErrors>true</TreatWarningsAsErrors>
-    <DebugType>full</DebugType>
-    <PlatformTarget>x64</PlatformTarget>
-    <ErrorReport>prompt</ErrorReport>
-    <CodeAnalysisRuleSet>MinimumRecommendedRules.ruleset</CodeAnalysisRuleSet>
-  </PropertyGroup>
-  <PropertyGroup Condition="'$(Configuration)|$(Platform)' == 'Development|x64'">
-    <OutputPath>..\..\..\..\Binaries\DotNET\AutomationScripts\</OutputPath>
-    <DefineConstants>TRACE</DefineConstants>
-    <TreatWarningsAsErrors>true</TreatWarningsAsErrors>
-    <DebugType>pdbonly</DebugType>
-    <PlatformTarget>x64</PlatformTarget>
-    <ErrorReport>prompt</ErrorReport>
-    <CodeAnalysisRuleSet>MinimumRecommendedRules.ruleset</CodeAnalysisRuleSet>
-  </PropertyGroup>
   <ItemGroup>
     <Reference Include="Ionic.Zip.Reduced">
       <HintPath>..\..\..\..\Binaries\DotNET\Ionic.Zip.Reduced.dll</HintPath>
       <Private>False</Private>
     </Reference>
-    <Reference Include="OneSky, Version=1.0.0.0, Culture=neutral, processorArchitecture=MSIL">
+    <Reference Include="OneSky">
       <HintPath>..\..\..\..\Binaries\DotNET\OneSky.dll</HintPath>
       <SpecificVersion>False</SpecificVersion>
       <Private>False</Private>
@@ -73,13 +54,16 @@
     <Reference Include="System.Data" />
     <Reference Include="System.Xml" />
     <Reference Include="OneSkyLocalization.Automation">
-      <HintPath>..\OneSkyLocalization\obj\Development\OneSkyLocalization.Automation.dll</HintPath>
+      <HintPath>..\..\..\..\Binaries\DotNET\AutomationScripts\OneSkyLocalization.Automation.dll</HintPath>
+      <Private>False</Private>
     </Reference>
     <Reference Include="Localization.Automation">
-      <HintPath>..\Localization\obj\Development\Localization.Automation.dll</HintPath>
+      <HintPath>..\..\..\..\Binaries\DotNET\AutomationScripts\Localization.Automation.dll</HintPath>
+      <Private>False</Private>
     </Reference>
     <Reference Include="GUBP.Automation">
-      <HintPath>..\GUBP\obj\Development\GUBP.Automation.dll</HintPath>
+      <HintPath>..\..\..\..\Binaries\DotNET\AutomationScripts\GUBP.Automation.dll</HintPath>
+      <Private>False</Private>
     </Reference>
   </ItemGroup>
   <ItemGroup>
@@ -153,6 +137,10 @@
       <Project>{39369368-5603-46da-baae-766bb3cb2b51}</Project>
       <Name>OneSkyLocalization.Automation</Name>
     </ProjectReference>
+    <ProjectReference Include="..\GUBP\GUBP.Automation.csproj">
+      <Project>{9A1A118B-37FC-4F33-BBF4-B431B6E3B3AC}</Project>
+      <Name>OneSkyLocalization.Automation</Name>
+    </ProjectReference>
   </ItemGroup>
   <Import Project="$(MSBuildToolsPath)\Microsoft.CSharp.targets" />
 </Project>
\ No newline at end of file
diff --git a/Engine/Source/Programs/AutomationTool/TVOS/TVOS.Automation.csproj b/Engine/Source/Programs/AutomationTool/TVOS/TVOS.Automation.csproj
--- a/Engine/Source/Programs/AutomationTool/TVOS/TVOS.Automation.csproj
+++ b/Engine/Source/Programs/AutomationTool/TVOS/TVOS.Automation.csproj
@@ -2,7 +2,7 @@
 <Project ToolsVersion="4.0" DefaultTargets="Build" xmlns="http://schemas.microsoft.com/developer/msbuild/2003">
   <Import Project="$(MSBuildExtensionsPath)\$(MSBuildToolsVersion)\Microsoft.Common.props" Condition="Exists('$(MSBuildExtensionsPath)\$(MSBuildToolsVersion)\Microsoft.Common.props')" />
   <PropertyGroup>
-    <Configuration Condition=" '$(Configuration)' == '' ">Debug</Configuration>
+    <Configuration Condition=" '$(Configuration)' == '' ">Development</Configuration>
     <Platform Condition=" '$(Platform)' == '' ">AnyCPU</Platform>
     <ProjectGuid>{E1F1261F-DA95-41BD-81C1-61FBE1D37D22}</ProjectGuid>
     <OutputType>Library</OutputType>
@@ -33,25 +33,6 @@
     <PlatformTarget>AnyCPU</PlatformTarget>
     <TreatWarningsAsErrors>true</TreatWarningsAsErrors>
   </PropertyGroup>
-  <PropertyGroup Condition="'$(Configuration)|$(Platform)' == 'Debug|x64'">
-    <DebugSymbols>true</DebugSymbols>
-    <OutputPath>..\..\..\..\Binaries\DotNET\AutomationScripts\TVOS\</OutputPath>
-    <DefineConstants>DEBUG;TRACE</DefineConstants>
-    <TreatWarningsAsErrors>true</TreatWarningsAsErrors>
-    <DebugType>full</DebugType>
-    <PlatformTarget>x64</PlatformTarget>
-    <ErrorReport>prompt</ErrorReport>
-    <CodeAnalysisRuleSet>MinimumRecommendedRules.ruleset</CodeAnalysisRuleSet>
-  </PropertyGroup>
-  <PropertyGroup Condition="'$(Configuration)|$(Platform)' == 'Development|x64'">
-    <OutputPath>..\..\..\..\Binaries\DotNET\AutomationScripts\TVOS\</OutputPath>
-    <DefineConstants>TRACE</DefineConstants>
-    <TreatWarningsAsErrors>true</TreatWarningsAsErrors>
-    <DebugType>pdbonly</DebugType>
-    <PlatformTarget>x64</PlatformTarget>
-    <ErrorReport>prompt</ErrorReport>
-    <CodeAnalysisRuleSet>MinimumRecommendedRules.ruleset</CodeAnalysisRuleSet>
-  </PropertyGroup>
   <ItemGroup>
     <Reference Include="Ionic.Zip.Reduced">
       <HintPath>..\..\..\..\Binaries\DotNET\Ionic.Zip.Reduced.dll</HintPath>
@@ -95,4 +76,4 @@
     </ProjectReference>
   </ItemGroup>
   <Import Project="$(MSBuildToolsPath)\Microsoft.CSharp.targets" />
-</Project>
\ No newline at end of file
+</Project>
diff --git a/Engine/Source/Programs/AutomationTool/Win/Win.Automation.csproj b/Engine/Source/Programs/AutomationTool/Win/Win.Automation.csproj
--- a/Engine/Source/Programs/AutomationTool/Win/Win.Automation.csproj
+++ b/Engine/Source/Programs/AutomationTool/Win/Win.Automation.csproj
@@ -2,7 +2,7 @@
 <Project ToolsVersion="4.0" DefaultTargets="Build" xmlns="http://schemas.microsoft.com/developer/msbuild/2003">
   <Import Project="$(MSBuildExtensionsPath)\$(MSBuildToolsVersion)\Microsoft.Common.props" Condition="Exists('$(MSBuildExtensionsPath)\$(MSBuildToolsVersion)\Microsoft.Common.props')" />
   <PropertyGroup>
-    <Configuration Condition=" '$(Configuration)' == '' ">Debug</Configuration>
+    <Configuration Condition=" '$(Configuration)' == '' ">Development</Configuration>
     <Platform Condition=" '$(Platform)' == '' ">AnyCPU</Platform>
     <ProjectGuid>{6F6F8E5B-53CA-4C9F-A696-96AA9B8FB128}</ProjectGuid>
     <OutputType>Library</OutputType>
@@ -33,25 +33,6 @@
     <PlatformTarget>AnyCPU</PlatformTarget>
     <TreatWarningsAsErrors>true</TreatWarningsAsErrors>
   </PropertyGroup>
-  <PropertyGroup Condition="'$(Configuration)|$(Platform)' == 'Debug|x64'">
-    <DebugSymbols>true</DebugSymbols>
-    <OutputPath>..\..\..\..\Binaries\DotNET\AutomationScripts\</OutputPath>
-    <DefineConstants>DEBUG;TRACE</DefineConstants>
-    <TreatWarningsAsErrors>true</TreatWarningsAsErrors>
-    <DebugType>full</DebugType>
-    <PlatformTarget>x64</PlatformTarget>
-    <ErrorReport>prompt</ErrorReport>
-    <CodeAnalysisRuleSet>MinimumRecommendedRules.ruleset</CodeAnalysisRuleSet>
-  </PropertyGroup>
-  <PropertyGroup Condition="'$(Configuration)|$(Platform)' == 'Development|x64'">
-    <OutputPath>..\..\..\..\Binaries\DotNET\AutomationScripts\</OutputPath>
-    <DefineConstants>TRACE</DefineConstants>
-    <TreatWarningsAsErrors>true</TreatWarningsAsErrors>
-    <DebugType>pdbonly</DebugType>
-    <PlatformTarget>x64</PlatformTarget>
-    <ErrorReport>prompt</ErrorReport>
-    <CodeAnalysisRuleSet>MinimumRecommendedRules.ruleset</CodeAnalysisRuleSet>
-  </PropertyGroup>
   <ItemGroup>
     <Reference Include="System" />
     <Reference Include="System.Core" />
@@ -89,4 +70,4 @@
     </ProjectReference>
   </ItemGroup>
   <Import Project="$(MSBuildToolsPath)\Microsoft.CSharp.targets" />
-</Project>
\ No newline at end of file
+</Project>
diff --git a/Engine/Source/Programs/AutomationTool/XLocLocalization/XLocLocalization.Automation.csproj b/Engine/Source/Programs/AutomationTool/XLocLocalization/XLocLocalization.Automation.csproj
--- a/Engine/Source/Programs/AutomationTool/XLocLocalization/XLocLocalization.Automation.csproj
+++ b/Engine/Source/Programs/AutomationTool/XLocLocalization/XLocLocalization.Automation.csproj
@@ -2,7 +2,7 @@
 <Project ToolsVersion="4.0" DefaultTargets="Build" xmlns="http://schemas.microsoft.com/developer/msbuild/2003">
   <Import Project="$(MSBuildExtensionsPath)\$(MSBuildToolsVersion)\Microsoft.Common.props" Condition="Exists('$(MSBuildExtensionsPath)\$(MSBuildToolsVersion)\Microsoft.Common.props')" />
   <PropertyGroup>
-    <Configuration Condition=" '$(Configuration)' == '' ">Debug</Configuration>
+    <Configuration Condition=" '$(Configuration)' == '' ">Development</Configuration>
     <Platform Condition=" '$(Platform)' == '' ">AnyCPU</Platform>
     <ProjectGuid>{7E2D6784-C9B7-40C6-A89B-0A8AFDDCE60D}</ProjectGuid>
     <OutputType>Library</OutputType>
@@ -33,25 +33,6 @@
     <PlatformTarget>AnyCPU</PlatformTarget>
     <TreatWarningsAsErrors>true</TreatWarningsAsErrors>
   </PropertyGroup>
-  <PropertyGroup Condition="'$(Configuration)|$(Platform)' == 'Debug|x64'">
-    <DebugSymbols>true</DebugSymbols>
-    <OutputPath>..\..\..\..\Binaries\DotNET\AutomationScripts\</OutputPath>
-    <DefineConstants>DEBUG;TRACE</DefineConstants>
-    <TreatWarningsAsErrors>true</TreatWarningsAsErrors>
-    <DebugType>full</DebugType>
-    <PlatformTarget>x64</PlatformTarget>
-    <ErrorReport>prompt</ErrorReport>
-    <CodeAnalysisRuleSet>MinimumRecommendedRules.ruleset</CodeAnalysisRuleSet>
-  </PropertyGroup>
-  <PropertyGroup Condition="'$(Configuration)|$(Platform)' == 'Development|x64'">
-    <OutputPath>..\..\..\..\Binaries\DotNET\AutomationScripts\</OutputPath>
-    <DefineConstants>TRACE</DefineConstants>
-    <TreatWarningsAsErrors>true</TreatWarningsAsErrors>
-    <DebugType>pdbonly</DebugType>
-    <PlatformTarget>x64</PlatformTarget>
-    <ErrorReport>prompt</ErrorReport>
-    <CodeAnalysisRuleSet>MinimumRecommendedRules.ruleset</CodeAnalysisRuleSet>
-  </PropertyGroup>
   <ItemGroup>
     <Reference Include="System" />
     <Reference Include="System.Core" />
@@ -90,4 +71,4 @@
   </ItemGroup>
   <ItemGroup />
   <Import Project="$(MSBuildToolsPath)\Microsoft.CSharp.targets" />
-</Project>
\ No newline at end of file
+</Project>
diff --git a/Engine/Source/Programs/AutomationToolLauncher/AutomationToolLauncher.csproj b/Engine/Source/Programs/AutomationToolLauncher/AutomationToolLauncher.csproj
--- a/Engine/Source/Programs/AutomationToolLauncher/AutomationToolLauncher.csproj
+++ b/Engine/Source/Programs/AutomationToolLauncher/AutomationToolLauncher.csproj
@@ -2,7 +2,7 @@
 <Project ToolsVersion="12.0" DefaultTargets="Build" xmlns="http://schemas.microsoft.com/developer/msbuild/2003">
   <Import Project="$(MSBuildExtensionsPath)\$(MSBuildToolsVersion)\Microsoft.Common.props" Condition="Exists('$(MSBuildExtensionsPath)\$(MSBuildToolsVersion)\Microsoft.Common.props')" />
   <PropertyGroup>
-    <Configuration Condition=" '$(Configuration)' == '' ">Debug</Configuration>
+    <Configuration Condition=" '$(Configuration)' == '' ">Development</Configuration>
     <Platform Condition=" '$(Platform)' == '' ">AnyCPU</Platform>
     <ProjectGuid>{1D67A316-7EEC-4C07-83E3-BBD9C67AD129}</ProjectGuid>
     <OutputType>Exe</OutputType>
@@ -35,24 +35,6 @@
     <Prefer32Bit>false</Prefer32Bit>
     <AllowUnsafeBlocks>false</AllowUnsafeBlocks>
   </PropertyGroup>
-  <PropertyGroup Condition="'$(Configuration)|$(Platform)' == 'Debug|x64'">
-    <DebugSymbols>true</DebugSymbols>
-    <OutputPath>..\..\..\Binaries\DotNET\</OutputPath>
-    <DefineConstants>DEBUG;TRACE</DefineConstants>
-    <DebugType>full</DebugType>
-    <PlatformTarget>x64</PlatformTarget>
-    <ErrorReport>prompt</ErrorReport>
-    <CodeAnalysisRuleSet>MinimumRecommendedRules.ruleset</CodeAnalysisRuleSet>
-  </PropertyGroup>
-  <PropertyGroup Condition="'$(Configuration)|$(Platform)' == 'Development|x64'">
-    <OutputPath>..\..\..\Binaries\DotNET\</OutputPath>
-    <DefineConstants>TRACE</DefineConstants>
-    <Optimize>true</Optimize>
-    <DebugType>pdbonly</DebugType>
-    <PlatformTarget>x64</PlatformTarget>
-    <ErrorReport>prompt</ErrorReport>
-    <CodeAnalysisRuleSet>MinimumRecommendedRules.ruleset</CodeAnalysisRuleSet>
-  </PropertyGroup>
   <ItemGroup>
     <Reference Include="System" />
     <Reference Include="System.Core" />
@@ -83,4 +65,4 @@
   <Target Name="AfterBuild">
   </Target>
   -->
-</Project>
\ No newline at end of file
+</Project>
diff --git a/Engine/Source/Programs/DotNETCommon/DotNETUtilities/DotNETUtilities.csproj b/Engine/Source/Programs/DotNETCommon/DotNETUtilities/DotNETUtilities.csproj
--- a/Engine/Source/Programs/DotNETCommon/DotNETUtilities/DotNETUtilities.csproj
+++ b/Engine/Source/Programs/DotNETCommon/DotNETUtilities/DotNETUtilities.csproj
@@ -2,7 +2,7 @@
 <Project ToolsVersion="4.0" DefaultTargets="Build" xmlns="http://schemas.microsoft.com/developer/msbuild/2003">
   <Import Project="$(MSBuildExtensionsPath)\$(MSBuildToolsVersion)\Microsoft.Common.props" Condition="Exists('$(MSBuildExtensionsPath)\$(MSBuildToolsVersion)\Microsoft.Common.props')" />
   <PropertyGroup>
-    <Configuration Condition=" '$(Configuration)' == '' ">Debug</Configuration>
+    <Configuration Condition=" '$(Configuration)' == '' ">Development</Configuration>
     <Platform Condition=" '$(Platform)' == '' ">AnyCPU</Platform>
     <ProjectGuid>{5D7D66E8-8C76-4AF9-B3EC-2EF03421D730}</ProjectGuid>
     <OutputType>Library</OutputType>
@@ -36,38 +36,11 @@
     <RunCodeAnalysis>false</RunCodeAnalysis>
   </PropertyGroup>
   <PropertyGroup Condition="'$(Configuration)|$(Platform)' == 'Release|AnyCPU'">
-    <OutputPath>bin\Release\</OutputPath>
-    <DefineConstants>TRACE</DefineConstants>
-    <Optimize>true</Optimize>
-    <DebugType>pdbonly</DebugType>
-    <PlatformTarget>AnyCPU</PlatformTarget>
-    <ErrorReport>prompt</ErrorReport>
-    <CodeAnalysisRuleSet>MinimumRecommendedRules.ruleset</CodeAnalysisRuleSet>
-  </PropertyGroup>
-  <PropertyGroup Condition="'$(Configuration)|$(Platform)' == 'Debug|x64'">
-    <DebugSymbols>true</DebugSymbols>
     <OutputPath>..\..\..\..\Binaries\DotNET\</OutputPath>
-    <DefineConstants>DEBUG;TRACE</DefineConstants>
-    <DebugType>full</DebugType>
-    <PlatformTarget>x64</PlatformTarget>
-    <ErrorReport>prompt</ErrorReport>
-    <CodeAnalysisRuleSet>MinimumRecommendedRules.ruleset</CodeAnalysisRuleSet>
-  </PropertyGroup>
-  <PropertyGroup Condition="'$(Configuration)|$(Platform)' == 'Development|x64'">
-    <OutputPath>..\..\..\..\Binaries\DotNET\</OutputPath>
-    <DefineConstants>TRACE</DefineConstants>
-    <Optimize>true</Optimize>
-    <DebugType>pdbonly</DebugType>
-    <PlatformTarget>x64</PlatformTarget>
-    <ErrorReport>prompt</ErrorReport>
-    <CodeAnalysisRuleSet>MinimumRecommendedRules.ruleset</CodeAnalysisRuleSet>
-  </PropertyGroup>
-  <PropertyGroup Condition="'$(Configuration)|$(Platform)' == 'Release|x64'">
-    <OutputPath>bin\x64\Release\</OutputPath>
     <DefineConstants>TRACE</DefineConstants>
     <Optimize>true</Optimize>
     <DebugType>pdbonly</DebugType>
-    <PlatformTarget>x64</PlatformTarget>
+    <PlatformTarget>AnyCPU</PlatformTarget>
     <ErrorReport>prompt</ErrorReport>
     <CodeAnalysisRuleSet>MinimumRecommendedRules.ruleset</CodeAnalysisRuleSet>
   </PropertyGroup>
@@ -102,4 +75,4 @@
   <Target Name="AfterBuild">
   </Target>
   -->
-</Project>
\ No newline at end of file
+</Project>
diff --git a/Engine/Source/Programs/GitDependencies/GitDependencies.csproj b/Engine/Source/Programs/GitDependencies/GitDependencies.csproj
--- a/Engine/Source/Programs/GitDependencies/GitDependencies.csproj
+++ b/Engine/Source/Programs/GitDependencies/GitDependencies.csproj
@@ -2,7 +2,7 @@
 <Project ToolsVersion="12.0" DefaultTargets="Build" xmlns="http://schemas.microsoft.com/developer/msbuild/2003">
   <Import Project="$(MSBuildExtensionsPath)\$(MSBuildToolsVersion)\Microsoft.Common.props" Condition="Exists('$(MSBuildExtensionsPath)\$(MSBuildToolsVersion)\Microsoft.Common.props')" />
   <PropertyGroup>
-    <Configuration Condition=" '$(Configuration)' == '' ">Debug</Configuration>
+    <Configuration Condition=" '$(Configuration)' == '' ">Development</Configuration>
     <Platform Condition=" '$(Platform)' == '' ">AnyCPU</Platform>
     <ProjectGuid>{9B740143-841C-4AF7-905E-D8883F1735DE}</ProjectGuid>
     <OutputType>Exe</OutputType>
@@ -32,25 +32,6 @@
     <ErrorReport>prompt</ErrorReport>
     <WarningLevel>4</WarningLevel>
   </PropertyGroup>
-  <PropertyGroup Condition="'$(Configuration)|$(Platform)' == 'Debug|x64'">
-    <DebugSymbols>true</DebugSymbols>
-    <OutputPath>bin\x64\Debug\</OutputPath>
-    <DefineConstants>DEBUG;TRACE</DefineConstants>
-    <DebugType>full</DebugType>
-    <PlatformTarget>x64</PlatformTarget>
-    <ErrorReport>prompt</ErrorReport>
-    <CodeAnalysisRuleSet>MinimumRecommendedRules.ruleset</CodeAnalysisRuleSet>
-    <Prefer32Bit>true</Prefer32Bit>
-  </PropertyGroup>
-  <PropertyGroup Condition="'$(Configuration)|$(Platform)' == 'Development|x64'">
-    <OutputPath>bin\x64\Development\</OutputPath>
-    <DefineConstants>TRACE</DefineConstants>
-    <DebugType>pdbonly</DebugType>
-    <PlatformTarget>x64</PlatformTarget>
-    <ErrorReport>prompt</ErrorReport>
-    <CodeAnalysisRuleSet>MinimumRecommendedRules.ruleset</CodeAnalysisRuleSet>
-    <Prefer32Bit>true</Prefer32Bit>
-  </PropertyGroup>
   <ItemGroup>
     <Reference Include="System" />
     <Reference Include="System.Core" />
@@ -84,4 +65,4 @@
   <Target Name="AfterBuild">
   </Target>
   -->
-</Project>
\ No newline at end of file
+</Project>
diff --git a/Engine/Source/Programs/HTML5/HTML5LaunchHelper/HTML5LaunchHelper.csproj b/Engine/Source/Programs/HTML5/HTML5LaunchHelper/HTML5LaunchHelper.csproj
--- a/Engine/Source/Programs/HTML5/HTML5LaunchHelper/HTML5LaunchHelper.csproj
+++ b/Engine/Source/Programs/HTML5/HTML5LaunchHelper/HTML5LaunchHelper.csproj
@@ -2,7 +2,7 @@
 <Project ToolsVersion="12.0" DefaultTargets="Build" xmlns="http://schemas.microsoft.com/developer/msbuild/2003">
   <Import Project="$(MSBuildExtensionsPath)\$(MSBuildToolsVersion)\Microsoft.Common.props" Condition="Exists('$(MSBuildExtensionsPath)\$(MSBuildToolsVersion)\Microsoft.Common.props')" />
   <PropertyGroup>
-    <Configuration Condition=" '$(Configuration)' == '' ">Debug</Configuration>
+    <Configuration Condition=" '$(Configuration)' == '' ">Development</Configuration>
     <Platform Condition=" '$(Platform)' == '' ">AnyCPU</Platform>
     <ProjectGuid>{C22FE09B-9D62-4B4C-AE2E-BDB5D66C143F}</ProjectGuid>
     <OutputType>Exe</OutputType>
@@ -31,26 +31,6 @@
     <ErrorReport>prompt</ErrorReport>
     <WarningLevel>4</WarningLevel>
   </PropertyGroup>
-  <PropertyGroup Condition="'$(Configuration)|$(Platform)' == 'Debug|x64'">
-    <DebugSymbols>true</DebugSymbols>
-    <OutputPath>..\..\..\..\Binaries\DotNET\</OutputPath>
-    <DefineConstants>DEBUG;TRACE</DefineConstants>
-    <DebugType>full</DebugType>
-    <PlatformTarget>x64</PlatformTarget>
-    <ErrorReport>prompt</ErrorReport>
-    <CodeAnalysisRuleSet>MinimumRecommendedRules.ruleset</CodeAnalysisRuleSet>
-    <Prefer32Bit>true</Prefer32Bit>
-  </PropertyGroup>
-  <PropertyGroup Condition="'$(Configuration)|$(Platform)' == 'Development|x64'">
-    <OutputPath>..\..\..\..\Binaries\DotNET\</OutputPath>
-    <DefineConstants>TRACE</DefineConstants>
-    <Optimize>true</Optimize>
-    <DebugType>pdbonly</DebugType>
-    <PlatformTarget>x64</PlatformTarget>
-    <ErrorReport>prompt</ErrorReport>
-    <CodeAnalysisRuleSet>MinimumRecommendedRules.ruleset</CodeAnalysisRuleSet>
-    <Prefer32Bit>true</Prefer32Bit>
-  </PropertyGroup>
   <ItemGroup>
     <Reference Include="System" />
     <Reference Include="System.Core" />
@@ -75,4 +55,4 @@
   <Target Name="AfterBuild">
   </Target>
   -->
-</Project>
\ No newline at end of file
+</Project>
diff --git a/Engine/Source/Programs/IOS/DeploymentInterface/DeploymentInterface.csproj b/Engine/Source/Programs/IOS/DeploymentInterface/DeploymentInterface.csproj
--- a/Engine/Source/Programs/IOS/DeploymentInterface/DeploymentInterface.csproj
+++ b/Engine/Source/Programs/IOS/DeploymentInterface/DeploymentInterface.csproj
@@ -1,7 +1,7 @@
 ﻿<?xml version="1.0" encoding="utf-8"?>
 <Project ToolsVersion="4.0" DefaultTargets="Build" xmlns="http://schemas.microsoft.com/developer/msbuild/2003">
   <PropertyGroup>
-    <Configuration Condition=" '$(Configuration)' == '' ">Debug</Configuration>
+    <Configuration Condition=" '$(Configuration)' == '' ">Development</Configuration>
     <Platform Condition=" '$(Platform)' == '' ">AnyCPU</Platform>
     <ProductVersion>8.0.30703</ProductVersion>
     <SchemaVersion>2.0</SchemaVersion>
@@ -30,7 +30,7 @@
     <DefineConstants>DEBUG;TRACE</DefineConstants>
     <ErrorReport>prompt</ErrorReport>
     <WarningLevel>4</WarningLevel>
-    <PlatformTarget>x64</PlatformTarget>
+    <PlatformTarget>AnyCPU</PlatformTarget>
   </PropertyGroup>
   <PropertyGroup Condition=" '$(Configuration)|$(Platform)' == 'Development|AnyCPU' ">
     <DebugType>pdbonly</DebugType>
@@ -39,25 +39,7 @@
     <DefineConstants>TRACE</DefineConstants>
     <ErrorReport>prompt</ErrorReport>
     <WarningLevel>4</WarningLevel>
-    <PlatformTarget>x64</PlatformTarget>
-  </PropertyGroup>
-  <PropertyGroup Condition="'$(Configuration)|$(Platform)' == 'Debug|x64'">
-    <DebugSymbols>true</DebugSymbols>
-    <OutputPath>..\..\..\..\Binaries\DotNET\IOS\</OutputPath>
-    <DefineConstants>DEBUG;TRACE</DefineConstants>
-    <DebugType>full</DebugType>
-    <PlatformTarget>x64</PlatformTarget>
-    <ErrorReport>prompt</ErrorReport>
-    <CodeAnalysisRuleSet>MinimumRecommendedRules.ruleset</CodeAnalysisRuleSet>
-  </PropertyGroup>
-  <PropertyGroup Condition="'$(Configuration)|$(Platform)' == 'Development|x64'">
-    <OutputPath>..\..\..\..\Binaries\DotNET\IOS\</OutputPath>
-    <DefineConstants>TRACE</DefineConstants>
-    <Optimize>true</Optimize>
-    <DebugType>pdbonly</DebugType>
-    <PlatformTarget>x64</PlatformTarget>
-    <ErrorReport>prompt</ErrorReport>
-    <CodeAnalysisRuleSet>MinimumRecommendedRules.ruleset</CodeAnalysisRuleSet>
+    <PlatformTarget>AnyCPU</PlatformTarget>
   </PropertyGroup>
   <ItemGroup>
     <Reference Include="System" />
@@ -87,4 +69,4 @@
   <Target Name="AfterBuild">
   </Target>
   -->
-</Project>
\ No newline at end of file
+</Project>
diff --git a/Engine/Source/Programs/IOS/DeploymentServer/DeploymentServer.csproj b/Engine/Source/Programs/IOS/DeploymentServer/DeploymentServer.csproj
--- a/Engine/Source/Programs/IOS/DeploymentServer/DeploymentServer.csproj
+++ b/Engine/Source/Programs/IOS/DeploymentServer/DeploymentServer.csproj
@@ -23,31 +23,12 @@
     <SccProvider>
     </SccProvider>
   </PropertyGroup>
-  <PropertyGroup Condition=" '$(Configuration)|$(Platform)' == 'Debug|x86' ">
-    <PlatformTarget>x64</PlatformTarget>
-    <DebugSymbols>true</DebugSymbols>
-    <DebugType>full</DebugType>
-    <Optimize>false</Optimize>
-    <OutputPath>..\..\..\..\Binaries\DotNET\</OutputPath>
-    <DefineConstants>DEBUG;TRACE</DefineConstants>
-    <ErrorReport>prompt</ErrorReport>
-    <WarningLevel>4</WarningLevel>
-  </PropertyGroup>
-  <PropertyGroup Condition=" '$(Configuration)|$(Platform)' == 'Development|x86' ">
-    <PlatformTarget>x86</PlatformTarget>
-    <DebugType>pdbonly</DebugType>
-    <Optimize>true</Optimize>
-    <OutputPath>..\..\..\..\Binaries\DotNET\</OutputPath>
-    <DefineConstants>TRACE</DefineConstants>
-    <ErrorReport>prompt</ErrorReport>
-    <WarningLevel>4</WarningLevel>
-  </PropertyGroup>
   <PropertyGroup Condition="'$(Configuration)|$(Platform)' == 'Debug|AnyCPU'">
     <DebugSymbols>true</DebugSymbols>
     <OutputPath>..\..\..\..\Binaries\DotNET\IOS\</OutputPath>
     <DefineConstants>DEBUG;TRACE</DefineConstants>
     <DebugType>full</DebugType>
-    <PlatformTarget>x64</PlatformTarget>
+    <PlatformTarget>AnyCPU</PlatformTarget>
     <ErrorReport>prompt</ErrorReport>
     <CodeAnalysisIgnoreBuiltInRuleSets>false</CodeAnalysisIgnoreBuiltInRuleSets>
     <CodeAnalysisIgnoreBuiltInRules>false</CodeAnalysisIgnoreBuiltInRules>
@@ -57,7 +38,7 @@
     <DefineConstants>TRACE</DefineConstants>
     <Optimize>true</Optimize>
     <DebugType>pdbonly</DebugType>
-    <PlatformTarget>x64</PlatformTarget>
+    <PlatformTarget>AnyCPU</PlatformTarget>
     <ErrorReport>prompt</ErrorReport>
     <CodeAnalysisIgnoreBuiltInRuleSets>false</CodeAnalysisIgnoreBuiltInRuleSets>
     <CodeAnalysisIgnoreBuiltInRules>false</CodeAnalysisIgnoreBuiltInRules>
@@ -126,4 +107,4 @@
   <Target Name="AfterBuild">
   </Target>
   -->
-</Project>
\ No newline at end of file
+</Project>
diff --git a/Engine/Source/Programs/IOS/MobileDeviceInterface/MobileDeviceInterface.csproj b/Engine/Source/Programs/IOS/MobileDeviceInterface/MobileDeviceInterface.csproj
--- a/Engine/Source/Programs/IOS/MobileDeviceInterface/MobileDeviceInterface.csproj
+++ b/Engine/Source/Programs/IOS/MobileDeviceInterface/MobileDeviceInterface.csproj
@@ -1,7 +1,7 @@
 ﻿<?xml version="1.0" encoding="utf-8"?>
 <Project ToolsVersion="4.0" DefaultTargets="Build" xmlns="http://schemas.microsoft.com/developer/msbuild/2003">
   <PropertyGroup>
-    <Configuration Condition=" '$(Configuration)' == '' ">Debug</Configuration>
+    <Configuration Condition=" '$(Configuration)' == '' ">Development</Configuration>
     <Platform Condition=" '$(Platform)' == '' ">AnyCPU</Platform>
     <ProductVersion>9.0.30729</ProductVersion>
     <SchemaVersion>2.0</SchemaVersion>
@@ -50,7 +50,7 @@
     <DefineConstants>DEBUG;TRACE</DefineConstants>
     <ErrorReport>prompt</ErrorReport>
     <WarningLevel>4</WarningLevel>
-    <PlatformTarget>x86</PlatformTarget>
+    <PlatformTarget>AnyCPU</PlatformTarget>
     <AllowUnsafeBlocks>false</AllowUnsafeBlocks>
     <TreatWarningsAsErrors>true</TreatWarningsAsErrors>
   </PropertyGroup>
@@ -61,30 +61,10 @@
     <DefineConstants>TRACE</DefineConstants>
     <ErrorReport>prompt</ErrorReport>
     <WarningLevel>4</WarningLevel>
-    <PlatformTarget>x64</PlatformTarget>
+    <PlatformTarget>AnyCPU</PlatformTarget>
     <AllowUnsafeBlocks>false</AllowUnsafeBlocks>
     <TreatWarningsAsErrors>true</TreatWarningsAsErrors>
   </PropertyGroup>
-  <PropertyGroup Condition="'$(Configuration)|$(Platform)' == 'Debug|x64'">
-    <DebugSymbols>true</DebugSymbols>
-    <OutputPath>..\..\..\..\Binaries\DotNET\IOS\</OutputPath>
-    <DefineConstants>DEBUG;TRACE</DefineConstants>
-    <TreatWarningsAsErrors>true</TreatWarningsAsErrors>
-    <DebugType>full</DebugType>
-    <PlatformTarget>x64</PlatformTarget>
-    <ErrorReport>prompt</ErrorReport>
-    <CodeAnalysisRuleSet>MinimumRecommendedRules.ruleset</CodeAnalysisRuleSet>
-  </PropertyGroup>
-  <PropertyGroup Condition="'$(Configuration)|$(Platform)' == 'Development|x64'">
-    <OutputPath>..\..\..\..\Binaries\DotNET\IOS\</OutputPath>
-    <DefineConstants>TRACE</DefineConstants>
-    <Optimize>true</Optimize>
-    <TreatWarningsAsErrors>true</TreatWarningsAsErrors>
-    <DebugType>pdbonly</DebugType>
-    <PlatformTarget>x64</PlatformTarget>
-    <ErrorReport>prompt</ErrorReport>
-    <CodeAnalysisRuleSet>MinimumRecommendedRules.ruleset</CodeAnalysisRuleSet>
-  </PropertyGroup>
   <ItemGroup>
     <Reference Include="System" />
     <Reference Include="System.Core">
@@ -141,4 +121,4 @@
   <Target Name="AfterBuild">
   </Target>
   -->
-</Project>
\ No newline at end of file
+</Project>
diff --git a/Engine/Source/Programs/IOS/iPhonePackager/iPhonePackager.csproj b/Engine/Source/Programs/IOS/iPhonePackager/iPhonePackager.csproj
--- a/Engine/Source/Programs/IOS/iPhonePackager/iPhonePackager.csproj
+++ b/Engine/Source/Programs/IOS/iPhonePackager/iPhonePackager.csproj
@@ -1,7 +1,7 @@
 ﻿<?xml version="1.0" encoding="utf-8"?>
 <Project ToolsVersion="4.0" DefaultTargets="Build" xmlns="http://schemas.microsoft.com/developer/msbuild/2003">
   <PropertyGroup>
-    <Configuration Condition=" '$(Configuration)' == '' ">Debug</Configuration>
+    <Configuration Condition=" '$(Configuration)' == '' ">Development</Configuration>
     <Platform Condition=" '$(Platform)' == '' ">AnyCPU</Platform>
     <ProductVersion>9.0.30729</ProductVersion>
     <SchemaVersion>2.0</SchemaVersion>
@@ -53,7 +53,7 @@
     <DefineConstants>DEBUG;TRACE</DefineConstants>
     <ErrorReport>prompt</ErrorReport>
     <WarningLevel>4</WarningLevel>
-    <PlatformTarget>x64</PlatformTarget>
+    <PlatformTarget>AnyCPU</PlatformTarget>
   </PropertyGroup>
   <PropertyGroup Condition=" '$(Configuration)|$(Platform)' == 'Development|AnyCPU' ">
     <DebugType>pdbonly</DebugType>
@@ -62,28 +62,9 @@
     <DefineConstants>TRACE</DefineConstants>
     <ErrorReport>prompt</ErrorReport>
     <WarningLevel>4</WarningLevel>
-    <PlatformTarget>x64</PlatformTarget>
+    <PlatformTarget>AnyCPU</PlatformTarget>
     <Prefer32Bit>false</Prefer32Bit>
   </PropertyGroup>
-  <PropertyGroup Condition="'$(Configuration)|$(Platform)' == 'Debug|x64'">
-    <DebugSymbols>true</DebugSymbols>
-    <OutputPath>..\..\..\..\Binaries\DotNET\IOS\</OutputPath>
-    <DefineConstants>DEBUG;TRACE</DefineConstants>
-    <DebugType>full</DebugType>
-    <PlatformTarget>x64</PlatformTarget>
-    <ErrorReport>prompt</ErrorReport>
-    <CodeAnalysisRuleSet>MinimumRecommendedRules.ruleset</CodeAnalysisRuleSet>
-    <Prefer32Bit>true</Prefer32Bit>
-  </PropertyGroup>
-  <PropertyGroup Condition="'$(Configuration)|$(Platform)' == 'Development|x64'">
-    <OutputPath>..\..\..\..\Binaries\DotNET\IOS\</OutputPath>
-    <DefineConstants>TRACE</DefineConstants>
-    <Optimize>true</Optimize>
-    <DebugType>pdbonly</DebugType>
-    <PlatformTarget>x64</PlatformTarget>
-    <ErrorReport>prompt</ErrorReport>
-    <CodeAnalysisRuleSet>MinimumRecommendedRules.ruleset</CodeAnalysisRuleSet>
-  </PropertyGroup>
   <ItemGroup>
     <Reference Include="BouncyCastle.Crypto, Version=1.6.1.0, Culture=neutral, PublicKeyToken=0e99375e54769942">
       <SpecificVersion>False</SpecificVersion>
@@ -258,4 +239,4 @@
   <Target Name="AfterBuild">
   </Target>
   -->
-</Project>
\ No newline at end of file
+</Project>
diff --git a/Engine/Source/Programs/UnrealBuildTool/System/VCProjectFileGenerator.cs b/Engine/Source/Programs/UnrealBuildTool/System/VCProjectFileGenerator.cs
--- a/Engine/Source/Programs/UnrealBuildTool/System/VCProjectFileGenerator.cs
+++ b/Engine/Source/Programs/UnrealBuildTool/System/VCProjectFileGenerator.cs
@@ -185,7 +185,7 @@ namespace UnrealBuildTool
 		static public readonly string DefaultPlatformName = "Win32";
 
 		/// The platform name that must be used for .NET projects
-		static public readonly string DotNetPlatformName = "x64";
+		static public readonly string DotNetPlatformName = "AnyCPU";
 
 
 		/// <summary>
diff --git a/Engine/Source/Programs/UnrealBuildTool/UnrealBuildTool.csproj b/Engine/Source/Programs/UnrealBuildTool/UnrealBuildTool.csproj
--- a/Engine/Source/Programs/UnrealBuildTool/UnrealBuildTool.csproj
+++ b/Engine/Source/Programs/UnrealBuildTool/UnrealBuildTool.csproj
@@ -1,7 +1,7 @@
 ﻿<?xml version="1.0" encoding="utf-8"?>
 <Project DefaultTargets="Build" xmlns="http://schemas.microsoft.com/developer/msbuild/2003" ToolsVersion="4.0">
   <PropertyGroup>
-    <Configuration Condition=" '$(Configuration)' == '' ">Debug</Configuration>
+    <Configuration Condition=" '$(Configuration)' == '' ">Development</Configuration>
     <Platform Condition=" '$(Platform)' == '' ">AnyCPU</Platform>
     <ProductVersion>9.0.30729</ProductVersion>
     <SchemaVersion>2.0</SchemaVersion>
@@ -46,17 +46,6 @@
     <AllowUnsafeBlocks>false</AllowUnsafeBlocks>
     <Prefer32Bit>false</Prefer32Bit>
   </PropertyGroup>
-  <PropertyGroup Condition=" '$(Configuration)|$(Platform)' == 'Development|x86' ">
-    <DebugType>pdbonly</DebugType>
-    <Optimize>true</Optimize>
-    <OutputPath>..\..\..\Binaries\DotNET\</OutputPath>
-    <DefineConstants>
-    </DefineConstants>
-    <Optimize>true</Optimize>
-    <DebugType>pdbonly</DebugType>
-    <PlatformTarget>x86</PlatformTarget>
-    <ErrorReport>prompt</ErrorReport>
-  </PropertyGroup>
   <PropertyGroup Condition=" '$(Configuration)|$(Platform)' == 'Debug|AnyCPU' ">
     <OutputPath>..\..\..\Binaries\DotNET\</OutputPath>
     <DefineConstants>DEBUG;TRACE</DefineConstants>
@@ -66,27 +55,6 @@
     <ErrorReport>prompt</ErrorReport>
     <Prefer32Bit>false</Prefer32Bit>
   </PropertyGroup>
-  <PropertyGroup Condition="'$(Configuration)|$(Platform)' == 'Debug|x86'">
-    <DefineConstants>TRACE;DEBUG</DefineConstants>
-    <Optimize>true</Optimize>
-    <OutputPath>bin\x86\Debug\</OutputPath>
-  </PropertyGroup>
-  <PropertyGroup Condition="'$(Configuration)|$(Platform)' == 'Development|x64'">
-    <OutputPath>..\..\..\Binaries\DotNET\</OutputPath>
-    <DefineConstants>TRACE</DefineConstants>
-    <DebugType>pdbonly</DebugType>
-    <PlatformTarget>x64</PlatformTarget>
-    <ErrorReport>prompt</ErrorReport>
-    <CodeAnalysisRuleSet>MinimumRecommendedRules.ruleset</CodeAnalysisRuleSet>
-  </PropertyGroup>
-  <PropertyGroup Condition="'$(Configuration)|$(Platform)' == 'Debug|x64'">
-    <OutputPath>..\..\..\Binaries\DotNET\</OutputPath>
-    <DefineConstants>DEBUG;TRACE</DefineConstants>
-    <DebugType>pdbonly</DebugType>
-    <PlatformTarget>x64</PlatformTarget>
-    <ErrorReport>prompt</ErrorReport>
-    <CodeAnalysisRuleSet>MinimumRecommendedRules.ruleset</CodeAnalysisRuleSet>
-  </PropertyGroup>
   <ItemGroup>
     <Reference Include="Ionic.Zip.Reduced">
       <HintPath>..\..\..\Binaries\DotNET\Ionic.Zip.Reduced.dll</HintPath>
@@ -295,12 +263,6 @@
   <ItemGroup>
     <None Include="app.config" />
   </ItemGroup>
-  <ItemGroup>
-    <ProjectReference Include="..\DotNETCommon\DotNETUtilities\DotNETUtilities.csproj">
-      <Project>{5d7d66e8-8c76-4af9-b3ec-2ef03421d730}</Project>
-      <Name>DotNETUtilities</Name>
-    </ProjectReference>
-  </ItemGroup>
   <Import Project="$(MSBuildBinPath)\Microsoft.CSharp.targets" />
   <!-- To modify your build process, add your task inside one of the targets below and uncomment it. 
        Other similar extension points exist, see Microsoft.Common.targets.
@@ -315,5 +277,10 @@
     <PreBuildEvent>
     </PreBuildEvent>
   </PropertyGroup>
-  <ItemGroup />
+  <ItemGroup>
+    <ProjectReference Include="..\DotNETCommon\DotNETUtilities\DotNETUtilities.csproj">
+      <Project>{5d7d66e8-8c76-4af9-b3ec-2ef03421d730}</Project>
+      <Name>DotNETUtilities</Name>
+    </ProjectReference>
+  </ItemGroup>
 </Project>
\ No newline at end of file
