Fixed issue with Installed Build
Added MobileDeviceInterface.dll to UAT Build products as it's now a direct dependency

diff --git a/Engine/Source/Programs/AutomationTool/AutomationUtils/UE4Build.cs b/Engine/Source/Programs/AutomationTool/AutomationUtils/UE4Build.cs
--- a/Engine/Source/Programs/AutomationTool/AutomationUtils/UE4Build.cs
+++ b/Engine/Source/Programs/AutomationTool/AutomationUtils/UE4Build.cs
@@ -1644,6 +1644,7 @@ namespace AutomationTool
 						"UnrealBuildTool.exe.config",
 						"AutomationUtils.Automation.dll",
 						"DotNETUtilities.dll",
+						"MobileDeviceInterface.dll"
 					});
 
 			foreach (var UATFile in UATFiles)
