Integrating codesign fix into //UE4/Main

diff --git a/Engine/Source/Programs/AutomationTool/AutomationUtils/CommandUtils.cs b/Engine/Source/Programs/AutomationTool/AutomationUtils/CommandUtils.cs
--- a/Engine/Source/Programs/AutomationTool/AutomationUtils/CommandUtils.cs
+++ b/Engine/Source/Programs/AutomationTool/AutomationUtils/CommandUtils.cs
@@ -2709,7 +2709,14 @@ namespace AutomationTool
 				return;
 			}
 
-			string SignToolName = "/usr/bin/codesign";
+			// Use the old codesigning tool after the upgrade due to segmentation fault on Sierra
+			string SignToolName = "/usr/local/bin/codesign_old";
+			
+			// unless it doesn't exist, then use the Sierra one.
+			if(!File.Exists(SignToolName))
+			{
+				SignToolName = "/usr/bin/codesign";
+			}
 
 			string CodeSignArgs = String.Format("-f --deep -s \"{0}\" -v \"{1}\" --no-strict", "Developer ID Application", InPath);
 
