Rollback //UE4/Main/Engine/Source/Programs/AutomationTool/Scripts/BuildProjectCommand.Automation.cs to revision 12
This fixes the failure to build non-code projects for mobile, PS4, and XB1

diff --git a/Engine/Source/Programs/AutomationTool/Scripts/BuildProjectCommand.Automation.cs b/Engine/Source/Programs/AutomationTool/Scripts/BuildProjectCommand.Automation.cs
--- a/Engine/Source/Programs/AutomationTool/Scripts/BuildProjectCommand.Automation.cs
+++ b/Engine/Source/Programs/AutomationTool/Scripts/BuildProjectCommand.Automation.cs
@@ -108,40 +108,29 @@ public partial class Project : CommandUtils
 		// Setup cooked targets
 		if (Params.HasClientCookedTargets && (TargetMask & ProjectBuildTargets.ClientCooked) == ProjectBuildTargets.ClientCooked)
 		{
-			List<UnrealTargetPlatform> UniquePlatformTypes = Params.ClientTargetPlatforms.ConvertAll(x => x.Type).Distinct().ToList();
+            List<UnrealTargetPlatform> UniquePlatformTypes = Params.ClientTargetPlatforms.ConvertAll(x => x.Type).Distinct().ToList();
 
-			foreach (var BuildConfig in Params.ClientConfigsToBuild)
+            foreach (var BuildConfig in Params.ClientConfigsToBuild)
 			{
-				foreach (var ClientPlatformType in UniquePlatformTypes)
-				{					
-					string AdditionalArgs = GetBlueprintPluginPathArgument(Params, true, ClientPlatformType);
-					AdditionalArgs += " -remoteini=\"" + Params.RawProjectPath.Directory.FullName + "\"";
-					if (Params.IsCodeBasedProject == false)
-					{
-						AdditionalArgs += " -project=\"" + Path.GetFullPath(Params.RawProjectPath.FullName) + "\"";
-					}
-					CrashReportPlatforms.Add(ClientPlatformType);
-					Agenda.AddTargets(Params.ClientCookedTargets.ToArray(), ClientPlatformType, BuildConfig, Params.CodeBasedUprojectPath, InAddArgs: AdditionalArgs);
+                foreach (var ClientPlatformType in UniquePlatformTypes)
+				{
+                    string ScriptPluginArgs = GetBlueprintPluginPathArgument(Params, true, ClientPlatformType);
+                    CrashReportPlatforms.Add(ClientPlatformType);
+					Agenda.AddTargets(Params.ClientCookedTargets.ToArray(), ClientPlatformType, BuildConfig, Params.CodeBasedUprojectPath, InAddArgs: ScriptPluginArgs + " -remoteini=\"" + Params.RawProjectPath.Directory.FullName + "\"");
 				}
 			}
 		}
 		if (Params.HasServerCookedTargets && (TargetMask & ProjectBuildTargets.ServerCooked) == ProjectBuildTargets.ServerCooked)
 		{
-			List<UnrealTargetPlatform> UniquePlatformTypes = Params.ServerTargetPlatforms.ConvertAll(x => x.Type).Distinct().ToList();
+            List<UnrealTargetPlatform> UniquePlatformTypes = Params.ServerTargetPlatforms.ConvertAll(x => x.Type).Distinct().ToList();
 
-			foreach (var BuildConfig in Params.ServerConfigsToBuild)
+            foreach (var BuildConfig in Params.ServerConfigsToBuild)
 			{
 				foreach (var ServerPlatformType in UniquePlatformTypes)
 				{
-					string AdditionalArgs = GetBlueprintPluginPathArgument(Params, false, ServerPlatformType);
-					AdditionalArgs += " -remoteini=\"" + Params.RawProjectPath.Directory.FullName + "\"";
-					if (Params.IsCodeBasedProject == false)
-					{
-						AdditionalArgs += " -project=\"" + Path.GetFullPath(Params.RawProjectPath.FullName) + "\"";
-					}
-
-					CrashReportPlatforms.Add(ServerPlatformType);
-					Agenda.AddTargets(Params.ServerCookedTargets.ToArray(), ServerPlatformType, BuildConfig, Params.CodeBasedUprojectPath, InAddArgs: AdditionalArgs);
+                    string ScriptPluginArgs = GetBlueprintPluginPathArgument(Params, false, ServerPlatformType);
+                    CrashReportPlatforms.Add(ServerPlatformType);
+					Agenda.AddTargets(Params.ServerCookedTargets.ToArray(), ServerPlatformType, BuildConfig, Params.CodeBasedUprojectPath, InAddArgs: ScriptPluginArgs + " -remoteini=\"" + Params.RawProjectPath.Directory.FullName + "\"");
 				}
 			}
 		}
