Fixing mangled English translations

diff --git a/Engine/Source/Programs/AutomationTool/OneSkyLocalization/OneSkyLocalizationProvider.cs b/Engine/Source/Programs/AutomationTool/OneSkyLocalization/OneSkyLocalizationProvider.cs
--- a/Engine/Source/Programs/AutomationTool/OneSkyLocalization/OneSkyLocalizationProvider.cs
+++ b/Engine/Source/Programs/AutomationTool/OneSkyLocalization/OneSkyLocalizationProvider.cs
@@ -65,6 +65,12 @@ public class OneSkyLocalizationProvider : LocalizationProvider
 			var CulturesToExport = new List<string>();
 			foreach (var OneSkyCulture in OneSkyProject.EnabledCultures)
 			{
+				// Skip the native culture, as OneSky has mangled it
+				if (OneSkyCulture == ProjectImportInfo.NativeCulture)
+				{
+					continue;
+				}
+
 				// Only export the OneSky cultures that we care about for this project
 				if (ProjectImportInfo.CulturesToGenerate.Contains(OneSkyCulture))
 				{
