BuildGraph: Ensure that -submit or -nosubmit can be overridden by arguments specified in a build graph <Command> element.

Needed to fix WEX Packaging failure caused by recent integrations to main.

diff --git a/Engine/Source/Programs/AutomationTool/BuildGraph/Tasks/CommandTask.cs b/Engine/Source/Programs/AutomationTool/BuildGraph/Tasks/CommandTask.cs
--- a/Engine/Source/Programs/AutomationTool/BuildGraph/Tasks/CommandTask.cs
+++ b/Engine/Source/Programs/AutomationTool/BuildGraph/Tasks/CommandTask.cs
@@ -10,6 +10,14 @@ using System.IO;
 
 namespace BuildGraph.Tasks
 {
+	static class StringExtensions
+	{
+		public static bool CaseInsensitiveContains(this string Text, string Value)
+		{
+			return System.Globalization.CultureInfo.InvariantCulture.CompareInfo.IndexOf(Text, Value, System.Globalization.CompareOptions.IgnoreCase) >= 0;
+		}
+	}
+
 	/// <summary>
 	/// Parameters for a task which calls another UAT command
 	/// </summary>
@@ -73,8 +81,14 @@ namespace BuildGraph.Tasks
 
 			// Run the command
 			StringBuilder CommandLine = new StringBuilder();
-			CommandLine.AppendFormat("{0} ", CommandUtils.P4Enabled ? "-p4" : "-nop4");
-			CommandLine.AppendFormat("{0} ", CommandUtils.AllowSubmit ? "-submit" : "-nosubmit");
+			if (!Parameters.Arguments.CaseInsensitiveContains("-p4") && !Parameters.Arguments.CaseInsensitiveContains("-nop4"))
+			{
+				CommandLine.AppendFormat("{0} ", CommandUtils.P4Enabled ? "-p4" : "-nop4");
+			}
+			if (!Parameters.Arguments.CaseInsensitiveContains("-submit") && !Parameters.Arguments.CaseInsensitiveContains("-nosubmit"))
+			{
+				CommandLine.AppendFormat("{0} ", CommandUtils.AllowSubmit ? "-submit" : "-nosubmit");
+			}
 			CommandLine.Append(Parameters.Name);
 			if (!String.IsNullOrEmpty(Parameters.Arguments))
 			{
