Log the exception thrown when a localization provider fails to initialize

This is failing on CIS, but not locally. Need the exception data to work out why.

diff --git a/Engine/Source/Programs/AutomationTool/Localization/LocalizationProvider.cs b/Engine/Source/Programs/AutomationTool/Localization/LocalizationProvider.cs
--- a/Engine/Source/Programs/AutomationTool/Localization/LocalizationProvider.cs
+++ b/Engine/Source/Programs/AutomationTool/Localization/LocalizationProvider.cs
@@ -143,9 +143,9 @@ public abstract class LocalizationProvider
 			{
 				return Activator.CreateInstance(LocalizationNodeType, new object[] { InLocalizationProviderArgs }) as LocalizationProvider;
 			}
-			catch
+			catch (Exception e)
 			{
-				BuildCommand.LogWarning("Unable to create an instance of the type '{0}'", LocalizationNodeType.FullName);
+				BuildCommand.LogWarning("Unable to create an instance of the type '{0}'. {1}", LocalizationNodeType.FullName, e.ToString());
 			}
 		}
 		else
