Fix for the localization automation submitting unchanged files and causing P4 spam

diff --git a/Engine/Source/Programs/AutomationTool/Scripts/Localisation.Automation.cs b/Engine/Source/Programs/AutomationTool/Scripts/Localisation.Automation.cs
--- a/Engine/Source/Programs/AutomationTool/Scripts/Localisation.Automation.cs
+++ b/Engine/Source/Programs/AutomationTool/Scripts/Localisation.Automation.cs
@@ -203,6 +203,9 @@ class Localise : BuildCommand
 				}
 			}
 
+			// Revert any other unchanged files
+			P4.RevertUnchanged(PendingChangeList);
+
 			int SubmittedChangeList;
 			P4.Submit(PendingChangeList, out SubmittedChangeList);
 		}
