Remove the -forcelogflush parameter from invocations of the editor from UAT. Recent changes to the async logging system have made this parameter significantly slower (to the point where it's causing a massive increase in cook times). The log now asynchronously flushes at frequent intervals anyway, so the original need for it is less pressing.

diff --git a/Engine/Source/Programs/AutomationTool/AutomationUtils/CommandletUtils.cs b/Engine/Source/Programs/AutomationTool/AutomationUtils/CommandletUtils.cs
--- a/Engine/Source/Programs/AutomationTool/AutomationUtils/CommandletUtils.cs
+++ b/Engine/Source/Programs/AutomationTool/AutomationUtils/CommandletUtils.cs
@@ -254,7 +254,7 @@ namespace AutomationTool
 			string LocalLogFile = LogUtils.GetUniqueLogName(CombinePaths(CmdEnv.EngineSavedFolder, Commandlet));
 			Log("Commandlet log file is {0}", LocalLogFile);
 			string Args = String.Format(
-				"{0} -run={1} {2} -abslog={3} -stdout -FORCELOGFLUSH -CrashForUAT -unattended {5}{4}",
+				"{0} -run={1} {2} -abslog={3} -stdout -CrashForUAT -unattended {5}{4}",
 				(ProjectName == null) ? "" : CommandUtils.MakePathSafeToUseWithCommandLine(ProjectName.FullName),
 				Commandlet,
 				String.IsNullOrEmpty(Parameters) ? "" : Parameters,
diff --git a/Engine/Source/Programs/AutomationTool/Scripts/FixupRedirects.Automation.cs b/Engine/Source/Programs/AutomationTool/Scripts/FixupRedirects.Automation.cs
--- a/Engine/Source/Programs/AutomationTool/Scripts/FixupRedirects.Automation.cs
+++ b/Engine/Source/Programs/AutomationTool/Scripts/FixupRedirects.Automation.cs
@@ -11,8 +11,8 @@ class FixupRedirects : BuildCommand
 	public override void ExecuteBuild()
 	{
 		var EditorExe = CombinePaths(CmdEnv.LocalRoot, @"Engine/Binaries/Win64/UE4Editor-Cmd.exe");
-		Log("********** Running FixupRedirects: {0} -run=FixupRedirects -unattended -nopause -buildmachine -forcelogflush", EditorExe);
-		var RunResult = Run(EditorExe, String.Format("-run=FixupRedirects -unattended -nopause -buildmachine -forcelogflush -autosubmit"));
+		Log("********** Running FixupRedirects: {0} -run=FixupRedirects -unattended -nopause -buildmachine", EditorExe);
+		var RunResult = Run(EditorExe, String.Format("-run=FixupRedirects -unattended -nopause -buildmachine -autosubmit"));
 		if (RunResult.ExitCode != 0)
 		{
 			throw new AutomationException("BUILD FAILED: FixupRedirects failed");
diff --git a/Engine/Source/Programs/AutomationTool/Scripts/RunProjectCommand.Automation.cs b/Engine/Source/Programs/AutomationTool/Scripts/RunProjectCommand.Automation.cs
--- a/Engine/Source/Programs/AutomationTool/Scripts/RunProjectCommand.Automation.cs
+++ b/Engine/Source/Programs/AutomationTool/Scripts/RunProjectCommand.Automation.cs
@@ -217,7 +217,7 @@ public partial class Project : CommandUtils
 				FileStream ClientProcessLog = null;
 				StreamReader ClientLogReader = null;
 				Log("Starting Client for unattended test....");
-				ClientProcess = Run(ClientApp, ClientCmdLine + " -FORCELOGFLUSH -testexit=\"" + LookFor + "\"", null, ClientRunFlags | ERunOptions.NoWaitForExit);
+				ClientProcess = Run(ClientApp, ClientCmdLine + " -testexit=\"" + LookFor + "\"", null, ClientRunFlags | ERunOptions.NoWaitForExit);
 				while (!FileExists(ClientLogFile) && !ClientProcess.HasExited)
 				{
 					Log("Waiting for client logging process to start...{0}", ClientLogFile);
@@ -435,7 +435,7 @@ public partial class Project : CommandUtils
 										(AllServerOutput.Contains("Game Engine Initialized") || AllServerOutput.Contains("Unreal Network File Server is ready")))
 								{
 									Log("Starting Client for unattended test....");
-									ClientProcess = Run(ClientApp, ClientCmdLine + " -FORCELOGFLUSH -testexit=\"" + LookFor + "\"", null, ClientRunFlags | ERunOptions.NoWaitForExit);
+									ClientProcess = Run(ClientApp, ClientCmdLine + " -testexit=\"" + LookFor + "\"", null, ClientRunFlags | ERunOptions.NoWaitForExit);
 									//@todo no testing is done on these
 									if (NumClients > 1 && NumClients < 9)
 									{
@@ -982,7 +982,7 @@ public partial class Project : CommandUtils
 				Map += "?fake";
 			}
 
-			Args += String.Format("{0} -server -abslog={1}  -unattended -FORCELOGFLUSH -log -Messaging", Map, CommandUtils.MakePathSafeToUseWithCommandLine(ServerLogFile));
+			Args += String.Format("{0} -server -abslog={1}  -unattended -log -Messaging", Map, CommandUtils.MakePathSafeToUseWithCommandLine(ServerLogFile));
 
 			// Do not blindly add -nomcp, only do so if the client is using it
 			if (Params.RunCommandline.Contains("-nomcp"))
@@ -1030,7 +1030,7 @@ public partial class Project : CommandUtils
 	private static IProcessResult RunCookOnTheFlyServer(FileReference ProjectName, string ServerLogFile, string TargetPlatform, string AdditionalCommandLine)
 	{
 		var ServerApp = HostPlatform.Current.GetUE4ExePath("UE4Editor.exe");
-		var Args = String.Format("{0} -run=cook -cookonthefly -unattended -CrashForUAT -FORCELOGFLUSH -log",
+		var Args = String.Format("{0} -run=cook -cookonthefly -unattended -CrashForUAT -log",
 			CommandUtils.MakePathSafeToUseWithCommandLine(ProjectName.FullName));
 		if (!String.IsNullOrEmpty(ServerLogFile))
 		{
@@ -1075,7 +1075,7 @@ public partial class Project : CommandUtils
 		var UnrealFileServerExe = HostPlatform.Current.GetUE4ExePath("UnrealFileServer.exe");
 
 		Log("Running UnrealFileServer *******");
-		var Args = String.Format("{0} -abslog={1} -unattended -CrashForUAT -FORCELOGFLUSH -log {2}",
+		var Args = String.Format("{0} -abslog={1} -unattended -CrashForUAT -log {2}",
 						CommandUtils.MakePathSafeToUseWithCommandLine(Params.RawProjectPath.FullName),
 						CommandUtils.MakePathSafeToUseWithCommandLine(ServerLogFile),
 						AdditionalCommandLine);
diff --git a/Engine/Source/Programs/AutomationTool/Scripts/Tests.Automation.cs b/Engine/Source/Programs/AutomationTool/Scripts/Tests.Automation.cs
--- a/Engine/Source/Programs/AutomationTool/Scripts/Tests.Automation.cs
+++ b/Engine/Source/Programs/AutomationTool/Scripts/Tests.Automation.cs
@@ -685,7 +685,7 @@ class TestTestFarm : BuildCommand
 
 		string Exe = CombinePaths(CmdEnv.LocalRoot, "Engine", "Binaries", "Win64", "UE4Editor.exe");
 		string ClientLogFile = CombinePaths(CmdEnv.LogFolder, "HoverGameRun");
-		string CmdLine = " ../../../Samples/HoverShip/HoverShip.uproject -game -forcelogflush -log -abslog=" + ClientLogFile;
+		string CmdLine = " ../../../Samples/HoverShip/HoverShip.uproject -game -log -abslog=" + ClientLogFile;
 
 		IProcessResult App = Run(Exe, CmdLine, null, ERunOptions.AllowSpew | ERunOptions.NoWaitForExit | ERunOptions.AppMustExist | ERunOptions.NoStdOutRedirect);
 
@@ -798,7 +798,7 @@ class TestLog : BuildCommand
 		string ServerLogFile = CombinePaths(CmdEnv.LogFolder, "Server.log");
 
 		string ServerApp = EditorExe;
-		string OtherArgs = String.Format("{0} -server -abslog={1} -FORCELOGFLUSH -log", MapToRun, ServerLogFile);
+		string OtherArgs = String.Format("{0} -server -abslog={1} -log", MapToRun, ServerLogFile);
 		string StartArgs = "ShooterGame ";
 
 		string ServerCmdLine = StartArgs + OtherArgs;
@@ -1581,7 +1581,7 @@ class TestKillAll : BuildCommand
 
 		string Exe = CombinePaths(CmdEnv.LocalRoot, "Engine", "Binaries", "Win64", "UE4Editor.exe");
 		string ClientLogFile = CombinePaths(CmdEnv.LogFolder, "HoverGameRun");
-		string CmdLine = " ../../../Samples/Sandbox/HoverShip/HoverShip.uproject -game -forcelogflush -log -abslog=" + ClientLogFile;
+		string CmdLine = " ../../../Samples/Sandbox/HoverShip/HoverShip.uproject -game -log -abslog=" + ClientLogFile;
 
 		Run(Exe, CmdLine, null, ERunOptions.AllowSpew | ERunOptions.NoWaitForExit | ERunOptions.AppMustExist | ERunOptions.NoStdOutRedirect);
 		Thread.Sleep(10000);
@@ -1609,7 +1609,7 @@ class TestStopProcess : BuildCommand
 
 		string Exe = CombinePaths(CmdEnv.LocalRoot, "Engine", "Binaries", "Win64", "UE4Editor.exe");
 		string ClientLogFile = CombinePaths(CmdEnv.LogFolder, "HoverGameRun");
-		string CmdLine = " ../../../Samples/Sandbox/HoverShip/HoverShip.uproject -game -forcelogflush -log -ddc=noshared -abslog=" + ClientLogFile;
+		string CmdLine = " ../../../Samples/Sandbox/HoverShip/HoverShip.uproject -game -log -ddc=noshared -abslog=" + ClientLogFile;
 
 		for (int Attempt = 0; Attempt < 5; ++Attempt)
 		{
diff --git a/Engine/Source/Programs/UnrealDocTool/APIDocTool/APIDocTool/Program.cs b/Engine/Source/Programs/UnrealDocTool/APIDocTool/APIDocTool/Program.cs
--- a/Engine/Source/Programs/UnrealDocTool/APIDocTool/APIDocTool/Program.cs
+++ b/Engine/Source/Programs/UnrealDocTool/APIDocTool/APIDocTool/Program.cs
@@ -1104,7 +1104,7 @@ namespace APIDocTool
 				// Create the output directory
 				Utility.SafeCreateDirectory(JsonDir);
 
-				string Arguments = "-run=GenerateBlueprintAPI -path=" + JsonDir + " -name=BlueprintAPI -stdout -FORCELOGFLUSH -CrashForUAT -unattended -AllowStdOutLogVerbosity" + (bBuildMachine? " -buildmachine" : "");
+				string Arguments = "-run=GenerateBlueprintAPI -path=" + JsonDir + " -name=BlueprintAPI -stdout -CrashForUAT -unattended -AllowStdOutLogVerbosity" + (bBuildMachine? " -buildmachine" : "");
 				Console.WriteLine("Running: {0} {1}", EditorPath, Arguments);
 
 				using (Process JsonExportProcess = new Process())
