Duplicating CL#3234348 from //Odin/Main
	Fixed issues with Build Plugin script
	Added option to create a sub folder for the packaged plugin, which is used in the editor
	Fixed code assuming plugin must have modules

diff --git a/Engine/Plugins/Editor/PluginBrowser/Source/PluginBrowser/Private/SPluginTile.cpp b/Engine/Plugins/Editor/PluginBrowser/Source/PluginBrowser/Private/SPluginTile.cpp
--- a/Engine/Plugins/Editor/PluginBrowser/Source/PluginBrowser/Private/SPluginTile.cpp
+++ b/Engine/Plugins/Editor/PluginBrowser/Source/PluginBrowser/Private/SPluginTile.cpp
@@ -592,7 +592,7 @@ void SPluginTile::OnPackagePlugin()
 	// Ensure path is full rather than relative (for macs)
 	FString DescriptorFilename = Plugin->GetDescriptorFileName();
 	FString DescriptorFullPath = FPaths::ConvertRelativePathToFull(DescriptorFilename);
-	FString CommandLine = FString::Printf(TEXT("BuildPlugin -Rocket -Plugin=\"%s\" -Package=\"%s\""), *DescriptorFullPath, *OutputDirectory);
+	FString CommandLine = FString::Printf(TEXT("BuildPlugin -Rocket -Plugin=\"%s\" -Package=\"%s\" -CreateSubFolder"), *DescriptorFullPath, *OutputDirectory);
 
 #if PLATFORM_WINDOWS
 	FText PlatformName = LOCTEXT("PlatformName_Windows", "Windows");
diff --git a/Engine/Source/Programs/AutomationTool/Scripts/BuildPluginCommand.Automation.cs b/Engine/Source/Programs/AutomationTool/Scripts/BuildPluginCommand.Automation.cs
--- a/Engine/Source/Programs/AutomationTool/Scripts/BuildPluginCommand.Automation.cs
+++ b/Engine/Source/Programs/AutomationTool/Scripts/BuildPluginCommand.Automation.cs
@@ -49,6 +49,13 @@ class BuildPlugin : BuildCommand
 			throw new AutomationException("Output directory for packaged plugin must be outside engine directory");
 		}
 
+		// Whether to create a new sub folder with the name of the plugin to prevent deletion of everything in the package directory
+		bool bCreateSubFolder = ParseParam("CreateSubFolder");
+		if (bCreateSubFolder)
+		{
+			PackageDir = DirectoryReference.Combine(PackageDir, PluginFile.GetFileNameWithoutExtension());
+		}
+
 		// Clear the output directory of existing stuff
 		if (PackageDir.Exists())
 		{
@@ -154,13 +161,16 @@ class BuildPlugin : BuildCommand
 	{
 		// Find a list of modules that need to be built for this plugin
 		List<string> ModuleNames = new List<string>();
-		foreach(ModuleDescriptor Module in Plugin.Modules)
+		if (Plugin.Modules != null)
 		{
-			bool bBuildDeveloperTools = (TargetType == TargetRules.TargetType.Editor || TargetType == TargetRules.TargetType.Program);
-			bool bBuildEditor = (TargetType == TargetRules.TargetType.Editor);
-			if(Module.IsCompiledInConfiguration(Platform, TargetType, bBuildDeveloperTools, bBuildEditor))
+			foreach (ModuleDescriptor Module in Plugin.Modules)
 			{
-				ModuleNames.Add(Module.Name);
+				bool bBuildDeveloperTools = (TargetType == TargetRules.TargetType.Editor || TargetType == TargetRules.TargetType.Program);
+				bool bBuildEditor = (TargetType == TargetRules.TargetType.Editor);
+				if (Module.IsCompiledInConfiguration(Platform, TargetType, bBuildDeveloperTools, bBuildEditor))
+				{
+					ModuleNames.Add(Module.Name);
+				}
 			}
 		}
 
diff --git a/Engine/Source/Programs/UnrealBuildTool/System/ModuleDescriptor.cs b/Engine/Source/Programs/UnrealBuildTool/System/ModuleDescriptor.cs
--- a/Engine/Source/Programs/UnrealBuildTool/System/ModuleDescriptor.cs
+++ b/Engine/Source/Programs/UnrealBuildTool/System/ModuleDescriptor.cs
@@ -149,7 +149,7 @@ namespace UnrealBuildTool
 		/// <param name="Modules">Array of modules</param>
 		public static void WriteArray(JsonWriter Writer, string Name, ModuleDescriptor[] Modules)
 		{
-			if (Modules.Length > 0)
+			if (Modules != null && Modules.Length > 0)
 			{
 				Writer.WriteArrayStart(Name);
 				foreach (ModuleDescriptor Module in Modules)
