Fixing issues blocking plugin loc automation

diff --git a/Engine/Source/Programs/AutomationTool/AutomationUtils/CommandUtils.cs b/Engine/Source/Programs/AutomationTool/AutomationUtils/CommandUtils.cs
--- a/Engine/Source/Programs/AutomationTool/AutomationUtils/CommandUtils.cs
+++ b/Engine/Source/Programs/AutomationTool/AutomationUtils/CommandUtils.cs
@@ -1790,9 +1790,16 @@ namespace AutomationTool
 		/// <returns>True if param was found, false otherwise.</returns>
 		public static bool ParseParam(object[] ArgList, string Param)
 		{
-			foreach (object Arg in ArgList)
+            string ValueParam = Param;
+            if (!ValueParam.EndsWith("="))
+            {
+                ValueParam += "=";
+            }
+
+            foreach (object Arg in ArgList)
 			{
-				if (Arg.ToString().Equals(Param, StringComparison.InvariantCultureIgnoreCase))
+                string ArgStr = Arg.ToString();
+                if (ArgStr.Equals(Param, StringComparison.InvariantCultureIgnoreCase) || ArgStr.StartsWith(ValueParam, StringComparison.InvariantCultureIgnoreCase))
 				{
 					return true;
 				}
diff --git a/Engine/Source/Programs/AutomationTool/Localization/LocalizationProvider.cs b/Engine/Source/Programs/AutomationTool/Localization/LocalizationProvider.cs
--- a/Engine/Source/Programs/AutomationTool/Localization/LocalizationProvider.cs
+++ b/Engine/Source/Programs/AutomationTool/Localization/LocalizationProvider.cs
@@ -67,6 +67,7 @@ public abstract class LocalizationProvider
 	public struct LocalizationProviderArgs
 	{
 		public string RootWorkingDirectory;
+		public string RootLocalizationTargetDirectory;
 		public string RemoteFilenamePrefix;
 		public BuildCommand Command;
 		public int PendingChangeList;
diff --git a/Engine/Source/Programs/AutomationTool/Scripts/Localisation.Automation.cs b/Engine/Source/Programs/AutomationTool/Scripts/Localisation.Automation.cs
--- a/Engine/Source/Programs/AutomationTool/Scripts/Localisation.Automation.cs
+++ b/Engine/Source/Programs/AutomationTool/Scripts/Localisation.Automation.cs
@@ -25,14 +25,16 @@ class Localise : BuildCommand
 {
 	private struct LocalizationBatch
 	{
-		public LocalizationBatch(string InUEProjectDirectory, string InRemoteFilenamePrefix, List<string> InLocalizationProjectNames)
+		public LocalizationBatch(string InUEProjectDirectory, string InLocalizationTargetDirectory, string InRemoteFilenamePrefix, List<string> InLocalizationProjectNames)
 		{
 			UEProjectDirectory = InUEProjectDirectory;
+			LocalizationTargetDirectory = InLocalizationTargetDirectory;
 			RemoteFilenamePrefix = InRemoteFilenamePrefix;
 			LocalizationProjectNames = InLocalizationProjectNames;
 		}
 
 		public string UEProjectDirectory;
+		public string LocalizationTargetDirectory;
 		public string RemoteFilenamePrefix;
 		public List<string> LocalizationProjectNames;
 	};
@@ -127,7 +129,7 @@ class Localise : BuildCommand
 		// Add the static set of localization projects as a batch
 		if (LocalizationProjectNames.Count > 0)
 		{
-			LocalizationBatches.Add(new LocalizationBatch(UEProjectDirectory, "", LocalizationProjectNames));
+			LocalizationBatches.Add(new LocalizationBatch(UEProjectDirectory, UEProjectDirectory, "", LocalizationProjectNames));
 		}
 
 		// Build up any additional batches needed for plugins
@@ -151,7 +153,7 @@ class Localise : BuildCommand
 						PluginTargetNames.Add(LocalizationTarget.Name);
 					}
 
-					LocalizationBatches.Add(new LocalizationBatch(RootRelativePluginPath, PluginInfo.Name, PluginTargetNames));
+					LocalizationBatches.Add(new LocalizationBatch(UEProjectDirectory, RootRelativePluginPath, PluginInfo.Name, PluginTargetNames));
 				}
 			}
 		}
@@ -215,12 +217,14 @@ class Localise : BuildCommand
 	{
 		var EditorExe = CombinePaths(CmdEnv.LocalRoot, @"Engine/Binaries/Win64/UE4Editor-Cmd.exe");
 		var RootWorkingDirectory = CombinePaths(UEProjectRoot, LocalizationBatch.UEProjectDirectory);
+		var RootLocalizationTargetDirectory = CombinePaths(UEProjectRoot, LocalizationBatch.LocalizationTargetDirectory);
 
 		// Try and find our localization provider
 		LocalizationProvider LocProvider = null;
 		{
 			LocalizationProvider.LocalizationProviderArgs LocProviderArgs;
 			LocProviderArgs.RootWorkingDirectory = RootWorkingDirectory;
+			LocProviderArgs.RootLocalizationTargetDirectory = RootLocalizationTargetDirectory;
 			LocProviderArgs.RemoteFilenamePrefix = LocalizationBatch.RemoteFilenamePrefix;
 			LocProviderArgs.Command = this;
 			LocProviderArgs.PendingChangeList = PendingChangeList;
@@ -231,15 +235,15 @@ class Localise : BuildCommand
 		if (P4Enabled)
 		{
 			Log("Sync necessary content to head revision");
-			P4.Sync(P4Env.BuildRootP4 + "/" + LocalizationBatch.UEProjectDirectory + "/Config/Localization/...");
-			P4.Sync(P4Env.BuildRootP4 + "/" + LocalizationBatch.UEProjectDirectory + "/Content/Localization/...");
+			P4.Sync(P4Env.BuildRootP4 + "/" + LocalizationBatch.LocalizationTargetDirectory + "/Config/Localization/...");
+			P4.Sync(P4Env.BuildRootP4 + "/" + LocalizationBatch.LocalizationTargetDirectory + "/Content/Localization/...");
 		}
 
 		// Generate the info we need to gather for each project
 		var ProjectInfos = new List<ProjectInfo>();
 		foreach (var ProjectName in LocalizationBatch.LocalizationProjectNames)
 		{
-			ProjectInfos.Add(GenerateProjectInfo(RootWorkingDirectory, ProjectName, LocalizationSteps));
+			ProjectInfos.Add(GenerateProjectInfo(RootLocalizationTargetDirectory, ProjectName, LocalizationSteps));
 		}
 
 		if (LocalizationSteps.Contains("Download") && LocProvider != null)
@@ -421,7 +425,7 @@ class Localise : BuildCommand
 
 		foreach (var LocalizationBatch in LocalizationBatches)
 		{
-			var LocalizationPath = CombinePaths(UEProjectRoot, LocalizationBatch.UEProjectDirectory, "Content", "Localization");
+			var LocalizationPath = CombinePaths(UEProjectRoot, LocalizationBatch.LocalizationTargetDirectory, "Content", "Localization");
 
 			string[] POFileNames = Directory.GetFiles(LocalizationPath, "*.po", SearchOption.AllDirectories);
 			foreach (var POFileName in POFileNames)
