Copying fixes to the ListThirdPartySoftware script from //UE4/Dev-Core.

diff --git a/Engine/Source/Programs/AutomationTool/Scripts/ListThirdPartySoftware.Automation.cs b/Engine/Source/Programs/AutomationTool/Scripts/ListThirdPartySoftware.Automation.cs
--- a/Engine/Source/Programs/AutomationTool/Scripts/ListThirdPartySoftware.Automation.cs
+++ b/Engine/Source/Programs/AutomationTool/Scripts/ListThirdPartySoftware.Automation.cs
@@ -6,6 +6,7 @@ using System.IO;
 using AutomationTool;
 using UnrealBuildTool;
 using System.Linq;
+using System.Text.RegularExpressions;
 
 [Help("Lists TPS files associated with any source used to build a specified target(s). Grabs TPS files associated with source modules, content, and engine shaders.")]
 [Help("Target", "one or more Project|Config|Platform values describing all the targets to get associated TPS files for.")]
@@ -38,10 +39,7 @@ class ListThirdPartySoftware : BuildCommand
 		})
 		.SelectMany(Target =>
 			// run UBT on each target to get the build folders, add a few hard-coded paths, then search for TPS files in relevant folders to them
-			RunAndLog(UE4Build.GetUBTExecutable(), string.Format("{0} {1} {2} -listbuildfolders", Target.Name, Target.Configuration, Target.Platform), Options: DebugRun ? ERunOptions.Default : ERunOptions.NoLoggingOfRunCommand) // run UBT to get the raw output
-			.EnumerateLines() // split it into separate lines
-			.Where(Line => Line.StartsWith("BuildFolder:", StringComparison.InvariantCultureIgnoreCase)) // find the output lines that we need
-			.Select(Line => Line.Substring("BuildFolder:".Length)) // chop it down to the path name
+			EnumerateBuildFolders(Target.Name, Target.Configuration, Target.Platform, DebugRun)
 			// tack on the target content directory, which we could infer from one of a few places
 			.Concat(new[] 
 				{
@@ -82,4 +80,57 @@ class ListThirdPartySoftware : BuildCommand
 		.OrderBy(TPSFile => TPSFile) // sort them.
 		.ToList().ForEach(TPSFile => UnrealBuildTool.Log.WriteLine(0, string.Empty, LogEventType.Console, TPSFile)); // log them.
 	}
+
+	IEnumerable<string> EnumerateBuildFolders(string Target, string Configuration, string Platform, bool DebugRun)
+	{
+		FileReference OutputFile = FileReference.Combine(CommandUtils.EngineDirectory, "Intermediate", "Build", "ThirdParty.json");
+
+		Run(UE4Build.GetUBTExecutable(), String.Format("{0} {1} {2} -jsonexport=\"{3}\" -skipbuild", Target, Configuration, Platform, OutputFile.FullName), Options: DebugRun ? ERunOptions.Default : ERunOptions.NoLoggingOfRunCommand);
+
+		JsonObject Object = JsonObject.Read(OutputFile.FullName);
+
+		// local function that takes a RuntimeDependency path and resolves it (replacing Env vars that we support)
+		Func<string, DirectoryReference> ResolveRuntimeDependencyFolder = (string DependencyPath) =>
+		{
+			return new DirectoryReference(Path.GetDirectoryName(
+				// Regex to replace the env vars we support $(EngineDir|ProjectDir), ignoring case
+				Regex.Replace(DependencyPath, @"\$\((?<Type>Engine|Project)Dir\)", M => 
+					M.Groups["Type"].Value.Equals("Engine", StringComparison.InvariantCultureIgnoreCase)
+						? CommandUtils.EngineDirectory.FullName 
+						: new FileReference(Object.GetStringField("ProjectFile")).Directory.FullName,
+				RegexOptions.IgnoreCase)));
+		};
+
+		JsonObject Modules = Object.GetObjectField("Modules");
+
+		// Create a set of directories used for each binary
+		List<DirectoryReference> DirectoriesToScan = Modules
+			// get all directories that the module uses (source folder and any runtime dependencies)
+			.KeyNames.Select(KeyName => Modules.GetObjectField(KeyName))
+			.SelectMany(Module => Module
+				// resolve any runtime dependency folders and add them.
+				.GetObjectArrayField("RuntimeDependencies").Select(Dependency => ResolveRuntimeDependencyFolder(Dependency.GetStringField("Path")))
+				// Add on the module source directory
+				.Concat(new[] { new DirectoryReference(Module.GetStringField("Directory")) }))
+			// remove any duplicate folders since some modules may be from the same plugin
+			.Distinct()
+			// Project to a list as we need to do an O(n^2) operation below.
+			.ToList();
+
+		List<string> FinalDirs = DirectoriesToScan.Where(RemovalCandidate =>
+			// O(n^2) search to remove subfolders of any we are already searching.
+			// look for directories that aren't subdirectories of any other directory in the list.
+			!DirectoriesToScan.Any(DirectoryToScan =>
+				// != check because this inner loop will eventually check against itself
+				RemovalCandidate != DirectoryToScan &&
+				RemovalCandidate.IsUnderDirectory(DirectoryToScan)))
+			// grab the full name
+			.Select(Dir=>Dir.FullName)
+			// sort the final output
+			.OrderBy(Dir=> Dir)
+			// log the folders
+			.ToList();
+
+		return FinalDirs;
+	}
 }
diff --git a/Engine/Source/Programs/UnrealBuildTool/Utilities/JsonObject.cs b/Engine/Source/Programs/UnrealBuildTool/Utilities/JsonObject.cs
--- a/Engine/Source/Programs/UnrealBuildTool/Utilities/JsonObject.cs
+++ b/Engine/Source/Programs/UnrealBuildTool/Utilities/JsonObject.cs
@@ -394,6 +394,21 @@ namespace UnrealBuildTool
 			}
 		}
 
+		/// <summary>
+		/// Gets an object array field by the given name from the object, throwing an exception if it is not there or cannot be parsed.
+		/// </summary>
+		/// <param name="FieldName">Name of the field to get</param>
+		/// <returns>The field value</returns>
+		public JsonObject[] GetObjectArrayField(string FieldName)
+		{
+			JsonObject[] Result;
+			if (!TryGetObjectArrayField(FieldName, out Result))
+			{
+				throw new JsonParseException("Missing or invalid '{0}' field", FieldName);
+			}
+			return Result;
+		}
+
 		/// <summary>
 		/// Tries to read an object array field by the given name from the object
 		/// </summary>
