Fix ListThirdPartySoftware commandlet when running using the remote toolchain, and fix UBT error codes being ignored.

diff --git a/Engine/Source/Programs/AutomationTool/Scripts/ListThirdPartySoftware.Automation.cs b/Engine/Source/Programs/AutomationTool/Scripts/ListThirdPartySoftware.Automation.cs
--- a/Engine/Source/Programs/AutomationTool/Scripts/ListThirdPartySoftware.Automation.cs
+++ b/Engine/Source/Programs/AutomationTool/Scripts/ListThirdPartySoftware.Automation.cs
@@ -85,7 +85,11 @@ class ListThirdPartySoftware : BuildCommand
 	{
 		FileReference OutputFile = FileReference.Combine(CommandUtils.EngineDirectory, "Intermediate", "Build", "ThirdParty.json");
 
-		Run(UE4Build.GetUBTExecutable(), String.Format("{0} {1} {2} -jsonexport=\"{3}\" -skipbuild", Target, Configuration, Platform, OutputFile.FullName), Options: DebugRun ? ERunOptions.Default : ERunOptions.NoLoggingOfRunCommand);
+		IProcessResult Result = Run(UE4Build.GetUBTExecutable(), String.Format("{0} {1} {2} -jsonexport=\"{3}\" -skipbuild", Target, Configuration, Platform, OutputFile.FullName), Options: DebugRun ? ERunOptions.Default : ERunOptions.NoLoggingOfRunCommand);
+		if(Result.ExitCode != 0)
+		{
+			throw new AutomationException("Failed to run UBT");
+		}
 
 		JsonObject Object = JsonObject.Read(OutputFile.FullName);
 
diff --git a/Engine/Source/Programs/UnrealBuildTool/Mac/UEBuildMac.cs b/Engine/Source/Programs/UnrealBuildTool/Mac/UEBuildMac.cs
--- a/Engine/Source/Programs/UnrealBuildTool/Mac/UEBuildMac.cs
+++ b/Engine/Source/Programs/UnrealBuildTool/Mac/UEBuildMac.cs
@@ -102,8 +102,8 @@ namespace UnrealBuildTool
 				Target.bUsePCHFiles = false;
 			}
 
-			// Needs OS X 10.11 for Metal
-			if (MacToolChain.Settings.MacOSSDKVersionFloat >= 10.11f && Target.bCompileAgainstEngine)
+			// Needs OS X 10.11 for Metal. The remote toolchain has not been initialized yet, so just assume it's a recent SDK.
+			if ((BuildHostPlatform.Current.Platform != UnrealTargetPlatform.Mac || MacToolChain.Settings.MacOSSDKVersionFloat >= 10.11f) && Target.bCompileAgainstEngine)
 			{
 				Target.GlobalDefinitions.Add("HAS_METAL=1");
 				Target.ExtraModuleNames.Add("MetalRHI");
