Loc automation now does its P4 revert using an arguments file to avoid overflowing the commandline size limit

diff --git a/Engine/Source/Programs/AutomationTool/Scripts/Localisation.Automation.cs b/Engine/Source/Programs/AutomationTool/Scripts/Localisation.Automation.cs
--- a/Engine/Source/Programs/AutomationTool/Scripts/Localisation.Automation.cs
+++ b/Engine/Source/Programs/AutomationTool/Scripts/Localisation.Automation.cs
@@ -192,16 +192,18 @@ class Localise : BuildCommand
 
 				if (POFilesToRevert.Count > 0)
 				{
-					var P4RevertCommandline = new StringBuilder();
-					foreach (var POFileToRevert in POFilesToRevert)
+					var P4RevertArgsFilename = CombinePaths(CmdEnv.LocalRoot, "Engine", "Intermediate", String.Format("LocalizationP4RevertArgs-{0}.txt", Guid.NewGuid().ToString()));
+
+					using (StreamWriter P4RevertArgsWriter = File.CreateText(P4RevertArgsFilename))
 					{
-						P4RevertCommandline.Append('"');
-						P4RevertCommandline.Append(POFileToRevert);
-						P4RevertCommandline.Append('"');
-						P4RevertCommandline.Append(' ');
+						foreach (var POFileToRevert in POFilesToRevert)
+						{
+							P4RevertArgsWriter.WriteLine(POFileToRevert);
+						}
 					}
 
-					P4.Revert(P4RevertCommandline.ToString());
+					P4.LogP4(String.Format("-x{0} revert", P4RevertArgsFilename));
+					DeleteFile_NoExceptions(P4RevertArgsFilename);
 				}
 			}
 
