Upping XLOC client content length limit as we're hitting it, and they have a much larger limit set on the server

diff --git a/Engine/Source/Programs/AutomationTool/XLocLocalization/XLocLocalizationProvider.cs b/Engine/Source/Programs/AutomationTool/XLocLocalization/XLocLocalizationProvider.cs
--- a/Engine/Source/Programs/AutomationTool/XLocLocalization/XLocLocalizationProvider.cs
+++ b/Engine/Source/Programs/AutomationTool/XLocLocalization/XLocLocalizationProvider.cs
@@ -302,7 +302,7 @@ public abstract class XLocLocalizationProvider : LocalizationProvider
 		Binding.TransferMode = TransferMode.Buffered;
 		Binding.UseDefaultWebProxy = true;
 		Binding.ReaderQuotas.MaxDepth = 32;
-		Binding.ReaderQuotas.MaxStringContentLength = 8192;
+		Binding.ReaderQuotas.MaxStringContentLength = 65536;
 		Binding.ReaderQuotas.MaxArrayLength = 16384;
 		Binding.ReaderQuotas.MaxBytesPerRead = 4096;
 		Binding.ReaderQuotas.MaxNameTableCharCount = 16384;
