Fix assembly version error that was occurring on some builders
Better reporting of this in OrionTestFramework should it occur again
Duplicated fix from //Orion/Dev-General where AI in tests were not being forced to the desired hero.

diff --git a/Engine/Source/Programs/AutomationTool/app.config b/Engine/Source/Programs/AutomationTool/app.config
--- a/Engine/Source/Programs/AutomationTool/app.config
+++ b/Engine/Source/Programs/AutomationTool/app.config
@@ -1,3 +1,14 @@
-<?xml version="1.0"?>
+﻿<?xml version="1.0" encoding="utf-8"?>
 <configuration>
-<startup><supportedRuntime version="v4.0" sku=".NETFramework,Version=v4.5"/></startup></configuration>
+	<startup>
+		<supportedRuntime version="v4.0" sku=".NETFramework,Version=v4.5" />
+	</startup>
+	<runtime>
+		<assemblyBinding xmlns="urn:schemas-microsoft-com:asm.v1">
+			<dependentAssembly>
+				<assemblyIdentity name="Newtonsoft.Json" publicKeyToken="30ad4fe6b2a6aeed" culture="neutral" />
+				<bindingRedirect oldVersion="0.0.0.0-9.0.0.0" newVersion="9.0.0.0" />
+			</dependentAssembly>
+		</assemblyBinding>
+	</runtime>
+</configuration>
