Copying //UE4/Dev-Core to //UE4/Dev-Main (Source: //UE4/Dev-Core @ 3345728)

diff --git a/Engine/Source/Developer/BlueprintNativeCodeGen/Private/BlueprintNativeCodeGenModule.cpp b/Engine/Source/Developer/BlueprintNativeCodeGen/Private/BlueprintNativeCodeGenModule.cpp
--- a/Engine/Source/Developer/BlueprintNativeCodeGen/Private/BlueprintNativeCodeGenModule.cpp
+++ b/Engine/Source/Developer/BlueprintNativeCodeGen/Private/BlueprintNativeCodeGenModule.cpp
@@ -26,7 +26,7 @@
 #include "Kismet2/BlueprintEditorUtils.h"
 #include "Engine/InheritableComponentHandler.h"
 #include "Animation/AnimBlueprint.h"
-
+#include "Interfaces/ITargetPlatform.h"
 /*******************************************************************************
 * NativizationCookControllerImpl
 ******************************************************************************/
@@ -56,23 +56,25 @@ public:
 	virtual void FinalizeManifest() override;
 	virtual void GenerateStubs() override;
 	virtual void GenerateFullyConvertedClasses() override;
-	virtual void MarkUnconvertedBlueprintAsNecessary(TAssetPtr<UBlueprint> BPPtr) override;
+	void MarkUnconvertedBlueprintAsNecessary(TAssetPtr<UBlueprint> BPPtr, const FCompilerNativizationOptions& NativizationOptions) override;
 	virtual const TMultiMap<FName, TAssetSubclassOf<UObject>>& GetFunctionsBoundToADelegate() override;
 
 	FFileHelper::EEncodingOptions::Type ForcedEncoding() const
 	{
 		return FFileHelper::EEncodingOptions::ForceUTF8;
 	}
+	virtual const FCompilerNativizationOptions& GetNativizationOptionsForPlatform(const FString& PlatformName) const override;
+	virtual void FillPlatformNativizationDetails(const ITargetPlatform* Platform, FPlatformNativizationDetails& OutDetails) override;
 protected:
 	virtual void Initialize(const FNativeCodeGenInitData& InitData) override;
 	virtual void InitializeForRerunDebugOnly(const TArray<FPlatformNativizationDetails>& CodegenTargets) override;
 	//~ End IBlueprintNativeCodeGenModule interface
 
 	//~ Begin FScriptCookReplacmentCoordinator interface
-	virtual EReplacementResult IsTargetedForReplacement(const UPackage* Package) const override;
-	virtual EReplacementResult IsTargetedForReplacement(const UObject* Object) const override;
-	virtual UClass* FindReplacedClassForObject(const UObject* Object) const override;
-	virtual UObject* FindReplacedNameAndOuter(UObject* Object, FName& OutName) const override;
+	virtual EReplacementResult IsTargetedForReplacement(const UPackage* Package, const FCompilerNativizationOptions& NativizationOptions) const override;
+	virtual EReplacementResult IsTargetedForReplacement(const UObject* Object, const FCompilerNativizationOptions& NativizationOptions) const override;
+	virtual UClass* FindReplacedClassForObject(const UObject* Object, const FCompilerNativizationOptions& NativizationOptions) const override;
+	virtual UObject* FindReplacedNameAndOuter(UObject* Object, FName& OutName, const FCompilerNativizationOptions& NativizationOptions) const override;
 	//~ End FScriptCookReplacmentCoordinator interface
 private:
 	void ReadConfig();
@@ -83,6 +85,18 @@ private:
 	void CollectBoundFunctions(UBlueprint* BP);
 	void GenerateSingleAsset(UField* ForConversion, const TCHAR* PlatformName, TSharedPtr<FNativizationSummary> NativizationSummary = TSharedPtr<FNativizationSummary>());
 
+	struct FStatePerPlatform
+	{
+		// A stub-wrapper must be generated only if the BP is really accessed/required by some other generated code.
+		TSet<TAssetPtr<UBlueprint>> StubsRequiredByGeneratedCode;
+
+		TSet<TAssetPtr<UBlueprint>> ToGenerate;
+		// Cached values from IsTargetedForReplacement
+		mutable TMap<FStringAssetReference, EReplacementResult> CachedIsTargetedForReplacement;
+	};
+
+	TMap< FString, FStatePerPlatform > StatesPerPlatform;
+
 	TMap< FString, TUniquePtr<FBlueprintNativeCodeGenManifest> > Manifests;
 
 	// Children of these classes won't be nativized
@@ -96,17 +110,22 @@ private:
 
 	TArray<FString> TargetPlatformNames;
 
-	// A stub-wrapper must be generated only if the BP is really accessed/required by some other generated code.
-	TSet<TAssetPtr<UBlueprint>> StubsRequiredByGeneratedCode;
-	TSet<TAssetPtr<UBlueprint>> AllPotentialStubs;
-
-	TSet<TAssetPtr<UBlueprint>> ToGenerate;
 	TMultiMap<FName, TAssetSubclassOf<UObject>> FunctionsBoundToADelegate; // is a function could be bound to a delegate, then it must have UFUNCTION macro. So we cannot optimize it.
-
-	// Cached values from IsTargetedForReplacement
-	mutable TMap<FStringAssetReference, EReplacementResult> CachedIsTargetedForReplacement;
 };
 
+const FCompilerNativizationOptions& FBlueprintNativeCodeGenModule::GetNativizationOptionsForPlatform(const FString& PlatformName) const
+{
+	const TUniquePtr<FBlueprintNativeCodeGenManifest>* Result = Manifests.Find(PlatformName);
+	if (ensure(Result && Result->IsValid()))
+	{
+		const FBlueprintNativeCodeGenManifest& Manifest = **Result;
+		return Manifest.GetCompilerNativizationOptions();
+	}
+	UE_LOG(LogBlueprintCodeGen, Error, TEXT("Cannot find manifest for platform: %s"), *PlatformName);
+	static const FCompilerNativizationOptions FallbackNativizationOptions{};
+	return FallbackNativizationOptions;
+}
+
 void FBlueprintNativeCodeGenModule::ReadConfig()
 {
 	GConfig->GetArray(TEXT("BlueprintNativizationSettings"), TEXT("ExcludedAssetTypes"), ExcludedAssetTypes, GEditorIni);
@@ -133,29 +152,33 @@ void FBlueprintNativeCodeGenModule::ReadConfig()
 	GConfig->GetArray(TEXT("BlueprintNativizationSettings"), TEXT("ExcludedFolderPaths"), ExcludedFolderPaths, GEditorIni);
 }
 
-void FBlueprintNativeCodeGenModule::MarkUnconvertedBlueprintAsNecessary(TAssetPtr<UBlueprint> BPPtr)
+void FBlueprintNativeCodeGenModule::MarkUnconvertedBlueprintAsNecessary(TAssetPtr<UBlueprint> BPPtr, const FCompilerNativizationOptions& NativizationOptions)
 {
-	StubsRequiredByGeneratedCode.Add(BPPtr);
+	FStatePerPlatform* StateForCurrentPlatform = StatesPerPlatform.Find(NativizationOptions.PlatformName);
+	if (ensure(StateForCurrentPlatform))
+	{
+		StateForCurrentPlatform->StubsRequiredByGeneratedCode.Add(BPPtr);
+	}
 }
 
 void FBlueprintNativeCodeGenModule::FillTargetedForReplacementQuery()
 {
 	IBlueprintCompilerCppBackendModule& BackEndModule = (IBlueprintCompilerCppBackendModule&)IBlueprintCompilerCppBackendModule::Get();
 	auto& ConversionQueryDelegate = BackEndModule.OnIsTargetedForConversionQuery();
-	auto ShouldConvert = [](const UObject* AssetObj)
+	auto ShouldConvert = [](const UObject* AssetObj, const FCompilerNativizationOptions& NativizationOptions)
 	{
 		if (ensure(IBlueprintNativeCodeGenCore::Get()))
 		{
-			EReplacementResult ReplacmentResult = IBlueprintNativeCodeGenCore::Get()->IsTargetedForReplacement(AssetObj);
+			EReplacementResult ReplacmentResult = IBlueprintNativeCodeGenCore::Get()->IsTargetedForReplacement(AssetObj, NativizationOptions);
 			return ReplacmentResult == EReplacementResult::ReplaceCompletely;
 		}
 		return false;
 	};
 	ConversionQueryDelegate.BindStatic(ShouldConvert);
 
-	auto LocalMarkUnconvertedBlueprintAsNecessary = [](TAssetPtr<UBlueprint> BPPtr)
+	auto LocalMarkUnconvertedBlueprintAsNecessary = [](TAssetPtr<UBlueprint> BPPtr, const FCompilerNativizationOptions& NativizationOptions)
 	{
-		IBlueprintNativeCodeGenModule::Get().MarkUnconvertedBlueprintAsNecessary(BPPtr);
+		IBlueprintNativeCodeGenModule::Get().MarkUnconvertedBlueprintAsNecessary(BPPtr, NativizationOptions);
 	};
 	BackEndModule.OnIncludingUnconvertedBP().BindStatic(LocalMarkUnconvertedBlueprintAsNecessary);
 }
@@ -237,7 +260,12 @@ void FBlueprintNativeCodeGenModule::FillIsFunctionUsedInADelegate()
 
 void FBlueprintNativeCodeGenModule::Initialize(const FNativeCodeGenInitData& InitData)
 {
-	CachedIsTargetedForReplacement.Reset();
+	StatesPerPlatform.Reset();
+	for (const FPlatformNativizationDetails& Platform : InitData.CodegenTargets)
+	{
+		StatesPerPlatform.Add(Platform.PlatformName);
+	}
+
 	ReadConfig();
 
 	IBlueprintNativeCodeGenCore::Register(this);
@@ -264,7 +292,11 @@ void FBlueprintNativeCodeGenModule::Initialize(const FNativeCodeGenInitData& Ini
 
 void FBlueprintNativeCodeGenModule::InitializeForRerunDebugOnly(const TArray<FPlatformNativizationDetails>& CodegenTargets)
 {
-	CachedIsTargetedForReplacement.Reset();
+	StatesPerPlatform.Reset();
+	for (const FPlatformNativizationDetails& Platform : CodegenTargets)
+	{
+		StatesPerPlatform.Add(Platform.PlatformName);
+	}
 	ReadConfig();
 	IBlueprintNativeCodeGenCore::Register(this);
 	FillTargetedForReplacementQuery();
@@ -308,8 +340,9 @@ void FBlueprintNativeCodeGenModule::InitializeForRerunDebugOnly(const TArray<FPl
 				GenerateSingleStub(BP, *Platform.PlatformName);
 			}
 		}
-
-		for (TAssetPtr<UBlueprint>& BPPtr : ToGenerate)
+		FStatePerPlatform* State = StatesPerPlatform.Find(Platform.PlatformName);
+		check(State);
+		for (TAssetPtr<UBlueprint>& BPPtr : State->ToGenerate)
 		{
 			UBlueprint* BP = BPPtr.LoadSynchronous();
 			if (ensure(BP))
@@ -328,12 +361,14 @@ void FBlueprintNativeCodeGenModule::GenerateFullyConvertedClasses()
 		CodeGenBackend.NativizationSummary() = NativizationSummary;
 	}
 
-	for (TAssetPtr<UBlueprint>& BPPtr : ToGenerate)
+	for (const FString& PlatformName : TargetPlatformNames)
 	{
-		UBlueprint* BP = BPPtr.LoadSynchronous();
-		if (ensure(BP))
+		FStatePerPlatform* State = StatesPerPlatform.Find(PlatformName);
+		check(State);
+		for (TAssetPtr<UBlueprint>& BPPtr : State->ToGenerate)
 		{
-			for (const FString& PlatformName : TargetPlatformNames)
+			UBlueprint* BP = BPPtr.LoadSynchronous();
+			if (ensure(BP))
 			{
 				GenerateSingleAsset(BP->GeneratedClass, *PlatformName, NativizationSummary);
 			}
@@ -397,7 +432,7 @@ void FBlueprintNativeCodeGenModule::GenerateSingleStub(UBlueprint* BP, const TCH
 	// Apparently we can only generate wrappers for classes, so any logic that results in non classes requesting
 	// wrappers will fail here:
 
-	FileContents = Backend_CPP->GenerateWrapperForClass(Class);
+	FileContents = Backend_CPP->GenerateWrapperForClass(Class, GetManifest(PlatformName).GetCompilerNativizationOptions());
 
 	if (!FileContents.IsEmpty())
 	{
@@ -474,33 +509,36 @@ void FBlueprintNativeCodeGenModule::GenerateSingleAsset(UField* ForConversion, c
 
 void FBlueprintNativeCodeGenModule::GenerateStubs()
 {
-	TSet<TAssetPtr<UBlueprint>> AlreadyGenerated;
-	while (AlreadyGenerated.Num() < StubsRequiredByGeneratedCode.Num())
+	for (auto& PlatformName : TargetPlatformNames)
 	{
-		const int32 OldGeneratedNum = AlreadyGenerated.Num();
-		for (TAssetPtr<UBlueprint>& BPPtr : StubsRequiredByGeneratedCode)
+		FStatePerPlatform* StateForCurrentPlatform = StatesPerPlatform.Find(PlatformName);
+		if (!ensure(StateForCurrentPlatform))
+		{
+			continue;
+		}
+		TSet<TAssetPtr<UBlueprint>> AlreadyGenerated;
+		while (AlreadyGenerated.Num() < StateForCurrentPlatform->StubsRequiredByGeneratedCode.Num())
 		{
-			bool bAlreadyGenerated = false;
-			AlreadyGenerated.Add(BPPtr, &bAlreadyGenerated);
-			if (bAlreadyGenerated)
+			const int32 OldGeneratedNum = AlreadyGenerated.Num();
+			TSet<TAssetPtr<UBlueprint>> LocalCopyStubsRequiredByGeneratedCode = StateForCurrentPlatform->StubsRequiredByGeneratedCode;
+			for (TAssetPtr<UBlueprint>& BPPtr : LocalCopyStubsRequiredByGeneratedCode)
 			{
-				continue;
-			}
+				bool bAlreadyGenerated = false;
+				AlreadyGenerated.Add(BPPtr, &bAlreadyGenerated);
+				if (bAlreadyGenerated)
+				{
+					continue;
+				}
 
-			ensureMsgf(AllPotentialStubs.Contains(BPPtr), TEXT("A required blueprint doesn't generate stub: %s"), *BPPtr.ToString());
-			for (auto& PlatformName : TargetPlatformNames)
-			{
 				GenerateSingleStub(BPPtr.LoadSynchronous(), *PlatformName);
 			}
-		}
-
-		if (!ensure(OldGeneratedNum != AlreadyGenerated.Num()))
-		{
-			break;
+			// make sure there was any progress
+			if (!ensure(OldGeneratedNum != AlreadyGenerated.Num()))
+			{
+				break;
+			}
 		}
 	}
-
-	UE_LOG(LogBlueprintCodeGen, Log, TEXT("GenerateStubs - all unconverted bp: %d, generated wrapers: %d"), AllPotentialStubs.Num(), StubsRequiredByGeneratedCode.Num());
 }
 
 void FBlueprintNativeCodeGenModule::Convert(UPackage* Package, ESavePackageResult CookResult, const TCHAR* PlatformName)
@@ -530,13 +568,15 @@ void FBlueprintNativeCodeGenModule::Convert(UPackage* Package, ESavePackageResul
 		return;
 	}
 
+	FStatePerPlatform* State = StatesPerPlatform.Find(PlatformName);
+	if (!ensure(State))
+	{
+		return;
+	}
 	if (CookResult == ESavePackageResult::GenerateStub)
 	{
-		if (ensure(BP))
-		{
-			ensure(!ToGenerate.Contains(BP));
-			AllPotentialStubs.Add(BP);
-		}
+		// No stub is generated for structs and enums.
+		ensure(!BP || !State->ToGenerate.Contains(BP));
 	}
 	else
 	{
@@ -545,8 +585,7 @@ void FBlueprintNativeCodeGenModule::Convert(UPackage* Package, ESavePackageResul
 		{
 			if (ensure(BP))
 			{
-				ensure(!AllPotentialStubs.Contains(BP));
-				ToGenerate.Add(BP);
+				State->ToGenerate.Add(BP);
 			}
 		}
 		else
@@ -577,18 +616,35 @@ void FBlueprintNativeCodeGenModule::MergeManifest(int32 ManifestIdentifier)
 
 void FBlueprintNativeCodeGenModule::FinalizeManifest()
 {
+	IBlueprintCompilerCppBackendModule& CodeGenBackend = (IBlueprintCompilerCppBackendModule&)IBlueprintCompilerCppBackendModule::Get();
+	TSharedPtr<FNativizationSummary> NativizationSummary = CodeGenBackend.NativizationSummary();
 	for(auto& PlatformName : TargetPlatformNames)
 	{
-		GetManifest(*PlatformName).Save(-1);
-		check(FBlueprintNativeCodeGenUtils::FinalizePlugin(GetManifest(*PlatformName)));
+		FBlueprintNativeCodeGenManifest& Manifest = GetManifest(*PlatformName);
+		if (NativizationSummary.IsValid())
+		{
+			TSet<TAssetPtr<UPackage>>* RequiredModules = NativizationSummary->ModulesRequiredByPlatform.Find(Manifest.GetCompilerNativizationOptions().PlatformName);
+			if (RequiredModules)
+			{
+				for (TAssetPtr<UPackage> ItPackage : *RequiredModules)
+				{
+					if (UPackage* Pkg = ItPackage.Get())
+					{
+						Manifest.AddSingleModuleDependency(Pkg);
+					}
+				}
+			}
+		}
+		Manifest.Save(-1);
+		check(FBlueprintNativeCodeGenUtils::FinalizePlugin(Manifest));
 	}
 }
 
-UClass* FBlueprintNativeCodeGenModule::FindReplacedClassForObject(const UObject* Object) const
+UClass* FBlueprintNativeCodeGenModule::FindReplacedClassForObject(const UObject* Object, const FCompilerNativizationOptions& NativizationOptions) const
 {
 	// we're only looking to replace class types:
 	if (Object && Object->IsA<UField>() 
-		&& (EReplacementResult::ReplaceCompletely == IsTargetedForReplacement(Object)))
+		&& (EReplacementResult::ReplaceCompletely == IsTargetedForReplacement(Object, NativizationOptions)))
 	{
 		for (const UClass* Class = Object->GetClass(); Class; Class = Class->GetSuperClass())
 		{
@@ -606,11 +662,10 @@ UClass* FBlueprintNativeCodeGenModule::FindReplacedClassForObject(const UObject*
 			}
 		}
 	}
-	ensure(!Object || Object->HasAnyFlags(RF_ClassDefaultObject) || !(Object->IsA<UUserDefinedStruct>() || Object->IsA<UUserDefinedEnum>()));
 	return nullptr;
 }
 
-UObject* FBlueprintNativeCodeGenModule::FindReplacedNameAndOuter(UObject* Object, FName& OutName) const
+UObject* FBlueprintNativeCodeGenModule::FindReplacedNameAndOuter(UObject* Object, FName& OutName, const FCompilerNativizationOptions& NativizationOptions) const
 {
 	OutName = NAME_None;
 
@@ -664,7 +719,7 @@ UObject* FBlueprintNativeCodeGenModule::FindReplacedNameAndOuter(UObject* Object
 			}
 		}
 
-		if (OuterCDO && (EReplacementResult::ReplaceCompletely == IsTargetedForReplacement(OuterCDO->GetClass())))
+		if (OuterCDO && (EReplacementResult::ReplaceCompletely == IsTargetedForReplacement(OuterCDO->GetClass(), NativizationOptions)))
 		{
 			OutName = NewName;
 			UE_LOG(LogBlueprintCodeGen, Log, TEXT("Object '%s' has replaced name '%s' and outer: '%s'"), *GetPathNameSafe(Object), *OutName.ToString(), *GetPathNameSafe(OuterCDO));
@@ -677,7 +732,7 @@ UObject* FBlueprintNativeCodeGenModule::FindReplacedNameAndOuter(UObject* Object
 		if (OuterCAC && OuterCAC->GetChildActorTemplate() == Object)
 		{
 			UBlueprintGeneratedClass* BPGC = GetOuterBPGC(OuterCAC->GetOuter());
-			if (BPGC && (EReplacementResult::ReplaceCompletely == IsTargetedForReplacement(BPGC)))
+			if (BPGC && (EReplacementResult::ReplaceCompletely == IsTargetedForReplacement(BPGC, NativizationOptions)))
 			{
 				return BPGC;
 			}
@@ -687,7 +742,7 @@ UObject* FBlueprintNativeCodeGenModule::FindReplacedNameAndOuter(UObject* Object
 	return nullptr;
 }
 
-EReplacementResult FBlueprintNativeCodeGenModule::IsTargetedForReplacement(const UPackage* Package) const
+EReplacementResult FBlueprintNativeCodeGenModule::IsTargetedForReplacement(const UPackage* Package, const FCompilerNativizationOptions& NativizationOptions) const
 {
 	// non-native packages with enums and structs should be converted, unless they are blacklisted:
 	UStruct* Struct = nullptr;
@@ -699,10 +754,10 @@ EReplacementResult FBlueprintNativeCodeGenModule::IsTargetedForReplacement(const
 	{
 		Target = Enum;
 	}
-	return IsTargetedForReplacement(Target);
+	return IsTargetedForReplacement(Target, NativizationOptions);
 }
 
-EReplacementResult FBlueprintNativeCodeGenModule::IsTargetedForReplacement(const UObject* Object) const
+EReplacementResult FBlueprintNativeCodeGenModule::IsTargetedForReplacement(const UObject* Object, const FCompilerNativizationOptions& NativizationOptions) const
 {
 	if (Object == nullptr)
 	{
@@ -717,9 +772,11 @@ EReplacementResult FBlueprintNativeCodeGenModule::IsTargetedForReplacement(const
 		return EReplacementResult::DontReplace;
 	}
 
+	const FStatePerPlatform* StateForCurrentPlatform = StatesPerPlatform.Find(NativizationOptions.PlatformName);
+	check(StateForCurrentPlatform);
 	const FStringAssetReference ObjectKey(Object);
 	{
-		const EReplacementResult* const CachedValue = CachedIsTargetedForReplacement.Find(ObjectKey); //THe referenced returned by FindOrAdd could be invalid later, when filled.
+		const EReplacementResult* const CachedValue = StateForCurrentPlatform->CachedIsTargetedForReplacement.Find(ObjectKey); //THe referenced returned by FindOrAdd could be invalid later, when filled.
 		if (CachedValue)
 		{
 			return *CachedValue;
@@ -784,14 +841,13 @@ EReplacementResult FBlueprintNativeCodeGenModule::IsTargetedForReplacement(const
 	};
 	if (ObjectIsNotReplacedAtAll())
 	{
-		CachedIsTargetedForReplacement.Add(ObjectKey, EReplacementResult::DontReplace);
+		StateForCurrentPlatform->CachedIsTargetedForReplacement.Add(ObjectKey, EReplacementResult::DontReplace);
 		return EReplacementResult::DontReplace;
 	}
 
 	auto ObjectGenratesOnlyStub = [&]() -> bool
 	{
 		// ExcludedFolderPaths
-		if (BlueprintClass)
 		{
 			const FString ObjPathName = Object->GetPathName();
 			for (const FString& ExcludedPath : ExcludedFolderPaths)
@@ -801,6 +857,13 @@ EReplacementResult FBlueprintNativeCodeGenModule::IsTargetedForReplacement(const
 					return true;
 				}
 			}
+			for (const FString& ExcludedPath : NativizationOptions.ExcludedFolderPaths)
+			{
+				if (ObjPathName.StartsWith(ExcludedPath))
+				{
+					return true;
+				}
+			}
 		}
 
 		// ExcludedAssetTypes
@@ -828,6 +891,10 @@ EReplacementResult FBlueprintNativeCodeGenModule::IsTargetedForReplacement(const
 			{
 				return true;
 			}
+			if (NativizationOptions.ExcludedAssets.Contains(Object->GetOutermost()))
+			{
+				return true;
+			}
 		}
 
 
@@ -881,7 +948,7 @@ EReplacementResult FBlueprintNativeCodeGenModule::IsTargetedForReplacement(const
 				}
 			}
 
-			const bool bFlaggedForNativization  = (Blueprint->NativizationFlag == EBlueprintNativizationFlag::Dependency) ? 
+			const bool bFlaggedForNativization = (Blueprint->NativizationFlag == EBlueprintNativizationFlag::Dependency) ?
 				PackagingSettings->IsBlueprintAssetInNativizationList(Blueprint) :
 				(Blueprint->NativizationFlag == EBlueprintNativizationFlag::ExplicitlyEnabled);
 			// Blueprint is not selected
@@ -894,7 +961,7 @@ EReplacementResult FBlueprintNativeCodeGenModule::IsTargetedForReplacement(const
 			for (const UBlueprintGeneratedClass* ParentClassIt = Cast<UBlueprintGeneratedClass>(BlueprintClass->GetSuperClass())
 				; ParentClassIt; ParentClassIt = Cast<UBlueprintGeneratedClass>(ParentClassIt->GetSuperClass()))
 			{
-				const EReplacementResult ParentResult = IsTargetedForReplacement(ParentClassIt);
+				const EReplacementResult ParentResult = IsTargetedForReplacement(ParentClassIt, NativizationOptions);
 				if (ParentResult != EReplacementResult::ReplaceCompletely)
 				{
 					if (bNativizeOnlySelectedBPs)
@@ -913,7 +980,7 @@ EReplacementResult FBlueprintNativeCodeGenModule::IsTargetedForReplacement(const
 				const UBlueprintGeneratedClass* InterfaceBPGC = Cast<const UBlueprintGeneratedClass>(InterfaceClassIt);
 				if (InterfaceBPGC)
 				{
-					const EReplacementResult InterfaceResult = IsTargetedForReplacement(InterfaceBPGC);
+					const EReplacementResult InterfaceResult = IsTargetedForReplacement(InterfaceBPGC, NativizationOptions);
 					if (InterfaceResult != EReplacementResult::ReplaceCompletely)
 					{
 						if (bNativizeOnlySelectedBPs)
@@ -929,12 +996,54 @@ EReplacementResult FBlueprintNativeCodeGenModule::IsTargetedForReplacement(const
 	};
 	if (ObjectGenratesOnlyStub())
 	{
-		CachedIsTargetedForReplacement.Add(ObjectKey, EReplacementResult::GenerateStub);
+		StateForCurrentPlatform->CachedIsTargetedForReplacement.Add(ObjectKey, EReplacementResult::GenerateStub);
 		return EReplacementResult::GenerateStub;
 	}
 
-	CachedIsTargetedForReplacement.Add(ObjectKey, EReplacementResult::ReplaceCompletely);
+	StateForCurrentPlatform->CachedIsTargetedForReplacement.Add(ObjectKey, EReplacementResult::ReplaceCompletely);
 	return EReplacementResult::ReplaceCompletely;
 }
 
+void FBlueprintNativeCodeGenModule::FillPlatformNativizationDetails(const ITargetPlatform* Platform, FPlatformNativizationDetails& Details)
+{
+	check(Platform);
+	Details.PlatformName = Platform->PlatformName();
+	// If you change this target path you must also update logic in CookCommand.Automation.CS. Passing a single directory around is cumbersome for testing, so I have hard coded it.
+	Details.PlatformTargetDirectory = FString(FPaths::Combine(*FPaths::GameIntermediateDir(), *(Platform->PlatformName())));
+	Details.CompilerNativizationOptions.PlatformName = Platform->PlatformName();
+	Details.CompilerNativizationOptions.ClientOnlyPlatform = Platform->IsClientOnly();
+	Details.CompilerNativizationOptions.ServerOnlyPlatform = Platform->IsServerOnly();
+
+	auto GatherExcludedStuff = [&](const TCHAR* KeyForExcludedModules, const TCHAR* KeyForExcludedPaths, const TCHAR* KeyForExcludedAssets)
+	{
+		const TCHAR* ConfigSection = TEXT("BlueprintNativizationSettings");
+		{
+			TArray<FString> ExcludedModuls;
+			GConfig->GetArray(ConfigSection, KeyForExcludedModules, ExcludedModuls, GEditorIni);
+			for (const FString& NameStr : ExcludedModuls)
+			{
+				Details.CompilerNativizationOptions.ExcludedModules.Add(FName(*NameStr));
+			}
+		}
+		GConfig->GetArray(ConfigSection, KeyForExcludedPaths, Details.CompilerNativizationOptions.ExcludedFolderPaths, GEditorIni);
+
+		{
+			TArray<FString> ExcludedAssetPaths;
+			GConfig->GetArray(ConfigSection, KeyForExcludedAssets, ExcludedAssetPaths, GEditorIni);
+			for (FString& Path : ExcludedAssetPaths)
+			{
+				Details.CompilerNativizationOptions.ExcludedAssets.Add(FStringAssetReference(Path));
+			}
+		}
+	};
+	if (Details.CompilerNativizationOptions.ServerOnlyPlatform)
+	{
+		GatherExcludedStuff(TEXT("ModulsExcludedFromNativizedServer"), TEXT("ExcludedFolderPathsFromServer"), TEXT("ExcludedAssetsFromServer"));
+	}
+	if (Details.CompilerNativizationOptions.ClientOnlyPlatform)
+	{
+		GatherExcludedStuff(TEXT("ModulsExcludedFromNativizedClient"), TEXT("ExcludedFolderPathsFromClient"), TEXT("ExcludedAssetsFromClient"));
+	}
+}
+
 IMPLEMENT_MODULE(FBlueprintNativeCodeGenModule, BlueprintNativeCodeGen);
diff --git a/Engine/Source/Developer/BlueprintNativeCodeGen/Private/BlueprintNativeCodeGenUtils.cpp b/Engine/Source/Developer/BlueprintNativeCodeGen/Private/BlueprintNativeCodeGenUtils.cpp
--- a/Engine/Source/Developer/BlueprintNativeCodeGen/Private/BlueprintNativeCodeGenUtils.cpp
+++ b/Engine/Source/Developer/BlueprintNativeCodeGen/Private/BlueprintNativeCodeGenUtils.cpp
@@ -201,9 +201,21 @@ static bool BlueprintNativeCodeGenUtilsImpl::GenerateModuleBuildFile(const FBlue
 		}
 	}
 
-	TArray<FString> AdditionalPublicDependencyModuleNames;
-	GConfig->GetArray(TEXT("BlueprintNativizationSettings"), TEXT("AdditionalPublicDependencyModuleNames"), AdditionalPublicDependencyModuleNames, GEditorIni);
-	PublicDependencies.Append(AdditionalPublicDependencyModuleNames);
+	auto IndludeAdditionalPublicDependencyModules = [&](const TCHAR* AdditionalPublicDependencyModuleSection)
+	{
+		TArray<FString> AdditionalPublicDependencyModuleNames;
+		GConfig->GetArray(TEXT("BlueprintNativizationSettings"), AdditionalPublicDependencyModuleSection, AdditionalPublicDependencyModuleNames, GEditorIni);
+		PublicDependencies.Append(AdditionalPublicDependencyModuleNames);
+	};
+	IndludeAdditionalPublicDependencyModules(TEXT("AdditionalPublicDependencyModuleNames"));
+	if (Manifest.GetCompilerNativizationOptions().ServerOnlyPlatform) //or !ClientOnlyPlatform ?
+	{
+		IndludeAdditionalPublicDependencyModules(TEXT("AdditionalPublicDependencyModuleNamesServer"));
+	}
+	if (Manifest.GetCompilerNativizationOptions().ClientOnlyPlatform)
+	{
+		IndludeAdditionalPublicDependencyModules(TEXT("AdditionalPublicDependencyModuleNamesClient"));
+	}
 
 	TArray<FString> PrivateDependencies;
 
@@ -215,6 +227,10 @@ static bool BlueprintNativeCodeGenUtilsImpl::GenerateModuleBuildFile(const FBlue
 		const FString PkgModuleName = FPackageName::GetLongPackageAssetName(ModulePkg->GetName());
 		if (ModuleManager.ModuleExists(*PkgModuleName))
 		{
+			if (Manifest.GetCompilerNativizationOptions().ExcludedModules.Contains(FName(*PkgModuleName)))
+			{
+				continue;
+			}
 			if (!PublicDependencies.Contains(PkgModuleName))
 			{
 				PrivateDependencies.Add(PkgModuleName);
@@ -281,7 +297,7 @@ bool FBlueprintNativeCodeGenUtils::FinalizePlugin(const FBlueprintNativeCodeGenM
 }
 
 //------------------------------------------------------------------------------
-void FBlueprintNativeCodeGenUtils::GenerateCppCode(UObject* Obj, TSharedPtr<FString> OutHeaderSource, TSharedPtr<FString> OutCppSource, TSharedPtr<FNativizationSummary> NativizationSummary, FCompilerNativizationOptions NativizationOptions)
+void FBlueprintNativeCodeGenUtils::GenerateCppCode(UObject* Obj, TSharedPtr<FString> OutHeaderSource, TSharedPtr<FString> OutCppSource, TSharedPtr<FNativizationSummary> NativizationSummary, const FCompilerNativizationOptions& NativizationOptions)
 {
 	auto UDEnum = Cast<UUserDefinedEnum>(Obj);
 	auto UDStruct = Cast<UUserDefinedStruct>(Obj);
@@ -370,7 +386,7 @@ void FBlueprintNativeCodeGenUtils::GenerateCppCode(UObject* Obj, TSharedPtr<FStr
 		}
 		else if (UDStruct)
 		{
-			*OutHeaderSource = Compiler.GenerateCppCodeForStruct(UDStruct);
+			*OutHeaderSource = Compiler.GenerateCppCodeForStruct(UDStruct, NativizationOptions);
 		}
 	}
 	else
diff --git a/Engine/Source/Developer/BlueprintNativeCodeGen/Private/BlueprintNativeCodeGenUtils.h b/Engine/Source/Developer/BlueprintNativeCodeGen/Private/BlueprintNativeCodeGenUtils.h
--- a/Engine/Source/Developer/BlueprintNativeCodeGen/Private/BlueprintNativeCodeGenUtils.h
+++ b/Engine/Source/Developer/BlueprintNativeCodeGen/Private/BlueprintNativeCodeGenUtils.h
@@ -35,7 +35,7 @@ public:
 	 * @param  OutHeaderSource	Output destination for the header source.
 	 * @param  OutCppSource		Output destination for the cpp source.
 	 */
-	static void GenerateCppCode(UObject* Obj, TSharedPtr<FString> OutHeaderSource, TSharedPtr<FString> OutCppSource, TSharedPtr<FNativizationSummary> NativizationSummary, FCompilerNativizationOptions NativizationOptions);
+	static void GenerateCppCode(UObject* Obj, TSharedPtr<FString> OutHeaderSource, TSharedPtr<FString> OutCppSource, TSharedPtr<FNativizationSummary> NativizationSummary, const FCompilerNativizationOptions& NativizationOptions);
 
 public: 
 	/** 
diff --git a/Engine/Source/Developer/BlueprintNativeCodeGen/Private/NativeCodeGenerationTool.cpp b/Engine/Source/Developer/BlueprintNativeCodeGen/Private/NativeCodeGenerationTool.cpp
--- a/Engine/Source/Developer/BlueprintNativeCodeGen/Private/NativeCodeGenerationTool.cpp
+++ b/Engine/Source/Developer/BlueprintNativeCodeGen/Private/NativeCodeGenerationTool.cpp
@@ -65,7 +65,8 @@ struct FGeneratedCodeData
 
 	void GatherUserDefinedDependencies(UBlueprint& InBlueprint)
 	{
-		FGatherConvertedClassDependencies ClassDependencies(InBlueprint.GeneratedClass);
+		FCompilerNativizationOptions BlankOptions{};
+		FGatherConvertedClassDependencies ClassDependencies(InBlueprint.GeneratedClass, BlankOptions);
 		for (auto Iter : ClassDependencies.ConvertedClasses)
 		{
 			DependentObjects.Add(Iter);
diff --git a/Engine/Source/Developer/BlueprintNativeCodeGen/Public/BlueprintNativeCodeGenModule.h b/Engine/Source/Developer/BlueprintNativeCodeGen/Public/BlueprintNativeCodeGenModule.h
--- a/Engine/Source/Developer/BlueprintNativeCodeGen/Public/BlueprintNativeCodeGenModule.h
+++ b/Engine/Source/Developer/BlueprintNativeCodeGen/Public/BlueprintNativeCodeGenModule.h
@@ -11,6 +11,7 @@
 
 class UBlueprint;
 enum class ESavePackageResult;
+class ITargetPlatform;
 
 struct FPlatformNativizationDetails
 {
@@ -80,8 +81,10 @@ public:
 	virtual void FinalizeManifest() = 0;
 	virtual void GenerateStubs() = 0;
 	virtual void GenerateFullyConvertedClasses() = 0;
-	virtual void MarkUnconvertedBlueprintAsNecessary(TAssetPtr<UBlueprint> BPPtr) = 0;
+	virtual void MarkUnconvertedBlueprintAsNecessary(TAssetPtr<UBlueprint> BPPtr, const FCompilerNativizationOptions& NativizationOptions) = 0;
 	virtual const TMultiMap<FName, TAssetSubclassOf<UObject>>& GetFunctionsBoundToADelegate() = 0;
+
+	virtual void FillPlatformNativizationDetails(const ITargetPlatform* Platform, FPlatformNativizationDetails& OutDetails) = 0;
 protected:
 	virtual void Initialize(const FNativeCodeGenInitData& InitData) = 0;
 	virtual void InitializeForRerunDebugOnly(const TArray<FPlatformNativizationDetails>& CodegenTargets) = 0;
diff --git a/Engine/Source/Developer/HotReload/Private/HotReload.cpp b/Engine/Source/Developer/HotReload/Private/HotReload.cpp
--- a/Engine/Source/Developer/HotReload/Private/HotReload.cpp
+++ b/Engine/Source/Developer/HotReload/Private/HotReload.cpp
@@ -919,7 +919,7 @@ ECompilationResult::Type FHotReloadModule::DoHotReloadInternal(const TMap<FStrin
 
 			if (UScriptStruct* ScriptStruct = Cast<UScriptStruct>(static_cast<UObject*>(It->Object)))
 			{
-				if (Packages.ContainsByPredicate([=](UPackage* Package) { return ScriptStruct->IsIn(Package); }) && ScriptStruct->GetCppStructOps())
+				if (Packages.ContainsByPredicate([=](UPackage* Package) { return ScriptStruct->IsIn(Package); }) && !ScriptStruct->HasAnyFlags(RF_ClassDefaultObject) && ScriptStruct->GetCppStructOps())
 				{
 					ScriptStructs.Add(ScriptStruct);
 				}
diff --git a/Engine/Source/Developer/HotReload/Private/HotReloadClassReinstancer.cpp b/Engine/Source/Developer/HotReload/Private/HotReloadClassReinstancer.cpp
--- a/Engine/Source/Developer/HotReload/Private/HotReloadClassReinstancer.cpp
+++ b/Engine/Source/Developer/HotReload/Private/HotReloadClassReinstancer.cpp
@@ -104,7 +104,6 @@ void FHotReloadClassReinstancer::SerializeCDOProperties(UObject* InObject, FHotR
 					PropertyInfo.SubobjectName = SubobjectName;
 					PropertyInfo.SerializedValueOffset = Tell();
 					PropertyInfo.SerializedValueSize = Num;
-					PropertyData.Properties.Add(SerializedProperty->GetFName(), PropertyInfo);
 				}
 				else
 				{
diff --git a/Engine/Source/Developer/LauncherServices/Private/Launcher/LauncherUATTask.h b/Engine/Source/Developer/LauncherServices/Private/Launcher/LauncherUATTask.h
--- a/Engine/Source/Developer/LauncherServices/Private/Launcher/LauncherUATTask.h
+++ b/Engine/Source/Developer/LauncherServices/Private/Launcher/LauncherUATTask.h
@@ -73,8 +73,7 @@ protected:
 		{
 			UATCommandLine += FString::Printf(TEXT(" -ue4exe=\"%s\""), *EditorExe);
 
-			bool bRunningDebug = FParse::Param(FCommandLine::Get(), TEXT("debug"));
-			if (bRunningDebug)
+			if (FApp::IsRunningDebug())
 			{
 				UATCommandLine += TEXT(" -UseDebugParamForEditorExe");
 			}
diff --git a/Engine/Source/Developer/LauncherServices/Private/Profiles/LauncherProfile.h b/Engine/Source/Developer/LauncherServices/Private/Profiles/LauncherProfile.h
--- a/Engine/Source/Developer/LauncherServices/Private/Profiles/LauncherProfile.h
+++ b/Engine/Source/Developer/LauncherServices/Private/Profiles/LauncherProfile.h
@@ -1067,7 +1067,7 @@ public:
 		Writer.WriteValue("nocompile", !IsBuildingUAT());
 		Writer.WriteValue("nocompileeditor", FApp::IsEngineInstalled());
 		Writer.WriteValue("ue4exe", GetEditorExe());
-		Writer.WriteValue("usedebugparamforeditorexe", FParse::Param(FCommandLine::Get(), TEXT("debug")));
+		Writer.WriteValue("usedebugparamforeditorexe", FApp::IsRunningDebug());
 		Writer.WriteValue("utf8output", true);
 
 		// client configurations
diff --git a/Engine/Source/Developer/PackageDependencyInfo/Private/PackageDependencyInfo.cpp b/Engine/Source/Developer/PackageDependencyInfo/Private/PackageDependencyInfo.cpp
--- a/Engine/Source/Developer/PackageDependencyInfo/Private/PackageDependencyInfo.cpp
+++ b/Engine/Source/Developer/PackageDependencyInfo/Private/PackageDependencyInfo.cpp
@@ -296,16 +296,22 @@ bool FPackageDependencyInfo::InitializeDependencyInfo( const TCHAR* InPackageNam
 		UE_LOG(LogPackageDependencyInfo, Display, TEXT("\tPackage Info not found for %s!"), InPackageName);
 		return false;
 	}
+	FPackageDependencyTrackingInfo* PkgInfo = *pPkgInfo;
+
+	if (PkgInfo->bValid == false)
+	{
+		UE_LOG(LogPackageDependencyInfo, Display, TEXT("\tPackage Info invalid for %s!"), InPackageName);
+		return false;
+	}
 
 	OutPkgInfo = *pPkgInfo;
 
-	FPackageDependencyTrackingInfo* PkgInfo = *pPkgInfo;
 	// Don't process if it was already processed
 	if (PkgInfo->bInitializedHashes)
 	{
 		return true;
 	}
-
+	
 	TSet<FPackageDependencyTrackingInfo*> PackagesToDetermineDependencies;
 	PackagesToDetermineDependencies.Add( PkgInfo );
 	TSet<FPackageDependencyTrackingInfo*> AllPackages;
@@ -338,7 +344,7 @@ bool FPackageDependencyInfo::InitializeDependencyInfo( const TCHAR* InPackageNam
 	for (FPackageDependencyTrackingInfo* DepPkgInfo : AllPackages)
 	{
 		check(DepPkgInfo->bBeingProcessed == false);
-		check(DepPkgInfo->bInitializedDependencies == true);
+		check(DepPkgInfo->bInitializedDependencies == true || DepPkgInfo->bValid == false);
 	}
 
 	for (FPackageDependencyTrackingInfo* DepPkgInfo : AllPackages)
@@ -359,7 +365,12 @@ bool FPackageDependencyInfo::InitializeDependencyInfo( const TCHAR* InPackageNam
 	for (FPackageDependencyTrackingInfo* DepPkgInfo : AllPackages)
 	{
 		check(DepPkgInfo->bBeingProcessed == false);
-		check(DepPkgInfo->bInitializedHashes == true);
+		check(DepPkgInfo->bInitializedHashes == true || DepPkgInfo->bValid == false);
+	}
+
+	if ( OutPkgInfo->bValid == false )
+	{
+		return false;
 	}
 
 	return true;
@@ -1212,6 +1223,11 @@ void FPackageDependencyInfo::RecursiveResolveDependentHash(FPackageDependencyTra
 			// circular reference
 			return;
 		}
+		if ( PkgInfo->bValid == false )
+		{
+			// package info is invalid so we don't care about it
+			return;
+		}
 		
 		PkgInfo->bBeingProcessed = true;
 
diff --git a/Engine/Source/Developer/Windows/WindowsTargetPlatform/Private/WindowsTargetPlatformModule.cpp b/Engine/Source/Developer/Windows/WindowsTargetPlatform/Private/WindowsTargetPlatformModule.cpp
--- a/Engine/Source/Developer/Windows/WindowsTargetPlatform/Private/WindowsTargetPlatformModule.cpp
+++ b/Engine/Source/Developer/Windows/WindowsTargetPlatform/Private/WindowsTargetPlatformModule.cpp
@@ -69,7 +69,24 @@ public:
 #endif
 
 		TargetSettings = NewObject<UWindowsTargetSettings>(GetTransientPackage(), "WindowsTargetSettings", RF_Standalone);
-		
+
+		FString Compiler;
+		if (GConfig->GetString(TEXT("/Script/WindowsTargetPlatform.WindowsTargetSettings"), TEXT("Compiler"), Compiler, GEngineIni))
+		{
+			if (Compiler == TEXT("VisualStudio2013"))
+			{
+				TargetSettings->Compiler = ECompilerVersion::VisualStudio2013;
+			}
+			else if (Compiler == TEXT("VisualStudio2015"))
+			{
+				TargetSettings->Compiler = ECompilerVersion::VisualStudio2015;
+			}
+			else if (Compiler == TEXT("VisualStudio2017"))
+			{
+				TargetSettings->Compiler = ECompilerVersion::VisualStudio2017;
+			}
+		}
+
 		// We need to manually load the config properties here, as this module is loaded before the UObject system is setup to do this
 		GConfig->GetArray(TEXT("/Script/WindowsTargetPlatform.WindowsTargetSettings"), TEXT("TargetedRHIs"), TargetSettings->TargetedRHIs, GEngineIni);
 
diff --git a/Engine/Source/Editor/ContentBrowser/Private/SAssetView.h b/Engine/Source/Editor/ContentBrowser/Private/SAssetView.h
--- a/Engine/Source/Editor/ContentBrowser/Private/SAssetView.h
+++ b/Engine/Source/Editor/ContentBrowser/Private/SAssetView.h
@@ -19,6 +19,7 @@
 #include "Widgets/Views/STableViewBase.h"
 #include "Widgets/Views/STableRow.h"
 #include "Editor/ContentBrowser/Private/AssetViewSortManager.h"
+#include "AssetViewTypes.h"
 
 class FMenuBuilder;
 class FWeakWidgetPath;
@@ -28,8 +29,6 @@ class SAssetListView;
 class SAssetTileView;
 class SComboButton;
 class UFactory;
-struct FAssetViewAsset;
-struct FAssetViewItem;
 struct FHistoryData;
 struct FPropertyChangedEvent;
 
diff --git a/Engine/Source/Editor/Kismet/Private/FindInBlueprintManager.cpp b/Engine/Source/Editor/Kismet/Private/FindInBlueprintManager.cpp
--- a/Engine/Source/Editor/Kismet/Private/FindInBlueprintManager.cpp
+++ b/Engine/Source/Editor/Kismet/Private/FindInBlueprintManager.cpp
@@ -438,6 +438,14 @@ namespace BlueprintSearchMetaDataHelpers
 
 				return false;
 			}
+
+			friend uint32 GetTypeHash(const FLookupTableItem& InObject)
+			{
+				FString Namespace = FTextInspector::GetNamespace(InObject.Text).Get(TEXT("DefaultNamespace"));
+				FString Key = FTextInspector::GetKey(InObject.Text).Get(TEXT("DefaultKey"));
+				uint32 Hash = HashCombine(GetTypeHash(InObject.Text.ToString()), HashCombine(GetTypeHash(Namespace), GetTypeHash(Key)));
+				return Hash;
+			}
 		};
 
 	protected:
@@ -494,14 +502,6 @@ namespace BlueprintSearchMetaDataHelpers
 		TMap<UStruct*, TArray<FSearchableProperty>> CachedPropertyMapping;
 	};
 
-	static uint32 GetTypeHash(const BlueprintSearchMetaDataHelpers::TJsonFindInBlueprintStringWriter<TCondensedJsonPrintPolicy<TCHAR>>::FLookupTableItem& InObject)
-	{
-		FString Namespace = FTextInspector::GetNamespace(InObject.Text).Get(TEXT("DefaultNamespace"));
-		FString Key = FTextInspector::GetKey(InObject.Text).Get(TEXT("DefaultKey"));
-		uint32 Hash = HashCombine(GetTypeHash(InObject.Text.ToString()), HashCombine(GetTypeHash(Namespace), GetTypeHash(Key)));
-		return Hash;
-	}
-
 	typedef TJsonWriter<TCHAR, TCondensedJsonPrintPolicy<TCHAR>> SearchMetaDataWriterParentClass;
 	typedef TJsonFindInBlueprintStringWriter<TCondensedJsonPrintPolicy<TCHAR>> SearchMetaDataWriter;
 
diff --git a/Engine/Source/Editor/KismetCompiler/Private/KismetCompilerModule.cpp b/Engine/Source/Editor/KismetCompiler/Private/KismetCompilerModule.cpp
--- a/Engine/Source/Editor/KismetCompiler/Private/KismetCompilerModule.cpp
+++ b/Engine/Source/Editor/KismetCompiler/Private/KismetCompilerModule.cpp
@@ -50,8 +50,8 @@ public:
 	virtual TArray<IBlueprintCompiler*>& GetCompilers() override { return Compilers; }
 	virtual void GetBlueprintTypesForClass(UClass* ParentClass, UClass*& OutBlueprintClass, UClass*& OutBlueprintGeneratedClass) const override;
 	virtual void GenerateCppCodeForEnum(UUserDefinedEnum* UDEnum, FString& OutHeaderCode, FString& OutCPPCode) override;
-	virtual FString GenerateCppCodeForStruct(UUserDefinedStruct* UDStruct) override;
-	virtual FString GenerateCppWrapper(UBlueprintGeneratedClass* BPGC) override;
+	virtual FString GenerateCppCodeForStruct(UUserDefinedStruct* UDStruct, const FCompilerNativizationOptions& NativizationOptions) override;
+	virtual FString GenerateCppWrapper(UBlueprintGeneratedClass* BPGC, const FCompilerNativizationOptions& NativizationOptions) override;
 	// End implementation
 private:
 	void CompileBlueprintInner(class UBlueprint* Blueprint, const FKismetCompilerOptions& CompileOptions, FCompilerResultsLog& Results, TSharedPtr<FBlueprintCompileReinstancer> Reinstancer, TArray<UObject*>* ObjLoaded);
@@ -185,16 +185,16 @@ void FKismet2CompilerModule::GenerateCppCodeForEnum(UUserDefinedEnum* UDEnum, FS
 	Backend_CPP->GenerateCodeFromEnum(UDEnum, OutHeaderCode, OutCPPCode);
 }
 
-FString FKismet2CompilerModule::GenerateCppCodeForStruct(UUserDefinedStruct* UDStruct)
+FString FKismet2CompilerModule::GenerateCppCodeForStruct(UUserDefinedStruct* UDStruct, const FCompilerNativizationOptions& NativizationOptions)
 {
 	TUniquePtr<IBlueprintCompilerCppBackend> Backend_CPP(IBlueprintCompilerCppBackendModuleInterface::Get().Create());
-	return Backend_CPP->GenerateCodeFromStruct(UDStruct);
+	return Backend_CPP->GenerateCodeFromStruct(UDStruct, NativizationOptions);
 }
 
-FString FKismet2CompilerModule::GenerateCppWrapper(UBlueprintGeneratedClass* BPGC)
+FString FKismet2CompilerModule::GenerateCppWrapper(UBlueprintGeneratedClass* BPGC, const FCompilerNativizationOptions& NativizationOptions)
 {
 	TUniquePtr<IBlueprintCompilerCppBackend> Backend_CPP(IBlueprintCompilerCppBackendModuleInterface::Get().Create());
-	return Backend_CPP->GenerateWrapperForClass(BPGC);
+	return Backend_CPP->GenerateWrapperForClass(BPGC, NativizationOptions);
 }
 
 extern UNREALED_API FSecondsCounterData BlueprintCompileAndLoadTimerData;
diff --git a/Engine/Source/Editor/KismetCompiler/Public/BlueprintCompilerCppBackendInterface.h b/Engine/Source/Editor/KismetCompiler/Public/BlueprintCompilerCppBackendInterface.h
--- a/Engine/Source/Editor/KismetCompiler/Public/BlueprintCompilerCppBackendInterface.h
+++ b/Engine/Source/Editor/KismetCompiler/Public/BlueprintCompilerCppBackendInterface.h
@@ -14,12 +14,12 @@ class UUserDefinedStruct;
 class IBlueprintCompilerCppBackend
 {
 public:
-	virtual FString GenerateCodeFromClass(UClass* SourceClass, TIndirectArray<FKismetFunctionContext>& Functions, bool bGenerateStubsOnly, FCompilerNativizationOptions NativizationOptions, FString& OutCppBody) = 0;
+	virtual FString GenerateCodeFromClass(UClass* SourceClass, TIndirectArray<FKismetFunctionContext>& Functions, bool bGenerateStubsOnly, const FCompilerNativizationOptions& NativizationOptions, FString& OutCppBody) = 0;
 	virtual void GenerateCodeFromEnum(UUserDefinedEnum* SourceEnum, FString& OutHeaderCode, FString& OutCPPCode) = 0;
-	virtual FString GenerateCodeFromStruct(UUserDefinedStruct* SourceStruct) = 0;
+	virtual FString GenerateCodeFromStruct(UUserDefinedStruct* SourceStruct, const FCompilerNativizationOptions& NativizationOptions) = 0;
 
 	// Generate a wrapper class, that helps accessing non-native properties and calling non-native functions
-	virtual FString GenerateWrapperForClass(UClass* SourceClass) = 0;
+	virtual FString GenerateWrapperForClass(UClass* SourceClass, const FCompilerNativizationOptions& NativizationOptions) = 0;
 
 	virtual ~IBlueprintCompilerCppBackend() {}
 };
diff --git a/Engine/Source/Editor/KismetCompiler/Public/KismetCompilerModule.h b/Engine/Source/Editor/KismetCompiler/Public/KismetCompilerModule.h
--- a/Engine/Source/Editor/KismetCompiler/Public/KismetCompilerModule.h
+++ b/Engine/Source/Editor/KismetCompiler/Public/KismetCompilerModule.h
@@ -12,6 +12,7 @@ class UBlueprintGeneratedClass;
 class UUserDefinedEnum;
 class UUserDefinedStruct;
 struct FKismetCompilerOptions;
+struct FCompilerNativizationOptions;
 
 #define KISMET_COMPILER_MODULENAME "KismetCompiler"
 
@@ -100,9 +101,9 @@ public:
 	virtual void GetBlueprintTypesForClass(UClass* ParentClass, UClass*& OutBlueprintClass, UClass*& OutBlueprintGeneratedClass) const = 0;
 
 	virtual void GenerateCppCodeForEnum(UUserDefinedEnum* UDEnum, FString& OutHeaderCode, FString& OutCPPCode) = 0;
-	virtual FString GenerateCppCodeForStruct(UUserDefinedStruct* UDStruct) = 0;
+	virtual FString GenerateCppCodeForStruct(UUserDefinedStruct* UDStruct, const FCompilerNativizationOptions& NativizationOptions) = 0;
 	// Generate a wrapper class, that helps accessing non-native properties and calling non-native functions
-	virtual FString GenerateCppWrapper(UBlueprintGeneratedClass* BPGC) = 0;
+	virtual FString GenerateCppWrapper(UBlueprintGeneratedClass* BPGC, const FCompilerNativizationOptions& NativizationOptions) = 0;
 };
 
 
diff --git a/Engine/Source/Editor/MainFrame/Private/Frame/MainFrameActions.cpp b/Engine/Source/Editor/MainFrame/Private/Frame/MainFrameActions.cpp
--- a/Engine/Source/Editor/MainFrame/Private/Frame/MainFrameActions.cpp
+++ b/Engine/Source/Editor/MainFrame/Private/Frame/MainFrameActions.cpp
@@ -391,9 +391,7 @@ void FMainFrameActionCallbacks::CookContent(const FName InPlatformInfoName)
 
 	OptionalParams += GetCookingOptionalParams();
 
-	const bool bRunningDebug = FParse::Param(FCommandLine::Get(), TEXT("debug"));
-
-	if (bRunningDebug)
+	if (FApp::IsRunningDebug())
 	{
 		OptionalParams += TEXT(" -UseDebugParamForEditorExe");
 	}
@@ -683,9 +681,7 @@ void FMainFrameActionCallbacks::PackageProject( const FName InPlatformInfoName )
 		OptionalParams += FString::Printf(TEXT(" -NumCookersToSpawn=%d"), NumCookers); 
 	}
 
-	const bool bRunningDebug = FParse::Param(FCommandLine::Get(), TEXT("debug"));
-
-	if (bRunningDebug)
+	if (FApp::IsRunningDebug())
 	{
 		OptionalParams += TEXT(" -UseDebugParamForEditorExe");
 	}
diff --git a/Engine/Source/Editor/PropertyEditor/Public/PropertyEditorModule.h b/Engine/Source/Editor/PropertyEditor/Public/PropertyEditorModule.h
--- a/Engine/Source/Editor/PropertyEditor/Public/PropertyEditorModule.h
+++ b/Engine/Source/Editor/PropertyEditor/Public/PropertyEditorModule.h
@@ -7,13 +7,13 @@
 #include "Modules/ModuleInterface.h"
 #include "UObject/StructOnScope.h"
 #include "Toolkits/IToolkitHost.h"
+#include "IDetailsView.h"
 
 #include "PropertyEditorDelegates.h"
 #include "IPropertyTypeCustomization.h"
 
 class FAssetEditorToolkit;
 class FNotifyHook;
-class IDetailsView;
 class IPropertyHandle;
 class IPropertyTable;
 class IPropertyTableCellPresenter;
diff --git a/Engine/Source/Editor/UnrealEd/Private/Commandlets/ContentCommandlets.cpp b/Engine/Source/Editor/UnrealEd/Private/Commandlets/ContentCommandlets.cpp
--- a/Engine/Source/Editor/UnrealEd/Private/Commandlets/ContentCommandlets.cpp
+++ b/Engine/Source/Editor/UnrealEd/Private/Commandlets/ContentCommandlets.cpp
@@ -1106,10 +1106,15 @@ void UResavePackagesCommandlet::PerformAdditionalOperations(class UWorld* World,
 
 						if (FPackageName::TryConvertLongPackageNameToFilename(MapBuildDataPackageName, MapBuildDataFilename, FPackageName::GetAssetPackageExtension()))
 						{
-							SavePackageHelper(MapBuildDataPackage, MapBuildDataFilename);
 							if (IFileManager::Get().FileExists(*MapBuildDataFilename))
 							{
 								CheckoutFile(MapBuildDataFilename, true);
+								SavePackageHelper(MapBuildDataPackage, MapBuildDataFilename); 
+							}
+							else
+							{
+								SavePackageHelper(MapBuildDataPackage, MapBuildDataFilename);
+								CheckoutFile(MapBuildDataFilename, true);
 							}
 						}
 					}
diff --git a/Engine/Source/Editor/UnrealEd/Private/CookOnTheFlyServer.cpp b/Engine/Source/Editor/UnrealEd/Private/CookOnTheFlyServer.cpp
--- a/Engine/Source/Editor/UnrealEd/Private/CookOnTheFlyServer.cpp
+++ b/Engine/Source/Editor/UnrealEd/Private/CookOnTheFlyServer.cpp
@@ -1100,20 +1100,20 @@ void UCookOnTheFlyServer::GetDependentPackages( const TSet<FName>& RootPackages,
 		{
 			// check(PackageDependency.ToString().StartsWith(TEXT("/")));
 			FName PackageDependency = OriginalPackageDependency;
-			FString PackageDepdencyString = PackageDependency.ToString();
+			FString PackageDependencyString = PackageDependency.ToString();
 
 			FText OutReason;
 			const bool bIncludeReadOnlyRoots = true; // Dependency packages are often script packages (read-only)
-			if (!FPackageName::IsValidLongPackageName(PackageDepdencyString, bIncludeReadOnlyRoots, &OutReason))
+			if (!FPackageName::IsValidLongPackageName(PackageDependencyString, bIncludeReadOnlyRoots, &OutReason))
 			{
 				const FText FailMessage = FText::Format(LOCTEXT("UnableToGeneratePackageName", "Unable to generate long package name for {0}. {1}"),
-					FText::FromString(PackageDepdencyString), OutReason);
+					FText::FromString(PackageDependencyString), OutReason);
 
 				LogCookerMessage(FailMessage.ToString(), EMessageSeverity::Warning);
 				UE_LOG(LogCook, Warning, TEXT("%s"), *( FailMessage.ToString() ));
 				continue;
 			}
-			else if (FPackageName::IsScriptPackage(PackageDepdencyString))
+			else if (FPackageName::IsScriptPackage(PackageDependencyString) || FPackageName::IsMemoryPackage(PackageDependencyString))
 			{
 				continue;
 			}
@@ -1209,7 +1209,7 @@ void UCookOnTheFlyServer::OnStringAssetReferenceLoadedPackage(const FName& Packa
 		}
 		else
 		{
-			if (!FPackageName::IsScriptPackage(PackageFName.ToString()))
+			if (!FPackageName::IsScriptPackage(PackageFName.ToString()) && !FPackageName::IsMemoryPackage(PackageFName.ToString()))
 			{
 				UE_LOG(LogCook, Warning, TEXT("Unable to find cached package name for package %s"), *PackageFName.ToString());
 			}
@@ -2120,31 +2120,31 @@ uint32 UCookOnTheFlyServer::TickCookOnTheSide( const float TimeSlice, uint32 &Co
 		
 	OUTPUT_TIMERS();
 
-	if ( CookByTheBookOptions )
-		{
+	if (CookByTheBookOptions)
+	{
 		CookByTheBookOptions->CookTime += Timer.GetTimeTillNow();
 	}
 
-	if ( IsCookByTheBookRunning() && 
-		( CookRequests.HasItems() == false ) &&
-		(!(Result & COSR_WaitingOnChildCookers)) )
-			{
+	if (IsCookByTheBookRunning() &&
+		(CookRequests.HasItems() == false) &&
+		(!(Result & COSR_WaitingOnChildCookers)))
+	{
 		check(IsCookByTheBookMode());
 		// UE_LOG(LogCook, Display, TEXT("String asset reference resolve tried %d did %d"), bTriedToRunStringAssetReferenceResolve, bHasRunStringAssetReferenceResolve);
 		// if we are out of stuff and we are in cook by the book from the editor mode then we finish up
 		CookByTheBookFinished();
-		}
+	}
 
 	return Result;
-	}
+}
 
 
-void UCookOnTheFlyServer::PostLoadPackageFixup( UPackage* Package )
-	{
+void UCookOnTheFlyServer::PostLoadPackageFixup(UPackage* Package)
+{
 	const FName StandardPackageName = GetCachedStandardPackageFileFName(Package);
 	const TArray<FName> EmptyTargets;
-	if ( CookedPackages.Exists(StandardPackageName, EmptyTargets) )
-		{
+	if (CookedPackages.Exists(StandardPackageName, EmptyTargets))
+	{
 		return;
 	}
 
@@ -2173,21 +2173,27 @@ void UCookOnTheFlyServer::PostLoadPackageFixup( UPackage* Package )
 
 			// Collect world composition tile packages to cook
 			if (World->WorldComposition)
-	{
+			{
 				World->WorldComposition->CollectTilesToCook(NewPackagesToCook);
-	}
+			}
 
 			for (auto PackageName : NewPackagesToCook)
 			{
+#if 0 // old way keeping for reference in case there are issues
 				FString Filename;
 				if (FPackageName::DoesPackageExist(PackageName, NULL, &Filename))
 				{
 					FString StandardFilename = FPaths::ConvertRelativePathToFull(Filename);
 					FPaths::MakeStandardFilename(StandardFilename);
 					FName StandardPackageFName = FName(*StandardFilename);
+#else
+				FName StandardPackageFName = GetCachedStandardPackageFileFName(FName(*PackageName));
+				if (StandardPackageFName != NAME_None)
+				{
+#endif
 					if (IsChildCooker())
-	{
-		check(IsCookByTheBookMode());
+					{
+						check(IsCookByTheBookMode());
 						// notify the main cooker that it should make sure this package gets cooked
 						CookByTheBookOptions->ChildUnsolicitedPackages.Add(StandardPackageFName);
 					}
@@ -2197,7 +2203,7 @@ void UCookOnTheFlyServer::PostLoadPackageFixup( UPackage* Package )
 					}
 				}
 			}
-	}
+		}
 
 	}
 }
@@ -4031,7 +4037,7 @@ bool UCookOnTheFlyServer::SaveCookedAssetRegistry(const FString& InCookedAssetRe
 		FName PackageName;
 		FString PlatformSandboxPath;
 		//FString PlatformSandboxPath = SandboxPath.Replace(TEXT("[Platform]"), *Platform->PlatformName());
-		if (FPackageName::IsScriptPackage(RelativePath))
+		if (FPackageName::IsScriptPackage(RelativePath) || FPackageName::IsMemoryPackage(RelativePath))
 		{
 			PackageName = Package;
 		}
@@ -4039,7 +4045,7 @@ bool UCookOnTheFlyServer::SaveCookedAssetRegistry(const FString& InCookedAssetRe
 		{
 			FString StringPackageName;
 			if (FPackageName::TryConvertFilenameToLongPackageName(RelativePath, StringPackageName) == false)
-		{
+			{
 				UE_LOG(LogCook, Warning, TEXT("Unable to determine extension for file %s skipping asset registry entry"), *PackageName.ToString());
 				continue;
 			}
@@ -4049,10 +4055,10 @@ bool UCookOnTheFlyServer::SaveCookedAssetRegistry(const FString& InCookedAssetRe
 			PlatformSandboxPath.ReplaceInline(TEXT("[Platform]"), *PlatformName);
 		}
 
-		if ( UncookedEditorPackages.Contains( PackageName ) )
-			{
-				continue;
-			}
+		if (UncookedEditorPackages.Contains(PackageName))
+		{
+			continue;
+		}
 
 		Json->WriteObjectStart(); // unnamed package start
 
@@ -4200,7 +4206,7 @@ bool UCookOnTheFlyServer::SaveCookedAssetRegistry(const FString& InCookedAssetRe
 		FName PackageName;
 		FString PlatformSandboxPath;
 		//FString PlatformSandboxPath = SandboxPath.Replace(TEXT("[Platform]"), *Platform->PlatformName());
-		if (FPackageName::IsScriptPackage(RelativePath))
+		if (FPackageName::IsScriptPackage(RelativePath) || FPackageName::IsMemoryPackage(RelativePath))
 		{
 			PackageName = EditorOnlyPackage;
 		}
@@ -4826,7 +4832,7 @@ void UCookOnTheFlyServer::VerifyCookedPackagesAreUptodate(const TArray<ITargetPl
 			{
 				FName DependencyPackageName = Dependencies[I];
 
-				if (FPackageName::IsScriptPackage(DependencyPackageName.ToString()))
+				if (FPackageName::IsScriptPackage(DependencyPackageName.ToString()) || FPackageName::IsMemoryPackage(DependencyPackageName.ToString()))
 				{
 					// don't care about script packages we are assuming the scripts can't be invalidated after the cooker has started
 					continue;
@@ -5152,7 +5158,7 @@ void UCookOnTheFlyServer::GenerateLongPackageNames(TArray<FName>& FilesInPath)
 
 void UCookOnTheFlyServer::AddFileToCook( TArray<FName>& InOutFilesToCook, const FString &InFilename ) const
 { 
-	if (!FPackageName::IsScriptPackage(InFilename))
+	if (!FPackageName::IsScriptPackage(InFilename) && !FPackageName::IsMemoryPackage(InFilename))
 	{
 		FName InFilenameName = FName(*InFilename );
 		if ( InFilenameName == NAME_None)
@@ -6087,6 +6093,7 @@ void UCookOnTheFlyServer::StartCookByTheBook( const FCookByTheBookStartupOptions
 	CookByTheBookOptions->ChildCookFilename = CookByTheBookStartupOptions.ChildCookFileName;
 	CookByTheBookOptions->bDisableUnsolicitedPackages = !!(CookOptions & ECookByTheBookOptions::DisableUnsolicitedPackages);
 	CookByTheBookOptions->ChildCookIdentifier = CookByTheBookStartupOptions.ChildCookIdentifier;
+	CookByTheBookOptions->bErrorOnEngineContentUse = CookByTheBookStartupOptions.bErrorOnEngineContentUse;
 
 	GenerateAssetRegistry();
 
@@ -6171,12 +6178,7 @@ void UCookOnTheFlyServer::StartCookByTheBook( const FCookByTheBookStartupOptions
 		for (const ITargetPlatform* Entry : CookByTheBookStartupOptions.TargetPlatforms)
 		{
 			FPlatformNativizationDetails PlatformNativizationDetails;
-			PlatformNativizationDetails.PlatformName = Entry->PlatformName();
-			// If you change this target path you must also update logic in CookCommand.Automation.CS. Passing a single directory around is cumbersome for testing, so I have hard coded it.
-			PlatformNativizationDetails.PlatformTargetDirectory = FString(FPaths::Combine(*FPaths::GameIntermediateDir(), *(Entry->PlatformName())));
-			PlatformNativizationDetails.CompilerNativizationOptions.ClientOnlyPlatform = Entry->IsClientOnly();
-			PlatformNativizationDetails.CompilerNativizationOptions.ServerOnlyPlatform = Entry->IsServerOnly();
-
+			IBlueprintNativeCodeGenModule::Get().FillPlatformNativizationDetails(Entry, PlatformNativizationDetails);
 			CodeGenData.CodegenTargets.Push(PlatformNativizationDetails);
 		}
 		CodeGenData.ManifestIdentifier = CookByTheBookStartupOptions.ChildCookIdentifier;
@@ -6546,7 +6548,7 @@ void UCookOnTheFlyServer::StartChildCookers(int32 NumCookersToSpawn, const TArra
 
 		for (const auto& Dependency : UnfilteredDependencies)
 		{
-			if (FPackageName::IsScriptPackage(Dependency.ToString()) == false)
+			if ((FPackageName::IsScriptPackage(Dependency.ToString()) == false) && (FPackageName::IsMemoryPackage(Dependency.ToString()) == false))
 			{
 				if (PackageNamesSet.Contains(Dependency) == false)
 				{
@@ -6571,7 +6573,7 @@ void UCookOnTheFlyServer::StartChildCookers(int32 NumCookersToSpawn, const TArra
 			UE_LOG(LogCook, Warning, TEXT("%s"), *(FailMessage.ToString()));
 			continue;
 		}
-		else if (FPackageName::IsScriptPackage(LongPackageName))
+		else if (FPackageName::IsScriptPackage(LongPackageName)|| FPackageName::IsMemoryPackage(LongPackageName))
 		{
 			continue;
 		}
diff --git a/Engine/Source/Editor/UnrealEd/Private/HierarchicalLOD.cpp b/Engine/Source/Editor/UnrealEd/Private/HierarchicalLOD.cpp
--- a/Engine/Source/Editor/UnrealEd/Private/HierarchicalLOD.cpp
+++ b/Engine/Source/Editor/UnrealEd/Private/HierarchicalLOD.cpp
@@ -34,13 +34,11 @@ FHierarchicalLODBuilder::FHierarchicalLODBuilder(UWorld* InWorld)
 :	World(InWorld)
 {}
 
-#if WITH_HOT_RELOAD_CTORS
 FHierarchicalLODBuilder::FHierarchicalLODBuilder()
 	: World(nullptr)
 {
 	EnsureRetrievingVTablePtrDuringCtor(TEXT("FHierarchicalLODBuilder()"));
 }
-#endif // WITH_HOT_RELOAD_CTORS
 
 void FHierarchicalLODBuilder::Build()
 {
diff --git a/Engine/Source/Editor/UnrealEd/Private/PlayLevel.cpp b/Engine/Source/Editor/UnrealEd/Private/PlayLevel.cpp
--- a/Engine/Source/Editor/UnrealEd/Private/PlayLevel.cpp
+++ b/Engine/Source/Editor/UnrealEd/Private/PlayLevel.cpp
@@ -1415,8 +1415,7 @@ void UEditorEngine::PlayStandaloneLocalPc(FString MapNameOverride, FIntPoint* Wi
 	}
 
 	FString AdditionalParameters(TEXT(" -messaging -SessionName=\"Play in Standalone Game\""));
-	bool bRunningDebug = FParse::Param(FCommandLine::Get(), TEXT("debug"));
-	if (bRunningDebug)
+	if (FApp::IsRunningDebug())
 	{
 		AdditionalParameters += TEXT(" -debug");
 	}
diff --git a/Engine/Source/Editor/UnrealEd/Public/HierarchicalLOD.h b/Engine/Source/Editor/UnrealEd/Public/HierarchicalLOD.h
--- a/Engine/Source/Editor/UnrealEd/Public/HierarchicalLOD.h
+++ b/Engine/Source/Editor/UnrealEd/Public/HierarchicalLOD.h
@@ -30,11 +30,9 @@ struct UNREALED_API FHierarchicalLODBuilder
 {
 	FHierarchicalLODBuilder(UWorld* InWorld);
 
-#if WITH_HOT_RELOAD_CTORS
 	/** DO NOT USE. This constructor is for internal usage only for hot-reload purposes. */
 	FHierarchicalLODBuilder();
-#endif // WITH_HOT_RELOAD_CTORS
-		
+
 	/**
 	* Build, Builds the clusters and spawn LODActors with their merged Static Meshes
 	*/
diff --git a/Engine/Source/Programs/AutomationTool/AutomationUtils/Automation.cs b/Engine/Source/Programs/AutomationTool/AutomationUtils/Automation.cs
--- a/Engine/Source/Programs/AutomationTool/AutomationUtils/Automation.cs
+++ b/Engine/Source/Programs/AutomationTool/AutomationUtils/Automation.cs
@@ -466,7 +466,7 @@ AutomationTool.exe [-verbose] [-compileonly] [-p4] Command0 [-Arg0 -Arg1 -Arg2 
 			}
 
 			// Setup environment
-			Log.TraceInformation("Setting up command environment.");
+			Log.TraceLog("Setting up command environment.");
 			CommandUtils.InitCommandEnvironment();
 
 			// Determine if the engine is installed
@@ -527,7 +527,7 @@ AutomationTool.exe [-verbose] [-compileonly] [-p4] Command0 [-Arg0 -Arg1 -Arg2 
 			CommandUtils.InitP4Support(CommandsToExecute, Compiler.Commands);
 			if (CommandUtils.P4Enabled)
 			{
-				Log.TraceInformation("Setting up Perforce environment.");
+				Log.TraceLog("Setting up Perforce environment.");
 				CommandUtils.InitP4Environment();
 				CommandUtils.InitDefaultP4Connection();
 			}
diff --git a/Engine/Source/Programs/AutomationTool/AutomationUtils/AutomationException.cs b/Engine/Source/Programs/AutomationTool/AutomationUtils/AutomationException.cs
--- a/Engine/Source/Programs/AutomationTool/AutomationUtils/AutomationException.cs
+++ b/Engine/Source/Programs/AutomationTool/AutomationUtils/AutomationException.cs
@@ -67,39 +67,20 @@ namespace AutomationTool
     public class AutomationException : System.Exception
 	{
         public ExitCode ErrorCode = ExitCode.Error_Unknown;
-        public string LogFileName = "";
-        public int ErrorNumber;
-
-		public AutomationException()
-		{
-		}
 
 		public AutomationException(string Msg)
 			:base(Msg)
 		{
 		}
 
-        public AutomationException(string LogFilename, int ErrorNumber, string Format, params object[] Args)
-            :base(string.Format(Format,Args))
-        {
-            this.LogFileName = LogFilename;
-            this.ErrorNumber = ErrorNumber;
-        }
-
         public AutomationException(ExitCode ErrorCode, string Msg)
             : base(Msg)
         {
             this.ErrorCode = ErrorCode;
         }
 
-        [Obsolete]
-        public AutomationException(string Msg, System.Exception InnerException)
-            : base(Msg, InnerException)
-        {
-        }
-
         public AutomationException(Exception InnerException, string Format, params object[] Args)
-			:base(string.Format(Format, Args), InnerException)
+			: base(string.Format(Format, Args), InnerException)
 		{
 		}
 
@@ -110,13 +91,14 @@ namespace AutomationTool
         }
 
         public AutomationException(string Format, params object[] Args)
-			: base(string.Format(Format, Args)) { }
+			: base(string.Format(Format, Args))
+		{
+		}
 
         public AutomationException(ExitCode ErrorCode, string Format, params object[] Args)
             : base(string.Format(Format, Args)) 
         {
             this.ErrorCode = ErrorCode;
         }
-
     }
 }
diff --git a/Engine/Source/Programs/AutomationTool/AutomationUtils/BuildEnvironment.cs b/Engine/Source/Programs/AutomationTool/AutomationUtils/BuildEnvironment.cs
--- a/Engine/Source/Programs/AutomationTool/AutomationUtils/BuildEnvironment.cs
+++ b/Engine/Source/Programs/AutomationTool/AutomationUtils/BuildEnvironment.cs
@@ -33,7 +33,6 @@ namespace AutomationTool
 		static public readonly string Client = "uebp_CLIENT";
 		static public readonly string BuildRootP4 = "uebp_BuildRoot_P4";
 		static public readonly string BuildRootEscaped = "uebp_BuildRoot_Escaped";		
-		static public readonly string LabelToSync = "uebp_LabelToSync";
 		static public readonly string P4Password = "uebp_PASS";
 	}
 
@@ -79,7 +78,7 @@ namespace AutomationTool
 			{
 				UATExe = Assembly.GetEntryAssembly().GetOriginalLocation();
 			}
-			if (!CommandUtils.FileExists_NoExceptions(UATExe))
+			if (!CommandUtils.FileExists(UATExe))
 			{
 				throw new AutomationException("Could not find AutomationTool.exe. Reflection indicated it was here: {0}", UATExe);
 			}
diff --git a/Engine/Source/Programs/AutomationTool/AutomationUtils/CommandUtils.cs b/Engine/Source/Programs/AutomationTool/AutomationUtils/CommandUtils.cs
--- a/Engine/Source/Programs/AutomationTool/AutomationUtils/CommandUtils.cs
+++ b/Engine/Source/Programs/AutomationTool/AutomationUtils/CommandUtils.cs
@@ -323,22 +323,21 @@ namespace AutomationTool
 		/// <param name="Recursive">Recursive search</param>
 		/// <param name="Paths">Paths to search</param>
 		/// <returns>An array of files found in the specified paths</returns>
-		public static string[] FindFiles(string SearchPattern, bool Recursive, params string[] Paths)
+		public static string[] FindFiles(string SearchPattern, bool Recursive, string PathToSearch)
 		{
 			List<string> FoundFiles = new List<string>();
-			foreach (var PathToSearch in Paths)
+
+			var NormalizedPath = ConvertSeparators(PathSeparator.Default, PathToSearch);
+			if (DirectoryExists(NormalizedPath))
 			{
-				var NormalizedPath = ConvertSeparators(PathSeparator.Default, PathToSearch);
-				if (DirectoryExists(NormalizedPath))
+				var FoundInPath = InternalUtils.SafeFindFiles(NormalizedPath, SearchPattern, Recursive);
+				if (FoundInPath == null)
 				{
-					var FoundInPath = InternalUtils.SafeFindFiles(NormalizedPath, SearchPattern, Recursive);
-					if (FoundInPath == null)
-					{
-						throw new AutomationException(String.Format("Failed to find files in '{0}'", NormalizedPath));
-					}
-					FoundFiles.AddRange(FoundInPath);
+					throw new AutomationException(String.Format("Failed to find files in '{0}'", NormalizedPath));
 				}
+				FoundFiles.AddRange(FoundInPath);
 			}
+
 			return FoundFiles.ToArray();
 		}
 
@@ -349,21 +348,20 @@ namespace AutomationTool
 		/// <param name="Recursive">Recursive search</param>
 		/// <param name="Paths">Paths to search</param>
 		/// <returns>An array of files found in the specified paths</returns>
-		public static string[] FindFiles_NoExceptions(string SearchPattern, bool Recursive, params string[] Paths)
+		public static string[] FindFiles_NoExceptions(string SearchPattern, bool Recursive, string PathToSearch)
 		{
 			List<string> FoundFiles = new List<string>();
-			foreach (var PathToSearch in Paths)
+
+			var NormalizedPath = ConvertSeparators(PathSeparator.Default, PathToSearch);
+			if (DirectoryExists(NormalizedPath))
 			{
-				var NormalizedPath = ConvertSeparators(PathSeparator.Default, PathToSearch);
-				if (DirectoryExists(NormalizedPath))
+				var FoundInPath = InternalUtils.SafeFindFiles(NormalizedPath, SearchPattern, Recursive);
+				if (FoundInPath != null)
 				{
-					var FoundInPath = InternalUtils.SafeFindFiles(NormalizedPath, SearchPattern, Recursive);
-					if (FoundInPath != null)
-					{
-						FoundFiles.AddRange(FoundInPath);
-					}
+					FoundFiles.AddRange(FoundInPath);
 				}
 			}
+
 			return FoundFiles.ToArray();
 		}
         /// <summary>
@@ -374,22 +372,21 @@ namespace AutomationTool
         /// <param name="Recursive">Recursive search</param>
         /// <param name="Paths">Paths to search</param>
         /// <returns>An array of files found in the specified paths</returns>
-        public static string[] FindFiles_NoExceptions(bool bQuiet, string SearchPattern, bool Recursive, params string[] Paths)
+        public static string[] FindFiles_NoExceptions(bool bQuiet, string SearchPattern, bool Recursive, string PathToSearch)
         {
             List<string> FoundFiles = new List<string>();
-            foreach (var PathToSearch in Paths)
+
+			var NormalizedPath = ConvertSeparators(PathSeparator.Default, PathToSearch);
+            if (DirectoryExists(NormalizedPath))
             {
-                var NormalizedPath = ConvertSeparators(PathSeparator.Default, PathToSearch);
-                if (DirectoryExists(NormalizedPath))
+                var FoundInPath = InternalUtils.SafeFindFiles(NormalizedPath, SearchPattern, Recursive, bQuiet);
+                if (FoundInPath != null)
                 {
-                    var FoundInPath = InternalUtils.SafeFindFiles(NormalizedPath, SearchPattern, Recursive, bQuiet);
-                    if (FoundInPath != null)
-                    {
-                        FoundFiles.AddRange(FoundInPath);
-                    }
+                    FoundFiles.AddRange(FoundInPath);
                 }
             }
-            return FoundFiles.ToArray();
+
+			return FoundFiles.ToArray();
         }
 		/// <summary>
 		/// Finds files in specified paths. 
@@ -399,22 +396,21 @@ namespace AutomationTool
 		/// <param name="Recursive">Recursive search</param>
 		/// <param name="Paths">Paths to search</param>
 		/// <returns>An array of files found in the specified paths</returns>
-		public static string[] FindDirectories(bool bQuiet, string SearchPattern, bool Recursive, params string[] Paths)
+		public static string[] FindDirectories(bool bQuiet, string SearchPattern, bool Recursive, string PathToSearch)
 		{
 			List<string> FoundDirs = new List<string>();
-			foreach (var PathToSearch in Paths)
+
+			var NormalizedPath = ConvertSeparators(PathSeparator.Default, PathToSearch);
+			if (DirectoryExists(NormalizedPath))
 			{
-				var NormalizedPath = ConvertSeparators(PathSeparator.Default, PathToSearch);
-				if (DirectoryExists(NormalizedPath))
+				var FoundInPath = InternalUtils.SafeFindDirectories(NormalizedPath, SearchPattern, Recursive, bQuiet);
+				if (FoundInPath == null)
 				{
-					var FoundInPath = InternalUtils.SafeFindDirectories(NormalizedPath, SearchPattern, Recursive, bQuiet);
-					if (FoundInPath == null)
-					{
-						throw new AutomationException(String.Format("Failed to find directories in '{0}'", NormalizedPath));
-					}
-					FoundDirs.AddRange(FoundInPath);
+					throw new AutomationException(String.Format("Failed to find directories in '{0}'", NormalizedPath));
 				}
+				FoundDirs.AddRange(FoundInPath);
 			}
+
 			return FoundDirs.ToArray();
 		}
 		/// <summary>
@@ -425,21 +421,20 @@ namespace AutomationTool
 		/// <param name="Recursive">Recursive search</param>
 		/// <param name="Paths">Paths to search</param>
 		/// <returns>An array of files found in the specified paths</returns>
-		public static string[] FindDirectories_NoExceptions(bool bQuiet, string SearchPattern, bool Recursive, params string[] Paths)
+		public static string[] FindDirectories_NoExceptions(bool bQuiet, string SearchPattern, bool Recursive, string PathToSearch)
 		{
 			List<string> FoundDirs = new List<string>();
-			foreach (var PathToSearch in Paths)
+
+			var NormalizedPath = ConvertSeparators(PathSeparator.Default, PathToSearch);
+			if (DirectoryExists(NormalizedPath))
 			{
-				var NormalizedPath = ConvertSeparators(PathSeparator.Default, PathToSearch);
-				if (DirectoryExists(NormalizedPath))
+				var FoundInPath = InternalUtils.SafeFindDirectories(NormalizedPath, SearchPattern, Recursive, bQuiet);
+				if (FoundInPath != null)
 				{
-					var FoundInPath = InternalUtils.SafeFindDirectories(NormalizedPath, SearchPattern, Recursive, bQuiet);
-					if (FoundInPath != null)
-					{
-						FoundDirs.AddRange(FoundInPath);
-					}
+					FoundDirs.AddRange(FoundInPath);
 				}
 			}
+
 			return FoundDirs.ToArray();
 		}
 		/// <summary>
@@ -448,15 +443,12 @@ namespace AutomationTool
 		/// If the deletion of the file fails, this function throws an Exception.
 		/// </summary>
 		/// <param name="Filenames">Filename</param>
-		public static void DeleteFile(params string[] Filenames)
+		public static void DeleteFile(string FileName)
 		{
-			foreach (var Filename in Filenames)
+			var NormalizedFilename = ConvertSeparators(PathSeparator.Default, FileName);
+			if (!InternalUtils.SafeDeleteFile(NormalizedFilename))
 			{
-				var NormalizedFilename = ConvertSeparators(PathSeparator.Default, Filename);
-				if (!InternalUtils.SafeDeleteFile(NormalizedFilename))
-				{
-					throw new AutomationException(String.Format("Failed to delete file '{0}'", NormalizedFilename));
-				}
+				throw new AutomationException(String.Format("Failed to delete file '{0}'", NormalizedFilename));
 			}
 		}
         /// <summary>
@@ -466,15 +458,12 @@ namespace AutomationTool
         /// </summary>
         /// <param name="bQuiet">When true, logging is suppressed.</param>
         /// <param name="Filenames">Filename</param>
-        public static void DeleteFile(bool bQuiet, params string[] Filenames)
+        public static void DeleteFile(bool bQuiet, string FileName)
         {
-            foreach (var Filename in Filenames)
+            var NormalizedFilename = ConvertSeparators(PathSeparator.Default, FileName);
+            if (!InternalUtils.SafeDeleteFile(NormalizedFilename, bQuiet))
             {
-                var NormalizedFilename = ConvertSeparators(PathSeparator.Default, Filename);
-                if (!InternalUtils.SafeDeleteFile(NormalizedFilename, bQuiet))
-                {
-                    throw new AutomationException(String.Format("Failed to delete file '{0}'", NormalizedFilename));
-                }
+                throw new AutomationException(String.Format("Failed to delete file '{0}'", NormalizedFilename));
             }
         }
 
@@ -483,17 +472,15 @@ namespace AutomationTool
 		/// If the deletion of the file fails, prints a warning.
 		/// </summary>
 		/// <param name="Filenames">Filename</param>
-        public static bool DeleteFile_NoExceptions(params string[] Filenames)
+        public static bool DeleteFile_NoExceptions(string FileName)
 		{
 			bool Result = true;
-			foreach (var Filename in Filenames)
+
+			var NormalizedFilename = ConvertSeparators(PathSeparator.Default, FileName);
+			if (!InternalUtils.SafeDeleteFile(NormalizedFilename))
 			{
-				var NormalizedFilename = ConvertSeparators(PathSeparator.Default, Filename);
-				if (!InternalUtils.SafeDeleteFile(NormalizedFilename))
-				{
-					LogWarning("Failed to delete file '{0}'", NormalizedFilename);
-					Result = false;
-				}
+				LogWarning("Failed to delete file '{0}'", NormalizedFilename);
+				Result = false;
 			}
 			return Result;
 		}
@@ -521,15 +508,12 @@ namespace AutomationTool
 		/// </summary>
 		/// <param name="bQuiet">Suppresses log output if true</param>
 		/// <param name="Directories">Directories</param>
-		public static void DeleteDirectory(bool bQuiet, params string[] Directories)
+		public static void DeleteDirectory(bool bQuiet, string Directory)
 		{
-			foreach (var Directory in Directories)
+			var NormalizedDirectory = ConvertSeparators(PathSeparator.Default, Directory);
+			if (!InternalUtils.SafeDeleteDirectory(NormalizedDirectory, bQuiet))
 			{
-				var NormalizedDirectory = ConvertSeparators(PathSeparator.Default, Directory);
-				if (!InternalUtils.SafeDeleteDirectory(NormalizedDirectory, bQuiet))
-				{
-					throw new AutomationException(String.Format("Failed to delete directory '{0}'", NormalizedDirectory));
-				}
+				throw new AutomationException(String.Format("Failed to delete directory '{0}'", NormalizedDirectory));
 			}
 		}
 
@@ -538,9 +522,9 @@ namespace AutomationTool
 		/// If the deletion of the directory fails, this function throws an Exception.
 		/// </summary>
         /// <param name="Directories">Directories</param>
-        public static void DeleteDirectory(params string[] Directories)
+        public static void DeleteDirectory(string Directory)
 		{
-			DeleteDirectory(false, Directories);
+			DeleteDirectory(false, Directory);
 		}
 
 		/// <summary>
@@ -549,31 +533,28 @@ namespace AutomationTool
 		/// </summary>
 		/// <param name="bQuiet">Suppresses log output if true</param>
         /// <param name="Directories">Directories</param>
-        public static bool DeleteDirectory_NoExceptions(bool bQuiet, params string[] Directories)
+        public static bool DeleteDirectory_NoExceptions(bool bQuiet, string Directory)
 		{
 			bool Result = true;
-            foreach (var Directory in Directories)
+
+			var NormalizedDirectory = ConvertSeparators(PathSeparator.Default, Directory);
+            try
             {
-                var NormalizedDirectory = ConvertSeparators(PathSeparator.Default, Directory);
-                try
-                {
-                    if (!InternalUtils.SafeDeleteDirectory(NormalizedDirectory, bQuiet))
-                    {
-                        LogWarning("Failed to delete directory '{0}'", NormalizedDirectory);
-                        Result = false;
-                    }
-                }
-                catch (Exception Ex)
+                if (!InternalUtils.SafeDeleteDirectory(NormalizedDirectory, bQuiet))
                 {
-					if (!bQuiet)
-					{
-						LogWarning("Failed to delete directory, exception '{0}'", NormalizedDirectory);
-						LogWarning(Ex.Message);
-					}
+                    LogWarning("Failed to delete directory '{0}'", NormalizedDirectory);
                     Result = false;
                 }
             }
-
+            catch (Exception Ex)
+            {
+				if (!bQuiet)
+				{
+					LogWarning("Failed to delete directory, exception '{0}'", NormalizedDirectory);
+					LogWarning(Ex.Message);
+				}
+                Result = false;
+            }
 
 			return Result;
 		}
@@ -583,9 +564,9 @@ namespace AutomationTool
 		/// If the deletion of the directory fails, prints a warning.
 		/// </summary>
         /// <param name="Directories">Directories</param>
-        public static bool DeleteDirectory_NoExceptions(params string[] Directories)
+        public static bool DeleteDirectory_NoExceptions(string DirectoryName)
 		{
-			return DeleteDirectory_NoExceptions(false, Directories);
+			return DeleteDirectory_NoExceptions(false, DirectoryName);
 		}
 
 
@@ -616,15 +597,10 @@ namespace AutomationTool
 		/// </summary>
         /// <param name="Directories">Directories</param>
         /// <returns>True if the directory exists, false otherwise.</returns>
-		public static bool DirectoryExists(params string[] Directories)
+		public static bool DirectoryExists(string DirectoryName)
 		{
-			bool bExists = Directories.Length > 0;
-			foreach (var DirectoryName in Directories)
-			{
-				var NormalizedDirectory = ConvertSeparators(PathSeparator.Default, DirectoryName);
-				bExists = System.IO.Directory.Exists(NormalizedDirectory) && bExists;
-			}
-			return bExists;
+			var NormalizedDirectory = ConvertSeparators(PathSeparator.Default, DirectoryName);
+			return System.IO.Directory.Exists(NormalizedDirectory);
 		}
 
 		/// <summary>
@@ -632,25 +608,19 @@ namespace AutomationTool
 		/// </summary>
         /// <param name="Directories">Directories</param>
         /// <returns>True if the directory exists, false otherwise.</returns>
-		public static bool DirectoryExists_NoExceptions(params string[] Directories)
+		public static bool DirectoryExists_NoExceptions(string DirectoryName)
 		{
-			bool bExists = Directories.Length > 0;
-			foreach (var DirectoryName in Directories)
+			var NormalizedDirectory = ConvertSeparators(PathSeparator.Default, DirectoryName);
+			try
 			{
-				var NormalizedDirectory = ConvertSeparators(PathSeparator.Default, DirectoryName);
-				try
-				{
-					bExists = System.IO.Directory.Exists(NormalizedDirectory) && bExists;
-				}
-				catch (Exception Ex)
-				{
-					LogWarning("Unable to check if directory exists: {0}", NormalizedDirectory);
-					LogWarning(Ex.Message);
-					bExists = false;
-					break;
-				}
+				return System.IO.Directory.Exists(NormalizedDirectory);
+			}
+			catch (Exception Ex)
+			{
+				LogWarning("Unable to check if directory exists: {0}", NormalizedDirectory);
+				LogWarning(Ex.Message);
+				return false;
 			}
-			return bExists;
 		}
 
 		/// <summary>
@@ -658,15 +628,12 @@ namespace AutomationTool
 		/// If the creation of the directory fails, this function throws an Exception.
 		/// </summary>
         /// <param name="Directories">Directories</param>
-        public static void CreateDirectory(params string[] Directories)
+        public static void CreateDirectory(string DirectoryName)
 		{
-			foreach (var DirectoryName in Directories)
+			var NormalizedDirectory = ConvertSeparators(PathSeparator.Default, DirectoryName);
+			if (!InternalUtils.SafeCreateDirectory(NormalizedDirectory))
 			{
-				var NormalizedDirectory = ConvertSeparators(PathSeparator.Default, DirectoryName);
-				if (!InternalUtils.SafeCreateDirectory(NormalizedDirectory))
-				{
-					throw new AutomationException(String.Format("Failed to create directory '{0}'", NormalizedDirectory));
-				}
+				throw new AutomationException(String.Format("Failed to create directory '{0}'", NormalizedDirectory));
 			}
 		}
 
@@ -676,15 +643,12 @@ namespace AutomationTool
         /// </summary>
         /// <param name="bQuiet">When true, logging is suppressed.</param>
         /// <param name="Directories">Directories</param>
-        public static void CreateDirectory(bool bQuiet, params string[] Directories)
+        public static void CreateDirectory(bool bQuiet, string DirectoryName)
         {
-            foreach (var DirectoryName in Directories)
+            var NormalizedDirectory = ConvertSeparators(PathSeparator.Default, DirectoryName);
+            if (!InternalUtils.SafeCreateDirectory(NormalizedDirectory, bQuiet))
             {
-                var NormalizedDirectory = ConvertSeparators(PathSeparator.Default, DirectoryName);
-                if (!InternalUtils.SafeCreateDirectory(NormalizedDirectory, bQuiet))
-                {
-                    throw new AutomationException(String.Format("Failed to create directory '{0}'", NormalizedDirectory));
-                }
+                throw new AutomationException(String.Format("Failed to create directory '{0}'", NormalizedDirectory));
             }
         }
 
@@ -693,17 +657,14 @@ namespace AutomationTool
 		/// If the creation of the directory fails, this function prints a warning.
 		/// </summary>
         /// <param name="Directories">Directories</param>
-        public static bool CreateDirectory_NoExceptions(params string[] Directories)
+        public static bool CreateDirectory_NoExceptions(string DirectoryName)
 		{
 			bool Result = true;
-			foreach (var DirectoryName in Directories)
+			var NormalizedDirectory = ConvertSeparators(PathSeparator.Default, DirectoryName);
+			if (!InternalUtils.SafeCreateDirectory(NormalizedDirectory))
 			{
-				var NormalizedDirectory = ConvertSeparators(PathSeparator.Default, DirectoryName);
-				if (!InternalUtils.SafeCreateDirectory(NormalizedDirectory))
-				{
-					LogWarning("Failed to create directory '{0}'", NormalizedDirectory);
-					Result = false;
-				}
+				LogWarning("Failed to create directory '{0}'", NormalizedDirectory);
+				Result = false;
 			}
 			return Result;
 		}
@@ -748,15 +709,10 @@ namespace AutomationTool
 		/// </summary>
 		/// <param name="Filenames">Filename.</param>
 		/// <returns>True if the file exists, false otherwise.</returns>
-		public static bool FileExists(params string[] Filenames)
+		public static bool FileExists(string FileName)
 		{
-			bool bExists = Filenames.Length > 0;
-			foreach (var Filename in Filenames)
-			{
-				var NormalizedFilename = ConvertSeparators(PathSeparator.Default, Filename);
-				bExists = InternalUtils.SafeFileExists(NormalizedFilename) && bExists;
-			}
-			return bExists;
+			var NormalizedFilename = ConvertSeparators(PathSeparator.Default, FileName);
+			return InternalUtils.SafeFileExists(NormalizedFilename);
 		}
 
 		/// <summary>
@@ -764,26 +720,22 @@ namespace AutomationTool
 		/// </summary>
 		/// <param name="Filenames">Filename.</param>
 		/// <returns>True if the file exists, false otherwise.</returns>
-		public static bool FileExists_NoExceptions(params string[] Filenames)
+		public static bool FileExists_NoExceptions(string FileName)
 		{
 			// Standard version doesn't throw, but keep this function for consistency.
-			return FileExists(Filenames);
+			return FileExists(FileName);
 		}
+
         /// <summary>
         /// Checks if a file(s) exists.
         /// </summary>
         /// <param name="bQuiet">When true, logging is suppressed.</param>
         /// <param name="Filenames">Filename.</param>
         /// <returns>True if the file exists, false otherwise.</returns>
-        public static bool FileExists(bool bQuiet, params string[] Filenames)
+        public static bool FileExists(bool bQuiet, string FileName)
         {
-            bool bExists = Filenames.Length > 0;
-            foreach (var Filename in Filenames)
-            {
-                var NormalizedFilename = ConvertSeparators(PathSeparator.Default, Filename);
-                bExists = InternalUtils.SafeFileExists(NormalizedFilename, bQuiet) && bExists;
-            }
-            return bExists;
+			var NormalizedFilename = ConvertSeparators(PathSeparator.Default, FileName);
+			return InternalUtils.SafeFileExists(NormalizedFilename, bQuiet);
         }
 
         /// <summary>
@@ -792,10 +744,10 @@ namespace AutomationTool
         /// <param name="bQuiet">When true, logging is suppressed.</param>
         /// <param name="Filenames">Filename.</param>
         /// <returns>True if the file exists, false otherwise.</returns>
-        public static bool FileExists_NoExceptions(bool bQuiet, params string[] Filenames)
+        public static bool FileExists_NoExceptions(bool bQuiet, string FileName)
         {
             // Standard version doesn't throw, but keep this function for consistency.
-            return FileExists(bQuiet, Filenames);
+            return FileExists(bQuiet, FileName);
         }
 
 		static Stack<string> WorkingDirectoryStack = new Stack<string>();
@@ -1532,7 +1484,7 @@ namespace AutomationTool
 			{
 				string TargetFilename = CommandUtils.CombinePaths(TargetPathBase, SourceFI.Name);
 
-				if (!bIncremental || !CommandUtils.FileExists_NoExceptions(TargetFilename))
+				if (!bIncremental || !CommandUtils.FileExists(TargetFilename))
 				{
 				SourceFI.CopyTo(TargetFilename);
 
@@ -1655,7 +1607,8 @@ namespace AutomationTool
 		{
 			// Filter all the relative paths
 			Log("Applying filter to {0}...", SourceDir);
-			var RelativePaths = Filter.ApplyToDirectory(SourceDir, bIgnoreSymlinks);
+			DirectoryReference SourceDirRef = new DirectoryReference(SourceDir);
+			var RelativePaths = Filter.ApplyToDirectory(SourceDirRef, bIgnoreSymlinks).Select(x => x.MakeRelativeTo(SourceDirRef)).ToList();
 			return ThreadedCopyFiles(SourceDir, TargetDir, RelativePaths);
 		}
 
@@ -2059,17 +2012,17 @@ namespace AutomationTool
 		/// <param name="ZipFileName">Filename for the zip</param>
 		/// <param name="Filter">Filter which selects files to be included in the zip</param>
 		/// <param name="BaseDirectory">Base directory to store relative paths in the zip file to</param>
-		public static void ZipFiles(string ZipFileName, string BaseDirectory, FileFilter Filter)
+		public static void ZipFiles(FileReference ZipFileName, DirectoryReference BaseDirectory, FileFilter Filter)
 		{
 			// Ionic.Zip.Zip64Option.Always option produces broken archives on Mono, so we use system zip tool instead
 			if (Utils.IsRunningOnMono)
 			{
-				CommandUtils.CreateDirectory(Path.GetDirectoryName(ZipFileName));
-				CommandUtils.PushDir(BaseDirectory);
+				CommandUtils.CreateDirectory(Path.GetDirectoryName(ZipFileName.FullName));
+				CommandUtils.PushDir(BaseDirectory.FullName);
 				string FilesList = "";
-				foreach (string FilteredFile in Filter.ApplyToDirectory(BaseDirectory, true))
+				foreach (FileReference FilteredFile in Filter.ApplyToDirectory(BaseDirectory, true))
 				{
-					FilesList += " \"" + FilteredFile + "\"";
+					FilesList += " \"" + FilteredFile.MakeRelativeTo(BaseDirectory) + "\"";
 					if (FilesList.Length > 32000)
 					{
 						CommandUtils.RunAndLog(CommandUtils.CmdEnv, "zip", "-g -q \"" + ZipFileName + "\"" + FilesList);
@@ -2087,12 +2040,12 @@ namespace AutomationTool
 				using (Ionic.Zip.ZipFile Zip = new Ionic.Zip.ZipFile())
 				{
 					Zip.UseZip64WhenSaving = Ionic.Zip.Zip64Option.Always;
-					foreach (string FilteredFile in Filter.ApplyToDirectory(BaseDirectory, true))
+					foreach (FileReference FilteredFile in Filter.ApplyToDirectory(BaseDirectory, true))
 					{
-						Zip.AddFile(Path.Combine(BaseDirectory, FilteredFile), Path.GetDirectoryName(FilteredFile));
+						Zip.AddFile(FilteredFile.FullName, Path.GetDirectoryName(FilteredFile.MakeRelativeTo(BaseDirectory)));
 					}
-					CommandUtils.CreateDirectory(Path.GetDirectoryName(ZipFileName));
-					Zip.Save(ZipFileName);
+					CommandUtils.CreateDirectory(Path.GetDirectoryName(ZipFileName.FullName));
+					Zip.Save(ZipFileName.FullName);
 				}
 			}
 		}
diff --git a/Engine/Source/Programs/AutomationTool/AutomationUtils/CommandletUtils.cs b/Engine/Source/Programs/AutomationTool/AutomationUtils/CommandletUtils.cs
--- a/Engine/Source/Programs/AutomationTool/AutomationUtils/CommandletUtils.cs
+++ b/Engine/Source/Programs/AutomationTool/AutomationUtils/CommandletUtils.cs
@@ -10,6 +10,36 @@ using UnrealBuildTool;
 
 namespace AutomationTool
 {
+	/// <summary>
+	/// Exception thrown when the execution of a commandlet fails
+	/// </summary>
+	public class CommandletException : AutomationException
+	{
+		/// <summary>
+		/// The log file output by this commandlet
+		/// </summary>
+		public string LogFileName;
+
+		/// <summary>
+		/// The exit code
+		/// </summary>
+		public int ErrorNumber;
+
+		/// <summary>
+		/// Constructor
+		/// </summary>
+		/// <param name="LogFilename">File containing the error log</param>
+		/// <param name="ErrorNumber">The exit code from the commandlet</param>
+		/// <param name="Format">Formatting string for the message</param>
+		/// <param name="Args">Arguments for the formatting string</param>
+		public CommandletException(string LogFilename, int ErrorNumber, string Format, params object[] Args)
+			: base(Format,Args)
+		{
+			this.LogFileName = LogFilename;
+			this.ErrorNumber = ErrorNumber;
+		}
+	}
+
 	public partial class CommandUtils
 	{
 		#region Commandlets
@@ -370,7 +400,7 @@ namespace AutomationTool
 
 			if (RunResult.ExitCode != 0)
 			{
-				throw new AutomationException(DestLogFile, RunResult.ExitCode, "BUILD FAILED: Failed while running {0} for {1}; see log {2}", Commandlet, ProjectName, DestLogFile);
+				throw new CommandletException(DestLogFile, RunResult.ExitCode, "BUILD FAILED: Failed while running {0} for {1}; see log {2}", Commandlet, ProjectName, DestLogFile);
 			}
 		}
 		
diff --git a/Engine/Source/Programs/AutomationTool/AutomationUtils/DeploymentContext.cs b/Engine/Source/Programs/AutomationTool/AutomationUtils/DeploymentContext.cs
--- a/Engine/Source/Programs/AutomationTool/AutomationUtils/DeploymentContext.cs
+++ b/Engine/Source/Programs/AutomationTool/AutomationUtils/DeploymentContext.cs
@@ -476,11 +476,81 @@ public class DeploymentContext //: ProjectParams
 		}
 	}
 
+	/// <summary>
+	/// Correctly collapses any ../ or ./ entries in a path.
+	/// </summary>
+	/// <param name="InPath">The path to be collapsed</param>
+	/// <returns>true if the path could be collapsed, false otherwise.</returns>
+	static bool CollapseRelativeDirectories(ref string InPath)
+	{
+		string LocalString = InPath;
+		bool bHadBackSlashes = false;
+		// look to see what kind of slashes we had
+		if (LocalString.IndexOf("\\") != -1)
+		{
+			LocalString = LocalString.Replace("\\", "/");
+			bHadBackSlashes = true;
+		}
+
+		string ParentDir = "/..";
+		int ParentDirLength = ParentDir.Length;
+
+		for (; ; )
+		{
+			// An empty path is finished
+			if (string.IsNullOrEmpty(LocalString))
+				break;
+
+			// Consider empty paths or paths which start with .. or /.. as invalid
+			if (LocalString.StartsWith("..") || LocalString.StartsWith(ParentDir))
+				return false;
+
+			// If there are no "/.."s left then we're done
+			int Index = LocalString.IndexOf(ParentDir);
+			if (Index == -1)
+				break;
+
+			int PreviousSeparatorIndex = Index;
+			for (; ; )
+			{
+				// Find the previous slash
+				PreviousSeparatorIndex = Math.Max(0, LocalString.LastIndexOf("/", PreviousSeparatorIndex - 1));
+
+				// Stop if we've hit the start of the string
+				if (PreviousSeparatorIndex == 0)
+					break;
+
+				// Stop if we've found a directory that isn't "/./"
+				if ((Index - PreviousSeparatorIndex) > 1 && (LocalString[PreviousSeparatorIndex + 1] != '.' || LocalString[PreviousSeparatorIndex + 2] != '/'))
+					break;
+			}
+
+			// If we're attempting to remove the drive letter, that's illegal
+			int Colon = LocalString.IndexOf(":", PreviousSeparatorIndex);
+			if (Colon >= 0 && Colon < Index)
+				return false;
+
+			LocalString = LocalString.Substring(0, PreviousSeparatorIndex) + LocalString.Substring(Index + ParentDirLength);
+		}
+
+		LocalString = LocalString.Replace("./", "");
+
+		// restore back slashes now
+		if (bHadBackSlashes)
+		{
+			LocalString = LocalString.Replace("/", "\\");
+		}
+
+		// and pass back out
+		InPath = LocalString;
+		return true;
+	}
+
     public void StageFiles(StagedFileType FileType, string InPath, string Wildcard = "*", bool bRecursive = true, string[] ExcludeWildcard = null, string NewPath = null, bool bAllowNone = false, bool bRemap = true, string NewName = null, bool bAllowNotForLicenseesFiles = true, bool bStripFilesForOtherPlatforms = true, bool bConvertToLower = false)
 	{
 		int FilesAdded = 0;
 		// make sure any ..'s are removed
-		Utils.CollapseRelativeDirectories(ref InPath);
+		CollapseRelativeDirectories(ref InPath);
 
 		if (CommandUtils.DirectoryExists(InPath))
 		{
