Copying //UE4/Dev-Core to //UE4/Dev-Main (Source: //UE4/Dev-Core @ 3345728)

diff --git a/Engine/Source/Programs/AutomationTool/AutomationUtils/Distiller.cs b/Engine/Source/Programs/AutomationTool/AutomationUtils/Distiller.cs
--- a/Engine/Source/Programs/AutomationTool/AutomationUtils/Distiller.cs
+++ b/Engine/Source/Programs/AutomationTool/AutomationUtils/Distiller.cs
@@ -46,220 +46,4 @@ namespace AutomationTool
             return DestFile;
         }
     }
-    public class Distiller
-    {
-        public string SourceBaseDir;
-        public string DestBaseDir;
-        public string DestBaseDirSymbols;
-        public DateTime DestinationTime;
-        public List<UnrealTargetPlatform> AllLegalPlatforms;
-        public List<string> RejectStrings;
-
-        public Distiller(string InDestBaseDir, string InDestBaseDirSymbols = null, string InSourceBaseDir = null, DateTime? InDestinationTime = null, List<UnrealTargetPlatform> InAllLegalPlatforms = null, bool AllowNoRedist = false)
-        {
-            DestBaseDir = CommandUtils.CombinePaths(InDestBaseDir);
-            SourceBaseDir = InSourceBaseDir;
-            if (SourceBaseDir == null)
-            {
-                SourceBaseDir = CommandUtils.CmdEnv.LocalRoot;
-            }
-            SourceBaseDir = CommandUtils.CombinePaths(SourceBaseDir);
-            if (InDestinationTime == null)
-            {
-                DestinationTime = DateTime.Now;
-            }
-            else
-            {
-                DestinationTime = InDestinationTime.Value;
-            }
-            DestBaseDirSymbols = InDestBaseDirSymbols;
-            if (DestBaseDirSymbols != null)
-            {
-                DestBaseDirSymbols = CommandUtils.CombinePaths(DestBaseDirSymbols);
-            }
-            AllLegalPlatforms = InAllLegalPlatforms;
-            if (AllLegalPlatforms == null)
-            {
-                AllLegalPlatforms = new List<UnrealTargetPlatform>();
-                foreach (UnrealTargetPlatform Plat in Enum.GetValues(typeof(UnrealTargetPlatform)))
-                {
-                    AllLegalPlatforms.Add(Plat);
-                }
-            }
-            RejectStrings = new List<string>
-            {
-                CommandUtils.CombinePaths("/CarefullyRedist/"),
-                CommandUtils.CombinePaths("/NotForLicensees/"),
-            };
-            if (!AllowNoRedist)
-            {
-                RejectStrings.Add(CommandUtils.CombinePaths("/NoRedist/"));
-            }
-            foreach (UnrealTargetPlatform Plat in Enum.GetValues(typeof(UnrealTargetPlatform)))
-            {
-                if (!AllLegalPlatforms.Contains(Plat))
-                {
-                    RejectStrings.Add(CommandUtils.CombinePaths("/" + Plat.ToString() + "/"));
-                }
-            }
-            if (!AllLegalPlatforms.Contains(UnrealTargetPlatform.Win64) && !AllLegalPlatforms.Contains(UnrealTargetPlatform.Win32))
-            {
-                RejectStrings.Add(CommandUtils.CombinePaths("/Windows/"));
-            }
-        }
-
-        public static bool RejectFileOnSubstring(string FileToCopy, string RejectString, string Root = null)
-        {
-            var Search = CommandUtils.CombinePaths(FileToCopy);
-            if (Root == null)
-            {
-                Root = CommandUtils.CmdEnv.LocalRoot;
-            }
-            if (Search.StartsWith(CommandUtils.CombinePaths(Root), StringComparison.InvariantCultureIgnoreCase))
-            {
-                Search = CommandUtils.CombinePaths("/", CommandUtils.StripBaseDirectory(FileToCopy, Root)); // lots of rejects are like "/path/" so we need to be sure we start with a slash
-            }
-            if (Search.IndexOf(CommandUtils.CombinePaths(RejectString), 0, StringComparison.InvariantCultureIgnoreCase) >= 0)
-            {
-                return true;
-            }
-            return false;
-        }
-
-        public bool RejectFile(string FileToCopy, string Root = null)
-        {
-            bool Result = false;
-            foreach (var RejectOn in RejectStrings)
-            {
-                if (RejectFileOnSubstring(FileToCopy, RejectOn, Root == null ? SourceBaseDir : Root))
-                {
-                    Result = true;
-                    break;
-                }
-            }
-            if (Result)
-            {
-                CommandUtils.LogVerbose("Rejecting file {0}", FileToCopy);
-            }
-            return Result;
-        }
-
-        public static bool IsSymbolFile(UnrealTargetPlatform ForPlatform, string FileToCopy)
-        {
-            TargetPlatformDescriptor PlatformDesc = new TargetPlatformDescriptor(ForPlatform);
-            var DebugExts = Platform.Platforms[PlatformDesc].GetDebugFileExtentions();
-            foreach (var DebugExt in DebugExts)
-            {
-                if (Path.GetExtension(FileToCopy).Equals(DebugExt, StringComparison.InvariantCultureIgnoreCase))
-                {
-                    return true;
-                }
-            }
-            return false;
-        }
-
-        public bool IsSymbolFile(string FileToCopy)
-        {
-            foreach (var ForPlatform in AllLegalPlatforms)
-            {
-                if (IsSymbolFile(ForPlatform, FileToCopy))
-                {
-                    return true;
-                }
-            }
-            return false;
-        }
-
-        public string CopyFileToDest(string FileToCopy, bool MoveSymbols = true)
-        {
-            if (!CommandUtils.FileExists_NoExceptions(FileToCopy))
-            {
-                throw new AutomationException("Can't distill {0}; it doesn't exist", FileToCopy);
-            }
-            {
-                FileInfo FileInfo = new FileInfo(FileToCopy);
-                FileToCopy = CommandUtils.CombinePaths(FileInfo.FullName);
-            }
-
-            string Dest;
-            if (MoveSymbols && !String.IsNullOrEmpty(DestBaseDirSymbols) && IsSymbolFile(FileToCopy))
-            {
-                Dest = CommandUtils.MakeRerootedFilePath(FileToCopy, SourceBaseDir, DestBaseDirSymbols);
-            }
-            else
-            {
-                Dest = CommandUtils.MakeRerootedFilePath(FileToCopy, SourceBaseDir, DestBaseDir);
-            }
-            if (CommandUtils.FileExists_NoExceptions(Dest))
-            {
-                throw new AutomationException("Can't distill to {0}; it already exists", Dest);
-            }
-            CommandUtils.CopyFile(FileToCopy, Dest);
-
-            {
-                FileInfo FileInfo = new FileInfo(Dest);
-                FileInfo.IsReadOnly = false;
-                FileInfo.CreationTime = DestinationTime;
-                FileInfo.LastWriteTime = DestinationTime;
-                Dest = CommandUtils.CombinePaths(FileInfo.FullName);
-            }
-            if (!CommandUtils.FileExists_NoExceptions(Dest))
-            {
-                throw new AutomationException("Couldn't distill {0} to {1}", FileToCopy, Dest);
-            }
-            return Dest;
-        }
-
-        public List<string> Distill(string FullPathAndWildcard, bool bRecursive = false, bool bAllowMissing = false, bool MoveSymbols = true, params string[] Exclusions)
-        {
-            var FilesToCopy = new List<string>();
-
-            string AbsFile = CommandUtils.CombinePaths(FullPathAndWildcard);
-            CommandUtils.LogVerbose("Distilling {0}", AbsFile);
-            string PathOnly = Path.GetDirectoryName(AbsFile);
-            if (!CommandUtils.DirectoryExists_NoExceptions(PathOnly))
-            {
-                if (bAllowMissing)
-                {
-                    return FilesToCopy;
-                }
-                throw new AutomationException("Path {0} was used to distill {1} but it doesn't exist.", PathOnly, AbsFile);
-            }
-            var Files = CommandUtils.FindFiles(Path.GetFileName(AbsFile), bRecursive, PathOnly);
-            var Exclude = new List<string>();
-            foreach (var Excl in Exclusions)
-            {
-                if (!Excl.Contains("/") && !Excl.Contains("\\"))
-                {
-                    var ExFiles = CommandUtils.FindFiles(Excl, bRecursive, PathOnly);
-                    Exclude.AddRange(ExFiles);
-                }
-            }
-
-            foreach (var FileToCopy in Files)
-            {
-                bool bOk = !Exclude.Contains(FileToCopy) && !RejectFile(FileToCopy, PathOnly);
-                foreach (var Excl in Exclusions)
-                {
-                    if (Excl.Contains("/") || Excl.Contains("\\"))
-                    {
-                        if (Excl.Contains("*") || Excl.Contains("?"))
-                        {
-                            throw new AutomationException("Exclusion {0} is illegal, must either be a substring or a wildcard without a path", Excl);
-                        }
-                        bOk = bOk && !RejectFileOnSubstring(FileToCopy, Excl, PathOnly);
-                    }
-                }
-                if (bOk)
-                {
-                    FilesToCopy.Add(CopyFileToDest(FileToCopy, MoveSymbols));
-                }
-            }
-            if (FilesToCopy.Count < 1 && !bAllowMissing)
-            {
-                throw new AutomationException("Distill {0} did not produce any files. Recursive option was {1}.", AbsFile, bRecursive.ToString());
-            }
-            return FilesToCopy;
-        }
-    }
 }
diff --git a/Engine/Source/Programs/AutomationTool/AutomationUtils/LocalP4Environment.cs b/Engine/Source/Programs/AutomationTool/AutomationUtils/LocalP4Environment.cs
--- a/Engine/Source/Programs/AutomationTool/AutomationUtils/LocalP4Environment.cs
+++ b/Engine/Source/Programs/AutomationTool/AutomationUtils/LocalP4Environment.cs
@@ -86,7 +86,6 @@ namespace AutomationTool
 				CommandUtils.ConditionallySetEnvVar(EnvVarNames.CodeChangelist, CodeCLString);
 			}
 
-			CommandUtils.ConditionallySetEnvVar(EnvVarNames.LabelToSync, "");
 			CommandUtils.ConditionallySetEnvVar("P4USER", UserName);
 			CommandUtils.ConditionallySetEnvVar("P4CLIENT", ThisClient.Name);
 
diff --git a/Engine/Source/Programs/AutomationTool/AutomationUtils/Log.cs b/Engine/Source/Programs/AutomationTool/AutomationUtils/Log.cs
--- a/Engine/Source/Programs/AutomationTool/AutomationUtils/Log.cs
+++ b/Engine/Source/Programs/AutomationTool/AutomationUtils/Log.cs
@@ -34,6 +34,7 @@ namespace AutomationTool
 				InLogLevel: (UnrealBuildTool.LogEventType)Enum.Parse(typeof(UnrealBuildTool.LogEventType), CommandUtils.ParseParamValue(CommandLine, "-Verbose=", "Log")),
                 bLogSeverity: true,
                 bLogSources: true,
+				bLogSourcesToConsole: true,
                 bColorConsoleOutput: true,
                 TraceListeners: new TraceListener[]
                 {
diff --git a/Engine/Source/Programs/AutomationTool/AutomationUtils/P4Environment.cs b/Engine/Source/Programs/AutomationTool/AutomationUtils/P4Environment.cs
--- a/Engine/Source/Programs/AutomationTool/AutomationUtils/P4Environment.cs
+++ b/Engine/Source/Programs/AutomationTool/AutomationUtils/P4Environment.cs
@@ -73,8 +73,6 @@ namespace AutomationTool
 		public string Client { get; protected set; }
 		public string BuildRootP4 { get; protected set; }
 		public string BuildRootEscaped { get; protected set; }
-		public string LabelToSync { get; protected set; }
-		public string LabelPrefix { get; protected set; }
 
 
 		internal P4Environment(P4Connection Connection, CommandEnvironment CmdEnv)
@@ -101,7 +99,6 @@ namespace AutomationTool
 				CommandUtils.SetEnvVar(EnvVarNames.BuildRootP4, BuildRootP4);
 			}
 			BuildRootEscaped = CommandUtils.GetEnvVar(EnvVarNames.BuildRootEscaped);
-			LabelToSync = CommandUtils.GetEnvVar(EnvVarNames.LabelToSync);
 
 			string CodeChangelistString = CommandUtils.GetEnvVar(EnvVarNames.CodeChangelist);
 			if(!String.IsNullOrEmpty(CodeChangelistString))
@@ -122,8 +119,6 @@ namespace AutomationTool
 				throw new AutomationException("BUILD FAILED Perforce Environment is not set up correctly. Please check your environment variables.");
 			}
 
-			LabelPrefix = BuildRootP4 + "/";
-
 			if (CommandUtils.P4Enabled)
 			{
 				if (CommandUtils.IsBuildMachine || ChangelistStringInternal != null)
@@ -141,19 +136,26 @@ namespace AutomationTool
 		{
 			// Log all vars
 			const bool bQuiet = true;
-			Log.TraceInformation("Perforce Environment settings:");
-			Log.TraceInformation("{0}={1}", EnvVarNames.P4Port, InternalUtils.GetEnvironmentVariable(EnvVarNames.P4Port, "", bQuiet));
-			Log.TraceInformation("{0}={1}", EnvVarNames.User, InternalUtils.GetEnvironmentVariable(EnvVarNames.User, "", bQuiet));
-			Log.TraceInformation("{0}={1}", EnvVarNames.Client, InternalUtils.GetEnvironmentVariable(EnvVarNames.Client, "", bQuiet));
-			Log.TraceInformation("{0}={1}", EnvVarNames.BuildRootP4, InternalUtils.GetEnvironmentVariable(EnvVarNames.BuildRootP4, "", bQuiet));
-			Log.TraceInformation("{0}={1}", EnvVarNames.BuildRootEscaped, InternalUtils.GetEnvironmentVariable(EnvVarNames.BuildRootEscaped, "", bQuiet));
-			Log.TraceInformation("{0}={1}", EnvVarNames.ClientRoot, InternalUtils.GetEnvironmentVariable(EnvVarNames.ClientRoot, "", bQuiet));
-			Log.TraceInformation("{0}={1}", EnvVarNames.Changelist, InternalUtils.GetEnvironmentVariable(EnvVarNames.Changelist, "", bQuiet));
-			Log.TraceInformation("{0}={1}", EnvVarNames.CodeChangelist, InternalUtils.GetEnvironmentVariable(EnvVarNames.CodeChangelist, "", bQuiet));
-			Log.TraceInformation("{0}={1}", EnvVarNames.LabelToSync, InternalUtils.GetEnvironmentVariable(EnvVarNames.LabelToSync, "", bQuiet));
-			Log.TraceInformation("{0}={1}", "P4USER", InternalUtils.GetEnvironmentVariable("P4USER", "", bQuiet));
-			Log.TraceInformation("{0}={1}", "P4CLIENT", InternalUtils.GetEnvironmentVariable("P4CLIENT", "", bQuiet));
-			Log.TraceInformation("LabelPrefix={0}", LabelPrefix);
+
+			Log.TraceInformation("Detected Perforce Settings:");
+			Log.TraceInformation("    Port: {0}", InternalUtils.GetEnvironmentVariable(EnvVarNames.P4Port, "", bQuiet));
+			Log.TraceInformation("    User: {0}", InternalUtils.GetEnvironmentVariable(EnvVarNames.User, "", bQuiet));
+			Log.TraceInformation("    Client: {0}", InternalUtils.GetEnvironmentVariable(EnvVarNames.Client, "", bQuiet));
+			Log.TraceInformation("    Branch: {0}", InternalUtils.GetEnvironmentVariable(EnvVarNames.BuildRootP4, "", bQuiet));
+			Log.TraceInformation("    Last Change: {0}", InternalUtils.GetEnvironmentVariable(EnvVarNames.Changelist, "", bQuiet));
+			Log.TraceInformation("    Last Code Change: {0}", InternalUtils.GetEnvironmentVariable(EnvVarNames.CodeChangelist, "", bQuiet));
+
+			Log.TraceLog("Perforce Environment Variables:");
+			Log.TraceLog("    {0}={1}", EnvVarNames.P4Port, InternalUtils.GetEnvironmentVariable(EnvVarNames.P4Port, "", bQuiet));
+			Log.TraceLog("    {0}={1}", EnvVarNames.User, InternalUtils.GetEnvironmentVariable(EnvVarNames.User, "", bQuiet));
+			Log.TraceLog("    {0}={1}", EnvVarNames.Client, InternalUtils.GetEnvironmentVariable(EnvVarNames.Client, "", bQuiet));
+			Log.TraceLog("    {0}={1}", EnvVarNames.BuildRootP4, InternalUtils.GetEnvironmentVariable(EnvVarNames.BuildRootP4, "", bQuiet));
+			Log.TraceLog("    {0}={1}", EnvVarNames.BuildRootEscaped, InternalUtils.GetEnvironmentVariable(EnvVarNames.BuildRootEscaped, "", bQuiet));
+			Log.TraceLog("    {0}={1}", EnvVarNames.ClientRoot, InternalUtils.GetEnvironmentVariable(EnvVarNames.ClientRoot, "", bQuiet));
+			Log.TraceLog("    {0}={1}", EnvVarNames.Changelist, InternalUtils.GetEnvironmentVariable(EnvVarNames.Changelist, "", bQuiet));
+			Log.TraceLog("    {0}={1}", EnvVarNames.CodeChangelist, InternalUtils.GetEnvironmentVariable(EnvVarNames.CodeChangelist, "", bQuiet));
+			Log.TraceLog("    {0}={1}", "P4USER", InternalUtils.GetEnvironmentVariable("P4USER", "", bQuiet));
+			Log.TraceLog("    {0}={1}", "P4CLIENT", InternalUtils.GetEnvironmentVariable("P4CLIENT", "", bQuiet));
 		}
 
 		protected void VerifyChangelistStringAndSetChangelistNumber()
diff --git a/Engine/Source/Programs/AutomationTool/AutomationUtils/P4Utils.cs b/Engine/Source/Programs/AutomationTool/AutomationUtils/P4Utils.cs
--- a/Engine/Source/Programs/AutomationTool/AutomationUtils/P4Utils.cs
+++ b/Engine/Source/Programs/AutomationTool/AutomationUtils/P4Utils.cs
@@ -21,7 +21,6 @@ namespace AutomationTool
 
 	public class P4Exception : AutomationException
 	{
-		public P4Exception() { }
 		public P4Exception(string Msg)
 			: base(Msg) { }
 		public P4Exception(string Msg, Exception InnerException)
@@ -2047,101 +2046,6 @@ namespace AutomationTool
 			return false;
 		}
 
-		/// <summary>
-        /// returns the full name of a label. //depot/UE4/TEST-GUBP-Promotable-GameName-CL-CLNUMBER
-		/// </summary>
-		/// <param name="BuildNamePrefix">Label Prefix</param>
-        public string FullLabelName(P4Environment Env, string BuildNamePrefix)
-        {
-            CheckP4Enabled();
-			var Label = Env.LabelPrefix + BuildNamePrefix + "-CL-" + Env.ChangelistString;
-			CommandUtils.LogLog("Label prefix {0}", BuildNamePrefix);
-			CommandUtils.LogLog("Full Label name {0}", Label); 
-            return Label;
-        }
-
-		/// <summary>
-		/// Creates a downstream label.
-		/// </summary>
-		/// <param name="BuildNamePrefix">Label Prefix</param>
-		public void MakeDownstreamLabel(P4Environment Env, string BuildNamePrefix, List<string> Files = null)
-		{
-			CheckP4Enabled();
-			string DOWNSTREAM_LabelPrefix = CommandUtils.GetEnvVar("DOWNSTREAM_LabelPrefix");
-			if (!String.IsNullOrEmpty(DOWNSTREAM_LabelPrefix))
-			{
-				BuildNamePrefix = DOWNSTREAM_LabelPrefix;
-			}
-			if (String.IsNullOrEmpty(BuildNamePrefix))
-			{
-				throw new P4Exception("Need a downstream label");
-			}
-
-			{
-				CommandUtils.LogLog("Making downstream label");
-                var Label = FullLabelName(Env, BuildNamePrefix);
-
-				CommandUtils.LogLog("Deleting old label {0} (if any)...", Label);
-                DeleteLabel(Label, false);
-
-				CommandUtils.LogLog("Creating new label...");
-				CreateLabel(
-					Name: Label,
-					Description: "BVT Time " + CommandUtils.CmdEnv.TimestampAsString + "  CL " + Env.ChangelistString,
-					Options: "unlocked noautoreload",
-					View: CommandUtils.CombinePaths(PathSeparator.Depot, Env.BuildRootP4, "...")
-					);
-				if (Files == null)
-				{
-					CommandUtils.LogLog("Adding all files to new label {0}...", Label);
-					LabelSync(Label, false);
-				}
-				else
-				{
-					CommandUtils.LogLog("Adding build products to new label {0}...", Label);
-					foreach (string LabelFile in Files)
-					{
-						LabelSync(Label, false, LabelFile);
-					}
-				}
-			}
-		}
-
-        /// <summary>
-        /// Creates a downstream label.
-        /// </summary>
-        /// <param name="BuildNamePrefix">Label Prefix</param>
-		public void MakeDownstreamLabelFromLabel(P4Environment Env, string BuildNamePrefix, string CopyFromBuildNamePrefix)
-        {
-            CheckP4Enabled();
-			string DOWNSTREAM_LabelPrefix = CommandUtils.GetEnvVar("DOWNSTREAM_LabelPrefix");
-            if (!String.IsNullOrEmpty(DOWNSTREAM_LabelPrefix))
-            {
-                BuildNamePrefix = DOWNSTREAM_LabelPrefix;
-            }
-            if (String.IsNullOrEmpty(BuildNamePrefix) || String.IsNullOrEmpty(CopyFromBuildNamePrefix))
-            {
-                throw new P4Exception("Need a downstream label");
-            }
-
-            {
-				CommandUtils.LogLog("Making downstream label");
-                var Label = FullLabelName(Env, BuildNamePrefix);
-
-				CommandUtils.LogLog("Deleting old label {0} (if any)...", Label);
-                DeleteLabel(Label, false);
-
-				CommandUtils.LogLog("Creating new label...");
-                CreateLabel(
-                    Name: Label,
-					Description: "BVT Time " + CommandUtils.CmdEnv.TimestampAsString + "  CL " + Env.ChangelistString,
-                    Options: "unlocked noautoreload",
-					View: CommandUtils.CombinePaths(PathSeparator.Depot, Env.BuildRootP4, "...")
-                    );
-                LabelToLabelSync(FullLabelName(Env, CopyFromBuildNamePrefix), Label, false);
-            }
-        }
-
         /// <summary>
         /// Given a file path in the depot, returns the local disk mapping for the current view
         /// </summary>
diff --git a/Engine/Source/Programs/AutomationTool/AutomationUtils/ProjectParams.cs b/Engine/Source/Programs/AutomationTool/AutomationUtils/ProjectParams.cs
--- a/Engine/Source/Programs/AutomationTool/AutomationUtils/ProjectParams.cs
+++ b/Engine/Source/Programs/AutomationTool/AutomationUtils/ProjectParams.cs
@@ -2176,7 +2176,7 @@ namespace AutomationTool
             {
                 throw new AutomationException("RawProjectPath {0} must end with .uproject", RawProjectPath);
             }
-            if (!CommandUtils.FileExists_NoExceptions(RawProjectPath.FullName))
+            if (!CommandUtils.FileExists(RawProjectPath.FullName))
             {
                 throw new AutomationException("RawProjectPath {0} file must exist", RawProjectPath);
             }
diff --git a/Engine/Source/Programs/AutomationTool/BuildGraph/Tasks/SubmitTask.cs b/Engine/Source/Programs/AutomationTool/BuildGraph/Tasks/SubmitTask.cs
--- a/Engine/Source/Programs/AutomationTool/BuildGraph/Tasks/SubmitTask.cs
+++ b/Engine/Source/Programs/AutomationTool/BuildGraph/Tasks/SubmitTask.cs
@@ -50,6 +50,12 @@ namespace AutomationTool.Tasks
 		/// </summary>
 		[TaskParameter(Optional = true, ValidationType = TaskParameterValidationType.DirectoryName)]
 		public string RootDir;
+
+		/// <summary>
+		/// Whether to revert unchanged files before attempting to submit
+		/// </summary>
+		[TaskParameter(Optional = true)]
+		public bool RevertUnchanged;
 	}
 
 	/// <summary>
@@ -117,6 +123,17 @@ namespace AutomationTool.Tasks
 					}
 				}
 
+				// Revert any unchanged files
+				if(Parameters.RevertUnchanged)
+				{
+					SubmitP4.RevertUnchanged(NewCL);
+					if(SubmitP4.TryDeleteEmptyChange(NewCL))
+					{
+						CommandUtils.Log("No files to submit; ignored.");
+						return true;
+					}
+				}
+
 				// Submit it
 				int SubmittedCL;
 				SubmitP4.Submit(NewCL, out SubmittedCL);
diff --git a/Engine/Source/Programs/AutomationTool/IOS/IOSPlatform.Automation.cs b/Engine/Source/Programs/AutomationTool/IOS/IOSPlatform.Automation.cs
--- a/Engine/Source/Programs/AutomationTool/IOS/IOSPlatform.Automation.cs
+++ b/Engine/Source/Programs/AutomationTool/IOS/IOSPlatform.Automation.cs
@@ -927,7 +927,7 @@ public class IOSPlatform : Platform
 			string dSYMName = (SC.IsCodeBasedProject ? Path.GetFileNameWithoutExtension(ProjectExe) : "UE4Game") + ".dSYM";
 			string dSYMSrcPath = Path.Combine(SC.ProjectBinariesFolder, dSYMName);
 
-			if(DirectoryExists(new string[] { dSYMSrcPath } ))
+			if(DirectoryExists(dSYMSrcPath))
 			{
 				// Create the dsyms archive folder
 				CreateDirectory(Path.Combine(ArchiveName, "dSYMs"));
diff --git a/Engine/Source/Programs/AutomationTool/Scripts/AutomationScripts.Automation.csproj b/Engine/Source/Programs/AutomationTool/Scripts/AutomationScripts.Automation.csproj
--- a/Engine/Source/Programs/AutomationTool/Scripts/AutomationScripts.Automation.csproj
+++ b/Engine/Source/Programs/AutomationTool/Scripts/AutomationScripts.Automation.csproj
@@ -74,6 +74,7 @@
     <Compile Include="BuildProjectCommand.Automation.cs" />
     <Compile Include="BuildTargetCommand.Automation.cs" />
     <Compile Include="BuildThirdPartyLibs.Automation.cs" />
+    <Compile Include="CheckCsprojDotNetVersion.cs" />
     <Compile Include="CleanFormalBuilds.Automation.cs" />
     <Compile Include="CodeSurgery.Automation.cs" />
     <Compile Include="CookCommand.Automation.cs" />
diff --git a/Engine/Source/Programs/AutomationTool/Scripts/BuildCommonTools.Automation.cs b/Engine/Source/Programs/AutomationTool/Scripts/BuildCommonTools.Automation.cs
--- a/Engine/Source/Programs/AutomationTool/Scripts/BuildCommonTools.Automation.cs
+++ b/Engine/Source/Programs/AutomationTool/Scripts/BuildCommonTools.Automation.cs
@@ -190,7 +190,7 @@ public class ZipProjectUp : BuildCommand
         Filter.Include("/Source/...");
         Filter.Include("*.uproject");
 
-        ZipFiles(InstallDirectory, ProjectDirectory, Filter);
+        ZipFiles(new FileReference(InstallDirectory), new DirectoryReference(ProjectDirectory), Filter);
 
 		Log("Completed zipping project up");
     }
diff --git a/Engine/Source/Programs/AutomationTool/Scripts/BuildForUGS.Automation.cs b/Engine/Source/Programs/AutomationTool/Scripts/BuildForUGS.Automation.cs
--- a/Engine/Source/Programs/AutomationTool/Scripts/BuildForUGS.Automation.cs
+++ b/Engine/Source/Programs/AutomationTool/Scripts/BuildForUGS.Automation.cs
@@ -74,30 +74,32 @@ namespace AutomationTool
 				Platform WindowsTargetPlatform = Platform.GetPlatform(UnrealTargetPlatform.Win64);
 
 				// Figure out all the files for the archive
-				Ionic.Zip.ZipFile Zip = new Ionic.Zip.ZipFile();
-				Zip.UseZip64WhenSaving = Ionic.Zip.Zip64Option.Always;
-				foreach(string BuildProduct in Builder.BuildProductFiles)
+				string ZipFileName = Path.Combine(OutputFolder, "Archive.zip");
+				using (Ionic.Zip.ZipFile Zip = new Ionic.Zip.ZipFile())
 				{
-					if(!File.Exists(BuildProduct))
-					{
-						throw new AutomationException("Missing build product: {0}", BuildProduct);
-					}
-					if(BuildProduct.EndsWith(".pdb", StringComparison.InvariantCultureIgnoreCase))
+					Zip.UseZip64WhenSaving = Ionic.Zip.Zip64Option.Always;
+					foreach(string BuildProduct in Builder.BuildProductFiles)
 					{
-						string StrippedFileName = CommandUtils.MakeRerootedFilePath(BuildProduct, CommandUtils.CmdEnv.LocalRoot, SymbolsFolder);
-						Directory.CreateDirectory(Path.GetDirectoryName(StrippedFileName));
-						WindowsTargetPlatform.StripSymbols(new FileReference(BuildProduct), new FileReference(StrippedFileName));
-						Zip.AddFile(StrippedFileName, Path.GetDirectoryName(CommandUtils.StripBaseDirectory(StrippedFileName, SymbolsFolder)));
-					}
-					else
-					{
-						Zip.AddFile(BuildProduct, Path.GetDirectoryName(CommandUtils.StripBaseDirectory(BuildProduct, CommandUtils.CmdEnv.LocalRoot)));
+						if(!File.Exists(BuildProduct))
+						{
+							throw new AutomationException("Missing build product: {0}", BuildProduct);
+						}
+						if(BuildProduct.EndsWith(".pdb", StringComparison.InvariantCultureIgnoreCase))
+						{
+							string StrippedFileName = CommandUtils.MakeRerootedFilePath(BuildProduct, CommandUtils.CmdEnv.LocalRoot, SymbolsFolder);
+							Directory.CreateDirectory(Path.GetDirectoryName(StrippedFileName));
+							WindowsTargetPlatform.StripSymbols(new FileReference(BuildProduct), new FileReference(StrippedFileName));
+							Zip.AddFile(StrippedFileName, Path.GetDirectoryName(CommandUtils.StripBaseDirectory(StrippedFileName, SymbolsFolder)));
+						}
+						else
+						{
+							Zip.AddFile(BuildProduct, Path.GetDirectoryName(CommandUtils.StripBaseDirectory(BuildProduct, CommandUtils.CmdEnv.LocalRoot)));
+						}
 					}
+					// Create the zip file
+					Console.WriteLine("Writing {0}...", ZipFileName);
+					Zip.Save(ZipFileName);
 				}
-				// Create the zip file
-				string ZipFileName = Path.Combine(OutputFolder, "Archive.zip");
-				Console.WriteLine("Writing {0}...", ZipFileName);
-				Zip.Save(ZipFileName);
 
 				// Submit it to Perforce if required
 				if(CommandUtils.AllowSubmit)
diff --git a/Engine/Source/Programs/AutomationTool/Scripts/BuildPhysX.Automation.cs b/Engine/Source/Programs/AutomationTool/Scripts/BuildPhysX.Automation.cs
--- a/Engine/Source/Programs/AutomationTool/Scripts/BuildPhysX.Automation.cs
+++ b/Engine/Source/Programs/AutomationTool/Scripts/BuildPhysX.Automation.cs
@@ -1187,7 +1187,7 @@ class BuildPhysX : BuildCommand
 			case UnrealTargetPlatform.Mac:
 				return "dSYM";
 		}
-		throw new AutomationException(String.Format("No debug database extension for platform ", TargetData.Platform.ToString()));
+		throw new AutomationException(String.Format("No debug database extension for platform '{0}'", TargetData.Platform.ToString()));
 	}
 
 	private static string GetPlatformBinaryExtension(TargetPlatformData TargetData)
@@ -1200,7 +1200,7 @@ class BuildPhysX : BuildCommand
 			case UnrealTargetPlatform.Mac:
 				return "dylib";
 		}
-		throw new AutomationException(String.Format("No binary extension for platform ", TargetData.Platform.ToString()));
+		throw new AutomationException(String.Format("No binary extension for platform '{0}'", TargetData.Platform.ToString()));
 	}
 
 	private static string GetPlatformLibExtension(TargetPlatformData TargetData)
@@ -1222,7 +1222,7 @@ class BuildPhysX : BuildCommand
 			case UnrealTargetPlatform.HTML5:
 				return "bc";
 		}
-		throw new AutomationException(String.Format("No lib extension for platform ", TargetData.Platform.ToString()));
+		throw new AutomationException(String.Format("No lib extension for platform '{0}'", TargetData.Platform.ToString()));
 	}
 
     private static bool FileGeneratedByLib(string FileNameUpper, PhysXTargetLib TargetLib)
diff --git a/Engine/Source/Programs/AutomationTool/Scripts/BuildPluginCommand.Automation.cs b/Engine/Source/Programs/AutomationTool/Scripts/BuildPluginCommand.Automation.cs
--- a/Engine/Source/Programs/AutomationTool/Scripts/BuildPluginCommand.Automation.cs
+++ b/Engine/Source/Programs/AutomationTool/Scripts/BuildPluginCommand.Automation.cs
@@ -245,7 +245,7 @@ class BuildPlugin : BuildCommand
 		if(FileReference.Exists(FilterFile))
 		{
 			CommandUtils.Log("Reading filter rules from {0}", FilterFile);
-			Filter.ReadRulesFromFile(FilterFile.FullName, "FilterPlugin");
+			Filter.ReadRulesFromFile(FilterFile, "FilterPlugin");
 		}
 
 		// Apply the standard exclusion rules
diff --git a/Engine/Source/Programs/AutomationTool/Scripts/CheckCsprojDotNetVersion.cs b/Engine/Source/Programs/AutomationTool/Scripts/CheckCsprojDotNetVersion.cs
new file mode 100644
--- /dev/null
+++ b/Engine/Source/Programs/AutomationTool/Scripts/CheckCsprojDotNetVersion.cs
@@ -0,0 +1,105 @@
+﻿// Copyright 1998-2017 Epic Games, Inc. All Rights Reserved.
+using System;
+using System.Collections.Generic;
+using System.Text;
+using System.IO;
+using AutomationTool;
+using UnrealBuildTool;
+using System.Text.RegularExpressions;
+
+class CheckCsprojDotNetVersion : BuildCommand
+{
+	public override void ExecuteBuild()
+	{
+		// just calling this DesiredTargetVersion because TargetVersion and TargetedVersion get super confusing.
+		string DesiredTargetVersion = this.ParseParamValue("TargetVersion", null);
+		if (string.IsNullOrEmpty(DesiredTargetVersion))
+		{
+			throw new AutomationException("-TargetVersion was not specified.");
+		}
+		CommandUtils.Log("Scanning for all csproj's...");
+		// Check for all csproj's in the engine dir
+		DirectoryReference EngineDir = CommandUtils.EngineDirectory;
+
+		// grab the targeted version.,
+		Regex FrameworkRegex = new Regex("<TargetFrameworkVersion>v(\\d\\.\\d\\.?\\d?)<\\/TargetFrameworkVersion>");
+		Regex PossibleAppConfigRegex = new Regex("<TargetFrameworkProfile>(.+)<\\/TargetFrameworkProfile>");
+		Regex AppConfigRegex = new Regex("<supportedRuntime version=\"v(\\d\\.\\d\\.?\\d?)\" sku=\"\\.NETFramework,Version=v(\\d\\.\\d\\.?\\d?),Profile=(.+)\"\\/>");
+		foreach (FileReference CsProj in DirectoryReference.EnumerateFiles(EngineDir, "*.csproj", SearchOption.AllDirectories))
+		{
+			if (CsProj.ContainsName(new FileSystemName("ThirdParty"), EngineDir) ||
+				(CsProj.ContainsName(new FileSystemName("UE4TemplateProject"), EngineDir) && CsProj.GetFileName().Equals("ProjectTemplate.csproj")))
+			{
+				continue;
+			}
+
+			// read in the file
+			string Contents = File.ReadAllText(CsProj.FullName);
+			Match m = FrameworkRegex.Match(Contents);
+			if (m.Success)
+			{
+				string TargetedVersion = m.Groups[1].Value;
+				// make sure we match, throw warning otherwise
+				if (!DesiredTargetVersion.Equals(TargetedVersion, StringComparison.InvariantCultureIgnoreCase))
+				{
+					CommandUtils.LogWarning("Targeted Framework version for project: {0} was not {1}! Targeted Version: {2}", CsProj, DesiredTargetVersion, TargetedVersion);
+				}
+			}
+			// if we don't have a TargetFrameworkVersion, check for the existence of TargetFrameworkProfile.
+			else
+			{
+				m = PossibleAppConfigRegex.Match(Contents);
+				if (!m.Success)
+				{
+					CommandUtils.Log("No TargetFrameworkVersion or TargetFrameworkProfile found for project {0}, is it a mono project? If not, does it compile properly?", CsProj);
+					continue;
+				}
+
+				// look for the app config
+				FileReference AppConfigFile = FileReference.Combine(CsProj.Directory, "app.config");
+				string Profile = m.Groups[1].Value;
+				if (!FileReference.Exists(AppConfigFile))
+				{
+					CommandUtils.Log("Found TargetFrameworkProfile but no associated app.config containing the version for project {0}.", CsProj);
+					continue;
+				}
+
+				// read in the app config
+				Contents = File.ReadAllText(AppConfigFile.FullName);
+				m = AppConfigRegex.Match(Contents);
+				if (!m.Success)
+				{
+					CommandUtils.Log("Couldn't find a supportedRuntime match for the version in the app.config for project {0}.", CsProj);
+					continue;
+				}
+
+				// Version1 is the one that appears right after supportedRuntime
+				// Version2 is the one in the sku
+				// ProfileString should match the TargetFrameworkProfile from the csproj
+				string Version1String = m.Groups[1].Value;
+				string Version2String = m.Groups[2].Value;
+				string ProfileString = m.Groups[3].Value;
+
+				// not sure how this is possible, but check for it anyway
+				if (!ProfileString.Equals(Profile, StringComparison.InvariantCultureIgnoreCase))
+				{
+					CommandUtils.LogWarning("The TargetFrameworkProfile in csproj {0} ({1}) doesn't match the sku in it's app.config ({2}).", CsProj, Profile, ProfileString);
+					continue;
+				}
+
+				// if the version numbers don't match the app.config is probably corrupt.
+				if (!Version1String.Equals(Version2String, StringComparison.InvariantCultureIgnoreCase))
+				{
+					CommandUtils.LogWarning("The supportedRunTimeVersion ({0}) and the sku version ({1}) in the app.config for project {2} don't match.", Version1String, Version2String, CsProj);
+					continue;
+				}
+
+				// make sure the versions match
+				if (!(DesiredTargetVersion.Equals(Version1String, StringComparison.InvariantCultureIgnoreCase)))
+				{
+					CommandUtils.LogWarning("Targeted Framework version for project: {0} was not {1}! Targeted Version: {2}", CsProj, DesiredTargetVersion, Version1String);
+				}
+			}
+		}
+	}
+}
\ No newline at end of file
diff --git a/Engine/Source/Programs/AutomationTool/Scripts/CookCommand.Automation.cs b/Engine/Source/Programs/AutomationTool/Scripts/CookCommand.Automation.cs
--- a/Engine/Source/Programs/AutomationTool/Scripts/CookCommand.Automation.cs
+++ b/Engine/Source/Programs/AutomationTool/Scripts/CookCommand.Automation.cs
@@ -598,7 +598,10 @@ public partial class Project : CommandUtils
 		}
 
 		const bool bQuiet = true;
-		DeleteDirectory(bQuiet, CleanDirs);
+		foreach(string CleanDir in CleanDirs)
+		{
+			DeleteDirectory(bQuiet, CleanDir);
+		}
 	}
 
 	#endregion
diff --git a/Engine/Source/Programs/AutomationTool/Scripts/CopyBuildToStagingDirectory.Automation.cs b/Engine/Source/Programs/AutomationTool/Scripts/CopyBuildToStagingDirectory.Automation.cs
--- a/Engine/Source/Programs/AutomationTool/Scripts/CopyBuildToStagingDirectory.Automation.cs
+++ b/Engine/Source/Programs/AutomationTool/Scripts/CopyBuildToStagingDirectory.Automation.cs
@@ -240,7 +240,7 @@ public partial class Project : CommandUtils
             PotentialParentCultures.Add(LocaleTags[0]);
         }
 
-        string[] FoundDirectories = CommandUtils.FindDirectories(true, "*", false, new string[] { SourceDirectory });
+        string[] FoundDirectories = CommandUtils.FindDirectories(true, "*", false, SourceDirectory);
         foreach (string FoundDirectory in FoundDirectories)
         {
             string DirectoryName = CommandUtils.GetLastDirectoryName(FoundDirectory);
@@ -291,7 +291,7 @@ public partial class Project : CommandUtils
 			SC.StageFiles(StagedFileType.UFS, PlatformCookDir, "*", false, ExcludeWildCards, SC.RelativeProjectRootForStage, true, !Params.UsePak(SC.StageTargetPlatform));
 
 			// Stage each sub directory separately so that we can skip Engine if need be
-			string[] SubDirs = CommandUtils.FindDirectories(true, "*", false, new string[] { PlatformCookDir });
+			string[] SubDirs = CommandUtils.FindDirectories(true, "*", false, PlatformCookDir);
 
             foreach (string SubDir in SubDirs)
             {
@@ -352,13 +352,7 @@ public partial class Project : CommandUtils
 
 
             // Initialize internationalization preset.
-            string InternationalizationPreset = null;
-
-            // Use parameters if provided.
-            if (string.IsNullOrEmpty(InternationalizationPreset))
-            {
-                InternationalizationPreset = Params.InternationalizationPreset;
-            }
+            string InternationalizationPreset = Params.InternationalizationPreset;
 
             // Use configuration if otherwise lacking an internationalization preset.
             if (string.IsNullOrEmpty(InternationalizationPreset))
@@ -438,7 +432,7 @@ public partial class Project : CommandUtils
 				var ProjectLocRootDirectory = CombinePaths(SC.ProjectRoot, "Content/Localization");
 				if (DirectoryExists(ProjectLocRootDirectory))
 				{
-					string[] ProjectLocTargetDirectories = CommandUtils.FindDirectories(true, "*", false, new string[] { ProjectLocRootDirectory });
+					string[] ProjectLocTargetDirectories = CommandUtils.FindDirectories(true, "*", false, ProjectLocRootDirectory);
 					foreach (string ProjectLocTargetDirectory in ProjectLocTargetDirectories)
 					{
 						StageLocalizationDataForTarget(SC, CulturesToStage, ProjectLocTargetDirectory, null, !Params.UsePak(SC.StageTargetPlatform));
diff --git a/Engine/Source/Programs/AutomationTool/Scripts/CopyUAT.Automation.cs b/Engine/Source/Programs/AutomationTool/Scripts/CopyUAT.Automation.cs
--- a/Engine/Source/Programs/AutomationTool/Scripts/CopyUAT.Automation.cs
+++ b/Engine/Source/Programs/AutomationTool/Scripts/CopyUAT.Automation.cs
@@ -16,10 +16,10 @@ public class CopyUAT : BuildCommand
 	public override void ExecuteBuild()
 	{
 		// Get the output directory
-		string TargetDir = ParseParamValue("TargetDir");
-		if(TargetDir == null)
+		string TargetDirParam = ParseParamValue("TargetDir");
+		if(TargetDirParam == null)
 		{
-			throw new AutomationException("Missing -Target argument to CopyUAT");
+			throw new AutomationException("Missing -TargetDir=... argument to CopyUAT");
 		}
 
 		// Construct a dummy UE4Build object to get a list of the UAT and UBT build products
@@ -32,24 +32,28 @@ public class CopyUAT : BuildCommand
 		Build.AddUBTFilesToBuildProducts();
 
 		// Get a list of all the input files
-		List<string> FileNames = new List<string>(Build.BuildProductFiles);
-		foreach(string FileName in Build.BuildProductFiles)
+		List<FileReference> SourceFiles = new List<FileReference>();
+		foreach(string BuildProductFile in Build.BuildProductFiles)
 		{
-			string SymbolFileName = Path.ChangeExtension(FileName, ".pdb");
-			if(File.Exists(SymbolFileName))
+			FileReference SourceFile = new FileReference(BuildProductFile);
+			SourceFiles.Add(SourceFile);
+
+			FileReference SourceSymbolFile = SourceFile.ChangeExtension(".pdb");
+			if(FileReference.Exists(SourceSymbolFile))
 			{
-				FileNames.Add(SymbolFileName);
+				SourceFiles.Add(SourceSymbolFile);
 			}
 		}
 
 		// Copy all the files over
-		foreach(string FileName in FileNames)
+		DirectoryReference TargetDir = new DirectoryReference(TargetDirParam);
+		foreach(FileReference SourceFile in SourceFiles)
 		{
-			string TargetFileName = Utils.MakeRerootedFilePath(FileName, CommandUtils.CmdEnv.LocalRoot, TargetDir);
-			Directory.CreateDirectory(Path.GetDirectoryName(TargetFileName));
-			File.Copy(FileName, TargetFileName);
+			FileReference TargetFile = FileReference.Combine(TargetDir, SourceFile.MakeRelativeTo(CommandUtils.RootDirectory));
+			DirectoryReference.CreateDirectory(TargetFile.Directory);
+			CommandUtils.CopyFile(SourceFile.FullName, TargetFile.FullName);
 		}
 
-		Log("Copied {0} files to {1}", FileNames.Count, TargetDir);
+		Log("Copied {0} files to {1}", SourceFiles.Count, TargetDir);
 	}
 }
diff --git a/Engine/Source/Programs/AutomationTool/Scripts/LauncherLocalization.Automation.cs b/Engine/Source/Programs/AutomationTool/Scripts/LauncherLocalization.Automation.cs
--- a/Engine/Source/Programs/AutomationTool/Scripts/LauncherLocalization.Automation.cs
+++ b/Engine/Source/Programs/AutomationTool/Scripts/LauncherLocalization.Automation.cs
@@ -36,8 +36,6 @@ class LauncherLocalization : BuildCommand
             P4.Sync(P4Env.BuildRootP4 + "/Portal/Config/...");
             P4.Sync(P4Env.BuildRootP4 + "/Portal/Content/...");
             P4.Sync(P4Env.BuildRootP4 + "/Portal/Source/...");
-
-			Log("Localize from label {0}", P4Env.LabelToSync);
 		}
 
 		OneSkyConfigData OneSkyConfig = OneSkyConfigHelper.Find("OneSkyConfig_EpicGames");
@@ -69,7 +67,7 @@ class LauncherLocalization : BuildCommand
 
 		// Setup commandlet arguments for SCC.
 		string CommandletSCCArguments = String.Empty;
-		if (P4Enabled) { CommandletSCCArguments += (string.IsNullOrEmpty(CommandletSCCArguments) ? "" : " ") + "-EnableSCC"; }
+		if (P4Enabled) { CommandletSCCArguments += "-EnableSCC"; }
 		if (!AllowSubmit) { CommandletSCCArguments += (string.IsNullOrEmpty(CommandletSCCArguments) ? "" : " ") + "-DisableSCCSubmit"; }
 
 		// Setup commandlet arguments with configurations.
diff --git a/Engine/Source/Programs/AutomationTool/Scripts/PackageCommand.Automation.cs b/Engine/Source/Programs/AutomationTool/Scripts/PackageCommand.Automation.cs
--- a/Engine/Source/Programs/AutomationTool/Scripts/PackageCommand.Automation.cs
+++ b/Engine/Source/Programs/AutomationTool/Scripts/PackageCommand.Automation.cs
@@ -13,7 +13,7 @@ public partial class Project : CommandUtils
 
 	public static void Package(ProjectParams Params, int WorkingCL=-1)
 	{
-		if ((!Params.SkipStage || Params.Package))
+		if (!Params.SkipStage || Params.Package)
 		{
 			Params.ValidateAndLog();
 			List<DeploymentContext> DeployContextList = new List<DeploymentContext>();
diff --git a/Engine/Source/Programs/AutomationTool/Scripts/RebuildLightMapsCommand.Automation.cs b/Engine/Source/Programs/AutomationTool/Scripts/RebuildLightMapsCommand.Automation.cs
--- a/Engine/Source/Programs/AutomationTool/Scripts/RebuildLightMapsCommand.Automation.cs
+++ b/Engine/Source/Programs/AutomationTool/Scripts/RebuildLightMapsCommand.Automation.cs
@@ -140,7 +140,7 @@ namespace AutomationScripts.Automation
 			catch (Exception Ex)
 			{
                 string FinalLogLines = "No log file found";
-                AutomationException AEx = Ex as AutomationException;
+                CommandletException AEx = Ex as CommandletException;
                 if ( AEx != null )
                 {
                     string LogFile = AEx.LogFileName;
diff --git a/Engine/Source/Programs/AutomationTool/Scripts/ResavePackagesCommand.Automation.cs b/Engine/Source/Programs/AutomationTool/Scripts/ResavePackagesCommand.Automation.cs
--- a/Engine/Source/Programs/AutomationTool/Scripts/ResavePackagesCommand.Automation.cs
+++ b/Engine/Source/Programs/AutomationTool/Scripts/ResavePackagesCommand.Automation.cs
@@ -139,7 +139,7 @@ namespace AutomationScripts.Automation
 			catch (Exception Ex)
 			{
                 string FinalLogLines = "No log file found";
-                AutomationException AEx = Ex as AutomationException;
+                CommandletException AEx = Ex as CommandletException;
                 if ( AEx != null )
                 {
                     string LogFile = AEx.LogFileName;
diff --git a/Engine/Source/Programs/AutomationTool/Scripts/RunProjectCommand.Automation.cs b/Engine/Source/Programs/AutomationTool/Scripts/RunProjectCommand.Automation.cs
--- a/Engine/Source/Programs/AutomationTool/Scripts/RunProjectCommand.Automation.cs
+++ b/Engine/Source/Programs/AutomationTool/Scripts/RunProjectCommand.Automation.cs
@@ -333,7 +333,7 @@ public partial class Project : CommandUtils
 					ClientProcess.StopProcess();
 					Thread.Sleep(10000);
 				}
-				while (ClientLogReader != null && !ClientLogReader.EndOfStream)
+				while (!ClientLogReader.EndOfStream)
 				{
 					string ClientOutput = ClientLogReader.ReadToEnd();
 					if (!String.IsNullOrEmpty(ClientOutput))
@@ -537,7 +537,7 @@ public partial class Project : CommandUtils
 				if (ClientProcess == null && !String.IsNullOrEmpty(Output))
 				{
 					AllClientOutput += Output;
-					if (ClientProcess == null && (AllClientOutput.Contains("Game Engine Initialized") || AllClientOutput.Contains("Unreal Network File Server is ready")))
+					if (AllClientOutput.Contains("Game Engine Initialized") || AllClientOutput.Contains("Unreal Network File Server is ready"))
 					{
 						Log("Starting Client....");
 						var SC = DeployContextList[0];
@@ -843,7 +843,7 @@ public partial class Project : CommandUtils
 		if (Params.CrashIndex > 0)
 		{
 			int RealIndex = Params.CrashIndex - 1;
-			if (RealIndex < 0 || RealIndex >= CrashCommands.Count())
+			if (RealIndex >= CrashCommands.Count())
 			{
 				throw new AutomationException("CrashIndex {0} is out of range...max={1}", Params.CrashIndex, CrashCommands.Count());
 			}
diff --git a/Engine/Source/Programs/AutomationTool/Scripts/StreamCopyDescription.Automation.cs b/Engine/Source/Programs/AutomationTool/Scripts/StreamCopyDescription.Automation.cs
--- a/Engine/Source/Programs/AutomationTool/Scripts/StreamCopyDescription.Automation.cs
+++ b/Engine/Source/Programs/AutomationTool/Scripts/StreamCopyDescription.Automation.cs
@@ -61,6 +61,21 @@ namespace AutomationTool
 			// Remove changes by the build machine
 			Descriptions = Regex.Replace(Descriptions, "[^\n]*buildmachine\n(\n|\t[^\n]*\n)*", "");
 
+			// Remove all the tags we don't care about
+			Descriptions = Regex.Replace(Descriptions, "^[ \t]*#(rb|fyi|codereview|lockdown)\\s.*$", "", RegexOptions.Multiline);
+
+			// Empty out lines which just contain whitespace
+			Descriptions = Regex.Replace(Descriptions, "^[ \t]+$", "", RegexOptions.Multiline);
+
+			// Remove multiple consecutive blank lines
+			Descriptions = Regex.Replace(Descriptions, "\n\n+", "\n\n");
+
+			// Only include one newline at the end of each description
+			Descriptions = Regex.Replace(Descriptions, "\n+Change", "\n\nChange");
+
+			// Remove merge-only changelists
+			Descriptions = Regex.Replace(Descriptions, "(?<=(^|\\n))Change .*\\s*Merging .* to .*\\s*\\n(?=(Change|$))", "");
+
 			// Figure out the target stream
 			IProcessResult StreamResult = P4.P4(String.Format("stream -o {0}", Stream), AllowSpew: false);
 			if (StreamResult.ExitCode != 0)
@@ -89,11 +104,7 @@ namespace AutomationTool
 
 				foreach (string Line in Descriptions.Split('\n'))
 				{
-					string TrimLine = Line.TrimStart();
-					if(!TrimLine.StartsWith("#codereview", StringComparison.OrdinalIgnoreCase) && !TrimLine.StartsWith("#rb", StringComparison.OrdinalIgnoreCase) && !TrimLine.StartsWith("#lockdown", StringComparison.OrdinalIgnoreCase))
-					{
-						Writer.WriteLine(Line);
-					}
+					Writer.WriteLine(Line);
 				}
 			}
 			Log("Written {0}.", OutputFileName);
diff --git a/Engine/Source/Programs/AutomationTool/Scripts/Tests.Automation.cs b/Engine/Source/Programs/AutomationTool/Scripts/Tests.Automation.cs
--- a/Engine/Source/Programs/AutomationTool/Scripts/Tests.Automation.cs
+++ b/Engine/Source/Programs/AutomationTool/Scripts/Tests.Automation.cs
@@ -778,8 +778,8 @@ class TestFileUtility : BuildCommand
 {
 	public override void ExecuteBuild()
 	{
-		string[] DummyFilenames = new string[] { ParseParamValue("file", "") };
-		DeleteFile(DummyFilenames);
+		string DummyFilename = ParseParamValue("file", "");
+		DeleteFile(DummyFilename);
 	}
 }
 
diff --git a/Engine/Source/Programs/AutomationTool/Scripts/UnrealSync.Automation.cs b/Engine/Source/Programs/AutomationTool/Scripts/UnrealSync.Automation.cs
--- a/Engine/Source/Programs/AutomationTool/Scripts/UnrealSync.Automation.cs
+++ b/Engine/Source/Programs/AutomationTool/Scripts/UnrealSync.Automation.cs
@@ -83,11 +83,11 @@ namespace AutomationScripts.Automation
 		{
 			var Combined = P4.GetLabels(BranchPath + "/Promot*" + (GameName != null ? ("-" + GameName) : "") + "-CL-*");
 
-			Combined.OrderByDescending((Label) => Label.Date);
+			IEnumerable<P4Label> CombinedOrdered = Combined.OrderByDescending((Label) => Label.Date);
 
 			var Output = new List<P4Label>();
 
-			foreach(var PossiblePromotable in Combined)
+			foreach(var PossiblePromotable in CombinedOrdered)
 			{
 				if(PossiblePromotable.Name.StartsWith(BranchPath + "/Promoted"))
 				{
diff --git a/Engine/Source/Programs/AutomationTool/Scripts/ZipUtils.Automation.cs b/Engine/Source/Programs/AutomationTool/Scripts/ZipUtils.Automation.cs
--- a/Engine/Source/Programs/AutomationTool/Scripts/ZipUtils.Automation.cs
+++ b/Engine/Source/Programs/AutomationTool/Scripts/ZipUtils.Automation.cs
@@ -16,178 +16,180 @@ using Ionic.Zip;
 [Help("compression=<0..9>", "Compression Level 0 - Copy  9-Best Compression")]
 public class ZipUtils : BuildCommand
 {
-    public override ExitCode Execute()
-    {
-        
-        Log("************************ STARTING ZIPUTILS ************************");
-
-        if (Params.Length < 2)
-        {
-            LogError("Invalid number of arguments {0}", Params.Length);
-            return ExitCode.Error_Arguments;
-        }
-        
-        string ZipFilePath = ParseParamValue("archive","");
-        var ExtractArgs = ParseParamValues("extract");
-        var AddArgs = ParseParamValues("add");
-        int CompressionLevel = ParseParamInt("compression", 5);
-
-
-        if (string.IsNullOrWhiteSpace(ZipFilePath) || System.IO.Path.GetExtension(ZipFilePath) != ".zip")
-        {
-            LogError("No zip file specified");
-            return ExitCode.Error_Arguments;
-        }
-
-        bool bUnzip = ExtractArgs.Count() > 0 ;
-        bool bZip = AddArgs.Count() > 0;
-
-        if(!bUnzip && !bZip)
-        {
-            LogError("Invalid arguments. Please specify -archive or -extract option.");
-            return ExitCode.Error_Arguments;
-        }
-
-        //setup filter to include all files
-        FileFilter Filter = new FileFilter();
-        Filter.Include("*.*");
-
-        if (bUnzip)
-        {
-            string OutputFolderPath = ExtractArgs[0];
-
-            Log("Running unzip : {0} OuputFolder:{1}", ZipFilePath.ToString(), OutputFolderPath);
-
-            if (!System.IO.File.Exists(ZipFilePath))
-            {
-                LogError("Invalid zip file path.");
-                return ExitCode.Error_Arguments;
-            }
-
-            var ExtractedFiles = InternalUnzipFiles(ZipFilePath, OutputFolderPath);
-
-            if(ExtractedFiles.Count() == 0)
-            {
-                LogWarning("No files extracted from file zip.");
-                return ExitCode.Error_Unknown;
-            }
-            else
-            {
-                Log("List of files extracted:");
-                foreach (var file in ExtractedFiles)
-                {
-                    Log("\t{0}",file);
-                }
-            }
-        }
-        else /*(bZip)*/
-        {
-
-            string ArchiveFolder = AddArgs[0];
-            Log("Compress level : {0}", CompressionLevel);
-            Log("Running zip : {0}", ArchiveFolder.ToString());
-            
-            FileAttributes attr = File.GetAttributes(ArchiveFolder);
-
-            if (!System.IO.Directory.Exists(ArchiveFolder) || !attr.HasFlag(FileAttributes.Directory))
-            {
-                LogError("Invalid zip file path.");
-                return ExitCode.Error_Arguments;
-            }
-
-           InternalZipFiles(ZipFilePath, ArchiveFolder, Filter, CompressionLevel);           
-        }
-
-        Log("************************ ZIPUTIL WORK COMPLETED ************************");
-
-        return ExitCode.Success;
-    }
-
-    /// <summary>
-    /// Creates a zip file containing the given input files
-    /// </summary>
-    /// <param name="ZipFileName">Filename for the zip</param>
-    /// <param name="Filter">Filter which selects files to be included in the zip</param>
-    /// <param name="BaseDirectory">Base directory to store relative paths in the zip file to</param>
-    /// <param name="CopyModeOnly">No compression will be done. Only acts like a container. The default value is set to false.</param>
-    internal static void InternalZipFiles(string ZipFileName, string BaseDirectory, FileFilter Filter, int CompressionLevel = 0)
-    {
-        // Ionic.Zip.Zip64Option.Always option produces broken archives on Mono, so we use system zip tool instead
-        if (Utils.IsRunningOnMono)
-        {
-            CommandUtils.CreateDirectory(Path.GetDirectoryName(ZipFileName));
-            CommandUtils.PushDir(BaseDirectory);
-            string FilesList = "";
-            foreach (string FilteredFile in Filter.ApplyToDirectory(BaseDirectory, true))
-            {
-                FilesList += " \"" + FilteredFile + "\"";
-                if (FilesList.Length > 32000)
-                {
-                    CommandUtils.RunAndLog(CommandUtils.CmdEnv, "zip", "-g -q \"" + ZipFileName + "\"" + FilesList);
-                    FilesList = "";
-                }
-            }
-            if (FilesList.Length > 0)
-            {
-                CommandUtils.RunAndLog(CommandUtils.CmdEnv, "zip", "-g -q \"" + ZipFileName + "\"" + FilesList);
-            }
-            CommandUtils.PopDir();
-        }
-        else
-        {
-            Ionic.Zip.ZipFile Zip = new Ionic.Zip.ZipFile();
-            Zip.UseZip64WhenSaving = Ionic.Zip.Zip64Option.Always;
-
-            Zip.CompressionLevel = (Ionic.Zlib.CompressionLevel)CompressionLevel;
-
-            if (Zip.CompressionLevel == Ionic.Zlib.CompressionLevel.Level0)            
-                Zip.CompressionMethod = Ionic.Zip.CompressionMethod.None;
-
-            foreach (string FilteredFile in Filter.ApplyToDirectory(BaseDirectory, true))
-            {
-                Zip.AddFile(Path.Combine(BaseDirectory, FilteredFile), Path.GetDirectoryName(FilteredFile));
-            }
-            CommandUtils.CreateDirectory(Path.GetDirectoryName(ZipFileName));
-            Zip.Save(ZipFileName);
-        }
-    }
-
-    /// <summary>
-    /// Extracts the contents of a zip file
-    /// </summary>
-    /// <param name="ZipFileName">Name of the zip file</param>
-    /// <param name="BaseDirectory">Output directory</param>
-    /// <returns>List of files written</returns>
-    public static IEnumerable<string> InternalUnzipFiles(string ZipFileName, string BaseDirectory)
-    {
-        // manually extract the files. There was a problem with the Ionic.Zip library that required this on non-PC at one point,
-        // but that problem is now fixed. Leaving this code as is as we need to return the list of created files and fix up their permissions anyway.
-        using (Ionic.Zip.ZipFile Zip = new Ionic.Zip.ZipFile(ZipFileName))
-        {
-            List<string> OutputFileNames = new List<string>();
-            foreach (Ionic.Zip.ZipEntry Entry in Zip.Entries)
-            {
-                string OutputFileName = Path.Combine(BaseDirectory, Entry.FileName);
-                Directory.CreateDirectory(Path.GetDirectoryName(OutputFileName));
-
-                //Check to make sure that we do not output to file stream if the entry is a directory.
-                if (!Entry.IsDirectory)
-                {
-                    using (FileStream OutputStream = new FileStream(OutputFileName, FileMode.Create, FileAccess.Write))
-                    {
-                        Entry.Extract(OutputStream);
-                    }
-                }
-
-                if (UnrealBuildTool.Utils.IsRunningOnMono && CommandUtils.IsProbablyAMacOrIOSExe(OutputFileName))
-                {
-                    FixUnixFilePermissions(OutputFileName);
-                }
-                OutputFileNames.Add(OutputFileName);
-            }
-            return OutputFileNames;
-        }
-    }
+	public override ExitCode Execute()
+	{
+
+		Log("************************ STARTING ZIPUTILS ************************");
+
+		if (Params.Length < 2)
+		{
+			LogError("Invalid number of arguments {0}", Params.Length);
+			return ExitCode.Error_Arguments;
+		}
+
+		string ZipFilePath = ParseParamValue("archive", "");
+		var ExtractArgs = ParseParamValues("extract");
+		var AddArgs = ParseParamValues("add");
+		int CompressionLevel = ParseParamInt("compression", 5);
+
+
+		if (string.IsNullOrWhiteSpace(ZipFilePath) || System.IO.Path.GetExtension(ZipFilePath) != ".zip")
+		{
+			LogError("No zip file specified");
+			return ExitCode.Error_Arguments;
+		}
+
+		bool bUnzip = ExtractArgs.Count() > 0;
+		bool bZip = AddArgs.Count() > 0;
+
+		if (!bUnzip && !bZip)
+		{
+			LogError("Invalid arguments. Please specify -archive or -extract option.");
+			return ExitCode.Error_Arguments;
+		}
+
+		//setup filter to include all files
+		FileFilter Filter = new FileFilter();
+		Filter.Include("*.*");
+
+		if (bUnzip)
+		{
+			string OutputFolderPath = ExtractArgs[0];
+
+			Log("Running unzip : {0} OuputFolder:{1}", ZipFilePath.ToString(), OutputFolderPath);
+
+			if (!System.IO.File.Exists(ZipFilePath))
+			{
+				LogError("Invalid zip file path.");
+				return ExitCode.Error_Arguments;
+			}
+
+			var ExtractedFiles = InternalUnzipFiles(ZipFilePath, OutputFolderPath);
+
+			if (ExtractedFiles.Count() == 0)
+			{
+				LogWarning("No files extracted from file zip.");
+				return ExitCode.Error_Unknown;
+			}
+			else
+			{
+				Log("List of files extracted:");
+				foreach (var file in ExtractedFiles)
+				{
+					Log("\t{0}", file);
+				}
+			}
+		}
+		else /*(bZip)*/
+		{
+
+			string ArchiveFolder = AddArgs[0];
+			Log("Compress level : {0}", CompressionLevel);
+			Log("Running zip : {0}", ArchiveFolder.ToString());
+
+			FileAttributes attr = File.GetAttributes(ArchiveFolder);
+
+			if (!System.IO.Directory.Exists(ArchiveFolder) || !attr.HasFlag(FileAttributes.Directory))
+			{
+				LogError("Invalid zip file path.");
+				return ExitCode.Error_Arguments;
+			}
+
+			InternalZipFiles(new FileReference(ZipFilePath), new DirectoryReference(ArchiveFolder), Filter, CompressionLevel);
+		}
+
+		Log("************************ ZIPUTIL WORK COMPLETED ************************");
+
+		return ExitCode.Success;
+	}
+
+	/// <summary>
+	/// Creates a zip file containing the given input files
+	/// </summary>
+	/// <param name="ZipFileName">Filename for the zip</param>
+	/// <param name="Filter">Filter which selects files to be included in the zip</param>
+	/// <param name="BaseDirectory">Base directory to store relative paths in the zip file to</param>
+	/// <param name="CopyModeOnly">No compression will be done. Only acts like a container. The default value is set to false.</param>
+	internal static void InternalZipFiles(FileReference ZipFileName, DirectoryReference BaseDirectory, FileFilter Filter, int CompressionLevel = 0)
+	{
+		// Ionic.Zip.Zip64Option.Always option produces broken archives on Mono, so we use system zip tool instead
+		if (Utils.IsRunningOnMono)
+		{
+			DirectoryReference.CreateDirectory(ZipFileName.Directory);
+			CommandUtils.PushDir(BaseDirectory.FullName);
+			string FilesList = "";
+			foreach (FileReference FilteredFile in Filter.ApplyToDirectory(BaseDirectory, true))
+			{
+				FilesList += " \"" + FilteredFile.MakeRelativeTo(BaseDirectory) + "\"";
+				if (FilesList.Length > 32000)
+				{
+					CommandUtils.RunAndLog(CommandUtils.CmdEnv, "zip", "-g -q \"" + ZipFileName + "\"" + FilesList);
+					FilesList = "";
+				}
+			}
+			if (FilesList.Length > 0)
+			{
+				CommandUtils.RunAndLog(CommandUtils.CmdEnv, "zip", "-g -q \"" + ZipFileName + "\"" + FilesList);
+			}
+			CommandUtils.PopDir();
+		}
+		else
+		{
+			using (Ionic.Zip.ZipFile Zip = new Ionic.Zip.ZipFile())
+			{
+				Zip.UseZip64WhenSaving = Ionic.Zip.Zip64Option.Always;
+
+				Zip.CompressionLevel = (Ionic.Zlib.CompressionLevel)CompressionLevel;
+
+				if (Zip.CompressionLevel == Ionic.Zlib.CompressionLevel.Level0)
+					Zip.CompressionMethod = Ionic.Zip.CompressionMethod.None;
+
+				foreach (FileReference FilteredFile in Filter.ApplyToDirectory(BaseDirectory, true))
+				{
+					Zip.AddFile(FilteredFile.FullName, Path.GetDirectoryName(FilteredFile.MakeRelativeTo(BaseDirectory)));
+				}
+				CommandUtils.CreateDirectory(Path.GetDirectoryName(ZipFileName.FullName));
+				Zip.Save(ZipFileName.FullName);
+			}
+		}
+	}
+
+	/// <summary>
+	/// Extracts the contents of a zip file
+	/// </summary>
+	/// <param name="ZipFileName">Name of the zip file</param>
+	/// <param name="BaseDirectory">Output directory</param>
+	/// <returns>List of files written</returns>
+	public static IEnumerable<string> InternalUnzipFiles(string ZipFileName, string BaseDirectory)
+	{
+		// manually extract the files. There was a problem with the Ionic.Zip library that required this on non-PC at one point,
+		// but that problem is now fixed. Leaving this code as is as we need to return the list of created files and fix up their permissions anyway.
+		using (Ionic.Zip.ZipFile Zip = new Ionic.Zip.ZipFile(ZipFileName))
+		{
+			List<string> OutputFileNames = new List<string>();
+			foreach (Ionic.Zip.ZipEntry Entry in Zip.Entries)
+			{
+				string OutputFileName = Path.Combine(BaseDirectory, Entry.FileName);
+				Directory.CreateDirectory(Path.GetDirectoryName(OutputFileName));
+
+				//Check to make sure that we do not output to file stream if the entry is a directory.
+				if (!Entry.IsDirectory)
+				{
+					using (FileStream OutputStream = new FileStream(OutputFileName, FileMode.Create, FileAccess.Write))
+					{
+						Entry.Extract(OutputStream);
+					}
+				}
+
+				if (UnrealBuildTool.Utils.IsRunningOnMono && CommandUtils.IsProbablyAMacOrIOSExe(OutputFileName))
+				{
+					FixUnixFilePermissions(OutputFileName);
+				}
+				OutputFileNames.Add(OutputFileName);
+			}
+			return OutputFileNames;
+		}
+	}
 }
 
 
diff --git a/Engine/Source/Programs/CrashReporter/CrashReportCommon/AmazonClient.cs b/Engine/Source/Programs/CrashReporter/CrashReportCommon/AmazonClient.cs
--- a/Engine/Source/Programs/CrashReporter/CrashReportCommon/AmazonClient.cs
+++ b/Engine/Source/Programs/CrashReporter/CrashReportCommon/AmazonClient.cs
@@ -9,6 +9,7 @@ using Amazon.S3;
 using Amazon.S3.Model;
 using Amazon.SQS;
 using Amazon.SQS.Model;
+using Amazon.S3.IO;
 
 namespace Tools.CrashReporter.CrashReportCommon
 {
@@ -164,6 +165,35 @@ namespace Tools.CrashReporter.CrashReportCommon
 			return S3Client.GetObject(ObjectRequest);
 		}
 
+        /// <summary>
+        /// Retrieves an URL with the appropriate signature to access the file for a predetermined amount of time
+        /// </summary>
+        /// <param name="BucketName">S3 bucket name</param>
+        /// <param name="Key">S3 object key</param>
+        /// <param name="expirationInMinutes">Expiration of link in minutes</param>
+        /// <returns>The URL that can be used to access the file</returns>
+        public string GetS3SignedUrl(string BucketName, string Key, int expirationInMinutes)
+        {
+            if (!IsS3Valid) return null;
+
+            string url = null;
+
+            S3FileInfo s3FileInfo = new Amazon.S3.IO.S3FileInfo(S3Client, BucketName, Key);
+            if (s3FileInfo.Exists)
+            {
+                url = S3Client.GetPreSignedURL(new GetPreSignedUrlRequest
+                {
+                    BucketName = BucketName,
+                    Key = Key,
+                    Expires = DateTime.UtcNow.AddMinutes(expirationInMinutes)
+                    //ResponseHeaderOverrides = new ResponseHeaderOverrides() { ContentDisposition = "attachment; filename=\"test.txt\"", }
+                });
+
+            }            
+
+            return url;
+        }
+
 		/// <summary>
 		/// Uploads a file to S3 synchronously (blocks until complete).
 		/// </summary>
diff --git a/Engine/Source/Programs/CrashReporter/CrashReportCommon/CrashContext.cs b/Engine/Source/Programs/CrashReporter/CrashReportCommon/CrashContext.cs
--- a/Engine/Source/Programs/CrashReporter/CrashReportCommon/CrashContext.cs
+++ b/Engine/Source/Programs/CrashReporter/CrashReportCommon/CrashContext.cs
@@ -157,14 +157,36 @@ namespace Tools.CrashReporter.CrashReportCommon
 		{
 		}
 
+		/// <summary>
+		/// Generates a crash context rom a Url
+		/// </summary>
+		/// <param name="Url"></param>
+		/// <returns></returns>
+        public static FGenericCrashContext FromUrl(string Url)
+        {
+            using (var wc = new System.Net.WebClient()) {
+                var CopyText = wc.DownloadString(Url);
+
+                return FromString(CopyText);
+            }
+        }
+
+		/// <summary> Creates a new instance of the crash report context from the specified file. </summary>
+        public static FGenericCrashContext FromFile(string Filepath)
+        {
+            string[] XMLText = File.ReadAllLines(Filepath);
+            string CopyText = string.Join(Environment.NewLine, XMLText);
+
+            return FromString(CopyText);
+        }
+
 		/// <summary> Creates a new instance of the crash report context from the specified file. </summary>
-		public static FGenericCrashContext FromFile( string Filepath )
+		public static FGenericCrashContext FromString(string Text)
 		{
 			// The crash report context contains invalid characters, we need to fix them all in order to parse the XML.
 			HashSet<string> TagsToIgnore = new HashSet<string>() { "<FGenericCrashContext>", "<RuntimeProperties>", "<PlatformProperties>", "</FGenericCrashContext>", "</RuntimeProperties>", "</PlatformProperties>" };
 
-			string[] XMLText = File.ReadAllLines( Filepath );
-			string CopyText = string.Join( Environment.NewLine, XMLText );
+            string[] XMLText = Text.Split(new string[] { Environment.NewLine }, StringSplitOptions.None);
 
 			string NodeOpen = "";
 			string NodeEnd = "";
diff --git a/Engine/Source/Programs/CrashReporter/CrashReportProcess/Config.cs b/Engine/Source/Programs/CrashReporter/CrashReportProcess/Config.cs
--- a/Engine/Source/Programs/CrashReporter/CrashReportProcess/Config.cs
+++ b/Engine/Source/Programs/CrashReporter/CrashReportProcess/Config.cs
@@ -362,6 +362,12 @@ namespace Tools.CrashReporter.CrashReportProcess
 		[XmlElement]
 		public bool MonitorPerformance { get; set; }
 
+		/// <summary>
+		/// List of GameNames to blacklist. Since we take crashes from external games, we need to filter out stuff we don't want, especially high volume stuff.
+		/// </summary>
+		[XmlElement]
+		public string GameNamesToBlacklist { get; set; }
+
 		/// <summary>
 		/// Get the default config object (lazy loads it on first access)
 		/// </summary>
diff --git a/Engine/Source/Programs/CrashReporter/CrashReportProcess/ReportProcessor.cs b/Engine/Source/Programs/CrashReporter/CrashReportProcess/ReportProcessor.cs
--- a/Engine/Source/Programs/CrashReporter/CrashReportProcess/ReportProcessor.cs
+++ b/Engine/Source/Programs/CrashReporter/CrashReportProcess/ReportProcessor.cs
@@ -122,7 +122,10 @@ namespace Tools.CrashReporter.CrashReportProcess
 		{
 			ProcessorTask = new Task(() =>
 			                         {
-				                         while (!CancelSource.IsCancellationRequested)
+										 // parse out hte list of blacklisted games.
+										 string[] GameNamesToBlacklist = (Config.Default.GameNamesToBlacklist == null ? "" : Config.Default.GameNamesToBlacklist).ToLowerInvariant().Split(',');
+
+										 while (!CancelSource.IsCancellationRequested)
 				                         {
 					                         try
 					                         {
@@ -133,7 +136,18 @@ namespace Tools.CrashReporter.CrashReportProcess
 							                         FGenericCrashContext NewContext = null;
 							                         if (Queue.TryDequeueReport(out NewContext))
 							                         {
-								                         ProcessReport(NewContext);
+														 bool bIsBlacklistedGame = GameNamesToBlacklist.Contains(NewContext.PrimaryCrashProperties.GameName.ToLowerInvariant());
+														 // if it's blacklisted, skip it, else process it.
+														 if (bIsBlacklistedGame)
+														 {
+															 CrashReporterProcessServicer.WriteEvent(string.Format("Discarding crash from blacklisted GameName '{0}'", NewContext.PrimaryCrashProperties.GameName));
+															 // delete the report from disk since we don't care about it.
+															 FinalizeReport(AddReportResult.Added, new DirectoryInfo(NewContext.CrashDirectory), NewContext);
+														 }
+														 else
+														 {
+															 ProcessReport(NewContext);
+														 }
 
 								                         // The effect of this break is to prioritize ReportQueues by their order in the list, from highest to lowest
 								                         bIdle = false;
diff --git a/Engine/Source/Programs/CrashReporter/CrashReportWebSite/Classes/PerformanceTracker.cs b/Engine/Source/Programs/CrashReporter/CrashReportWebSite/Classes/PerformanceTracker.cs
new file mode 100644
--- /dev/null
+++ b/Engine/Source/Programs/CrashReporter/CrashReportWebSite/Classes/PerformanceTracker.cs
@@ -0,0 +1,86 @@
+﻿using System;
+using System.Collections.Generic;
+using System.Linq;
+using System.Web;
+using Tools.CrashReporter.CrashReportCommon;
+using Tools.CrashReporter.CrashReportWebSite.Properties;
+
+namespace Tools.CrashReporter.CrashReportWebSite.Classes
+{
+    public class PerformanceTracker
+    {
+        private int _count = 0;
+        private int _reportsThreshold = 100;
+
+        private Dictionary<string, double> _actionTimeDictionary;
+
+        private SlackWriter _slackWriter;
+
+        public PerformanceTracker()
+        {
+            _slackWriter = new SlackWriter()
+            {
+                WebhookUrl = Settings.Default.SlackWebhookUrl,
+                Channel = Settings.Default.SlackChannel,
+                Username = Settings.Default.SlackUsername,
+                IconEmoji = Settings.Default.SlackEmoji
+            };
+            _actionTimeDictionary = new Dictionary<string, double>();
+        }
+
+        public PerformanceTracker(int reportsThreshold = 100)
+        {
+            _reportsThreshold = reportsThreshold;
+            _slackWriter = new SlackWriter()
+            {
+                WebhookUrl = Settings.Default.SlackWebhookUrl,
+                Channel = Settings.Default.SlackChannel,
+                Username = Settings.Default.SlackUsername,
+                IconEmoji = Settings.Default.SlackEmoji
+            };
+            _actionTimeDictionary = new Dictionary<string, double>();
+        }
+
+        public void AddStat(string actionName, double elapsedTime)
+        {
+            if (_actionTimeDictionary.ContainsKey(actionName))
+            {
+                _actionTimeDictionary[actionName] += elapsedTime;
+            }
+            else
+            {
+                _actionTimeDictionary.Add(actionName, elapsedTime);
+            }
+        }
+
+        public void ClearStats()
+        {
+            _actionTimeDictionary.Clear();
+            _count = 0;
+        }
+
+        public void IncrementCount()
+        {
+            _count++;
+            if (_count >= _reportsThreshold)
+            {
+                WritePerfReport();
+                ClearStats();
+            }
+        }
+
+        public void WritePerfReport()
+        {
+            foreach (var stat in _actionTimeDictionary)
+            {
+                _slackWriter.Write("Average time for " + stat.Key + " over last " + _reportsThreshold.ToString("N") + " crashes. : " + (stat.Value / _reportsThreshold).ToString("F2"));
+            }
+            _count -= _reportsThreshold;
+        }
+
+        public int GetCount()
+        {
+            return _count;
+        }
+    }
+}
\ No newline at end of file
