Merging //UE4/Dev-Main to Dev-VREditor (//UE4/Dev-VREditor) 3350064

diff --git a/Engine/Source/Programs/AutomationTool/AutomationUtils/UE4Build.cs b/Engine/Source/Programs/AutomationTool/AutomationUtils/UE4Build.cs
--- a/Engine/Source/Programs/AutomationTool/AutomationUtils/UE4Build.cs
+++ b/Engine/Source/Programs/AutomationTool/AutomationUtils/UE4Build.cs
@@ -177,6 +177,7 @@ namespace AutomationTool
 			{
 				AddArgs += " " + CommandUtils.MakePathSafeToUseWithCommandLine(UprojectPath.FullName);
 			}
+			AddArgs += " -NoUBTMakefiles";
 			AddArgs += " " + InAddArgs;
 			if (ForceMonolithic)
 			{
@@ -283,6 +284,7 @@ namespace AutomationTool
 			{
 				AddArgs += " " + CommandUtils.MakePathSafeToUseWithCommandLine(UprojectPath.FullName);
 			}
+			AddArgs += " -NoUBTMakefiles";
 			AddArgs += " " + InAddArgs;
 			if (ForceMonolithic)
 			{
@@ -319,6 +321,7 @@ namespace AutomationTool
 			{
 				AddArgs += " " + CommandUtils.MakePathSafeToUseWithCommandLine(UprojectPath.FullName);
 			}
+			AddArgs += " -NoUBTMakefiles";
 			AddArgs += " " + InAddArgs;
 			if (ForceMonolithic)
 			{
diff --git a/Engine/Source/Programs/AutomationTool/Scripts/CopyUAT.Automation.cs b/Engine/Source/Programs/AutomationTool/Scripts/CopyUAT.Automation.cs
--- a/Engine/Source/Programs/AutomationTool/Scripts/CopyUAT.Automation.cs
+++ b/Engine/Source/Programs/AutomationTool/Scripts/CopyUAT.Automation.cs
@@ -43,6 +43,12 @@ public class CopyUAT : BuildCommand
 			{
 				SourceFiles.Add(SourceSymbolFile);
 			}
+
+			FileReference DocumentationFile = SourceFile.ChangeExtension(".xml");
+			if(FileReference.Exists(DocumentationFile))
+			{
+				SourceFiles.Add(DocumentationFile);
+			}
 		}
 
 		// Copy all the files over
diff --git a/Engine/Source/Programs/UnrealBuildTool/Configuration/UEBuildTarget.cs b/Engine/Source/Programs/UnrealBuildTool/Configuration/UEBuildTarget.cs
--- a/Engine/Source/Programs/UnrealBuildTool/Configuration/UEBuildTarget.cs
+++ b/Engine/Source/Programs/UnrealBuildTool/Configuration/UEBuildTarget.cs
@@ -2025,7 +2025,7 @@ namespace UnrealBuildTool
 				if (SharedBuildIdFile != null)
 				{
 					string VariablePath = TargetReceipt.InsertPathVariables(SharedBuildIdFile, UnrealBuildTool.EngineDirectory, ProjectDirectory);
-					Receipt.PrecompiledBuildDependencies.Add(VariablePath);
+					Receipt.BuildProducts.Add(new BuildProduct(VariablePath, BuildProductType.BuildResource));
 				}
 			}
 
diff --git a/Engine/Source/Programs/UnrealBuildTool/System/TargetReceipt.cs b/Engine/Source/Programs/UnrealBuildTool/System/TargetReceipt.cs
--- a/Engine/Source/Programs/UnrealBuildTool/System/TargetReceipt.cs
+++ b/Engine/Source/Programs/UnrealBuildTool/System/TargetReceipt.cs
@@ -50,6 +50,11 @@ namespace UnrealBuildTool
 		/// A resource file which was generated by the build and is required for the executable to run.
 		/// </summary>
 		RequiredResource,
+
+		/// <summary>
+		/// A build resource which was generated by the build, but is not required for the executable to run.
+		/// </summary>
+		BuildResource,
 	}
 
 	/// <summary>
