Merging CL 3347480 from Dev-Core to fix incremental compiles in development streams. Disable UBT makefiles when compiling from UAT. We're typically building as part of a formal build here, we are less concerned about fast and potentially unreliable iteration.

diff --git a/Engine/Source/Programs/AutomationTool/AutomationUtils/UE4Build.cs b/Engine/Source/Programs/AutomationTool/AutomationUtils/UE4Build.cs
--- a/Engine/Source/Programs/AutomationTool/AutomationUtils/UE4Build.cs
+++ b/Engine/Source/Programs/AutomationTool/AutomationUtils/UE4Build.cs
@@ -177,6 +177,7 @@ namespace AutomationTool
 			{
 				AddArgs += " " + CommandUtils.MakePathSafeToUseWithCommandLine(UprojectPath.FullName);
 			}
+			AddArgs += " -NoUBTMakefiles";
 			AddArgs += " " + InAddArgs;
 			if (ForceMonolithic)
 			{
@@ -283,6 +284,7 @@ namespace AutomationTool
 			{
 				AddArgs += " " + CommandUtils.MakePathSafeToUseWithCommandLine(UprojectPath.FullName);
 			}
+			AddArgs += " -NoUBTMakefiles";
 			AddArgs += " " + InAddArgs;
 			if (ForceMonolithic)
 			{
@@ -319,6 +321,7 @@ namespace AutomationTool
 			{
 				AddArgs += " " + CommandUtils.MakePathSafeToUseWithCommandLine(UprojectPath.FullName);
 			}
+			AddArgs += " -NoUBTMakefiles";
 			AddArgs += " " + InAddArgs;
 			if (ForceMonolithic)
 			{
