Fix "UAT & UBT Documentation" build step. UBT is normally compiled on a different agent, so the generated XML documentation was not available when we come to run this step. Include XML files matching assembly names that are compiled in the UAT build products.

diff --git a/Engine/Source/Programs/AutomationTool/Scripts/CopyUAT.Automation.cs b/Engine/Source/Programs/AutomationTool/Scripts/CopyUAT.Automation.cs
--- a/Engine/Source/Programs/AutomationTool/Scripts/CopyUAT.Automation.cs
+++ b/Engine/Source/Programs/AutomationTool/Scripts/CopyUAT.Automation.cs
@@ -43,6 +43,12 @@ public class CopyUAT : BuildCommand
 			{
 				SourceFiles.Add(SourceSymbolFile);
 			}
+
+			FileReference DocumentationFile = SourceFile.ChangeExtension(".xml");
+			if(FileReference.Exists(DocumentationFile))
+			{
+				SourceFiles.Add(DocumentationFile);
+			}
 		}
 
 		// Copy all the files over
