Fix the letters 'D', 'e', 'f', 'a', 'u', 'l', 't' being stripped from the end of directory names when copying directories in UAT. Broken code didn't use to do any harm because it wasn't assigned to anything (.NET string is immutable) - fixing that due to a static analyzer warning caused it to be activated.

diff --git a/Engine/Source/Programs/AutomationTool/AutomationUtils/CommandUtils.cs b/Engine/Source/Programs/AutomationTool/AutomationUtils/CommandUtils.cs
--- a/Engine/Source/Programs/AutomationTool/AutomationUtils/CommandUtils.cs
+++ b/Engine/Source/Programs/AutomationTool/AutomationUtils/CommandUtils.cs
@@ -1362,7 +1362,7 @@ namespace AutomationTool
 		{
 			Source = ConvertSeparators(PathSeparator.Default, Source);
 			Dest = ConvertSeparators(PathSeparator.Default, Dest);
-			Dest = Dest.TrimEnd(PathSeparator.Default.ToString().ToCharArray());
+			Dest = Dest.TrimEnd(new char[] { Path.DirectorySeparatorChar, Path.AltDirectorySeparatorChar });
 
 			if (InternalUtils.SafeDirectoryExists(Dest))
 			{
