UAT: Fix null dereference when parsing project params for staging. (CIS fix)

diff --git a/Engine/Source/Programs/AutomationTool/AutomationUtils/ProjectParams.cs b/Engine/Source/Programs/AutomationTool/AutomationUtils/ProjectParams.cs
--- a/Engine/Source/Programs/AutomationTool/AutomationUtils/ProjectParams.cs
+++ b/Engine/Source/Programs/AutomationTool/AutomationUtils/ProjectParams.cs
@@ -569,7 +569,7 @@ namespace AutomationTool
 			this.CookAll = GetParamValueIfNotSpecified(Command, CookAll, this.CookAll, "CookAll");
 			this.CookPartialGC = GetParamValueIfNotSpecified(Command, CookPartialGC, this.CookPartialGC, "CookPartialGC");
 			this.CookInEditor = GetParamValueIfNotSpecified(Command, CookInEditor, this.CookInEditor, "CookInEditor");
-			this.CookOutputDir = String.IsNullOrEmpty(CookOutputDir)? Command.ParseParamValue("CookOutputDir") : CookOutputDir;
+			this.CookOutputDir = ParseParamValueIfNotSpecified(Command, CookOutputDir, "CookOutputDir", String.Empty, true);
 			this.CookMapsOnly = GetParamValueIfNotSpecified(Command, CookMapsOnly, this.CookMapsOnly, "CookMapsOnly");
 			this.FileServer = GetParamValueIfNotSpecified(Command, FileServer, this.FileServer, "fileserver");
 			this.DedicatedServer = GetParamValueIfNotSpecified(Command, DedicatedServer, this.DedicatedServer, "dedicatedserver", "server");
