Change documentation build steps to sync tools from //depot/UnrealDocTool rather than using outdated versions in the engine branch. Remove Doxygen and HTML Help Workshop from //UE4/Main.

As part of this, added a new UAT command ("SyncDepotPath") to create a temporary workspace, sync a folder, and delete the workspace.

diff --git a/Engine/Source/Programs/AutomationTool/Scripts/AutomationScripts.Automation.csproj b/Engine/Source/Programs/AutomationTool/Scripts/AutomationScripts.Automation.csproj
--- a/Engine/Source/Programs/AutomationTool/Scripts/AutomationScripts.Automation.csproj
+++ b/Engine/Source/Programs/AutomationTool/Scripts/AutomationScripts.Automation.csproj
@@ -103,6 +103,7 @@
     <Compile Include="StageTargetCommand.Automation.cs" />
     <Compile Include="StreamCopyDescription.Automation.cs" />
     <Compile Include="SubmitUtililzationReportToEC.Automation.cs" />
+    <Compile Include="SyncDepotPath.cs" />
     <Compile Include="Tests.Automation.cs" />
     <Compile Include="TestSnippet.Automation.cs" />
     <Compile Include="UE4BuildUtils.cs" />
