UAT: Fix launching projects that aren't directly under the UE4 root directory on non-Windows platforms. (Merge from CL 3397949 in //UE4/Release-4.16)

diff --git a/Engine/Source/Programs/AutomationTool/AutomationUtils/Platform.cs b/Engine/Source/Programs/AutomationTool/AutomationUtils/Platform.cs
--- a/Engine/Source/Programs/AutomationTool/AutomationUtils/Platform.cs
+++ b/Engine/Source/Programs/AutomationTool/AutomationUtils/Platform.cs
@@ -236,15 +236,22 @@ namespace AutomationTool
 			{
 				if (SC.StageTargets.Count() > 0)
 				{
-					DirectoryReference LocalRoot = new DirectoryReference(SC.LocalRoot);
+					DirectoryReference ProjectRoot = new DirectoryReference(SC.ProjectRoot);
 					foreach (StageTarget Target in SC.StageTargets)
 					{
 						foreach (BuildProduct Product in Target.Receipt.BuildProducts)
 						{
 							if (Product.Type == BuildProductType.Executable)
 							{
-								string RelativeExePath = new FileReference(Product.Path).MakeRelativeTo(LocalRoot);
-								ExecutableNames.Add(Path.Combine(SC.RuntimeRootDir, RelativeExePath));
+								FileReference BuildProductFile = new FileReference(Product.Path);
+								if (BuildProductFile.IsUnderDirectory(ProjectRoot))
+								{
+									ExecutableNames.Add(CombinePaths(SC.RuntimeProjectRootDir, BuildProductFile.MakeRelativeTo(ProjectRoot)));
+								}
+								else
+								{
+									ExecutableNames.Add(CombinePaths(SC.RuntimeRootDir, BuildProductFile.MakeRelativeTo(RootDirectory)));
+								}
 							}
 						}
 					}
