Merging //UE4/Dev-Main to Release-Staging-4.16 (//UE4/Release-Staging-4.16)

diff --git a/Engine/Source/Programs/AutomationTool/AutomationUtils/P4Utils.cs b/Engine/Source/Programs/AutomationTool/AutomationUtils/P4Utils.cs
--- a/Engine/Source/Programs/AutomationTool/AutomationUtils/P4Utils.cs
+++ b/Engine/Source/Programs/AutomationTool/AutomationUtils/P4Utils.cs
@@ -8,6 +8,8 @@ using System.ComponentModel;
 using System.Diagnostics;
 using System.Text.RegularExpressions;
 using Tools.DotNETCommon.CaselessDictionary;
+using System.Reflection;
+using System.Collections;
 
 namespace AutomationTool
 {
@@ -325,6 +327,602 @@ namespace AutomationTool
 		}
 	}
 
+	/// <summary>
+	/// Describes the action performed by the user when resolving the integration
+	/// </summary>
+	public enum P4IntegrateAction
+	{
+		/// <summary>
+		/// file did not previously exist; it was created as a copy of partner-file
+		/// </summary>
+		[Description("branch from")]
+		BranchFrom,
+
+		/// <summary>
+		/// partner-file did not previously exist; it was created as a copy of file.
+		/// </summary>
+		[Description("branch into")]
+		BranchInto,
+
+		/// <summary>
+		/// file was integrated from partner-file, accepting merge.
+		/// </summary>
+		[Description("merge from")]
+		MergeFrom,
+
+		/// <summary>
+		/// file was integrated into partner-file, accepting merge.
+		/// </summary>
+		[Description("merge into")]
+		MergeInto,
+
+		/// <summary>
+		/// file was integrated from partner-file, accepting theirs and deleting the original.
+		/// </summary>
+		[Description("moved from")]
+		MovedFrom,
+
+		/// <summary>
+		/// file was integrated into partner-file, accepting theirs and creating partner-file if it did not previously exist.
+		/// </summary>
+		[Description("moved into")]
+		MovedInto,
+
+		/// <summary>
+		/// file was integrated from partner-file, accepting theirs.
+		/// </summary>
+		[Description("copy from")]
+		CopyFrom,
+
+		/// <summary>
+		/// file was integrated into partner-file, accepting theirs.
+		/// </summary>
+		[Description("copy into")]
+		CopyInto,
+
+		/// <summary>
+		/// file was integrated from partner-file, accepting yours.
+		/// </summary>
+		[Description("ignored")]
+		Ignored,
+
+		/// <summary>
+		/// file was integrated into partner-file, accepting yours.
+		/// </summary>
+		[Description("ignored by")]
+		IgnoredBy,
+
+		/// <summary>
+		/// file was integrated from partner-file, and partner-file had been previously deleted.
+		/// </summary>
+		[Description("delete from")]
+		DeleteFrom,
+
+		/// <summary>
+		/// file was integrated into partner-file, and file had been previously deleted.
+		/// </summary>
+		[Description("delete into")]
+		DeleteInto,
+
+		/// <summary>
+		/// file was integrated from partner-file, and file was edited within the p4 resolve process.
+		/// </summary>
+		[Description("edit from")]
+		EditFrom,
+
+		/// <summary>
+		/// file was integrated into partner-file, and partner-file was reopened for edit before submission.
+		/// </summary>
+		[Description("edit into")]
+		EditInto,
+
+		/// <summary>
+		/// file was integrated from a deleted partner-file, and partner-file was reopened for add (that is, someone restored a deleted file by syncing back to a pre-deleted revision and adding the file).
+		/// </summary>
+		[Description("add from")]
+		AddFrom,
+
+		/// <summary>
+		/// file was integrated into previously nonexistent partner-file, and partner-file was reopened for add before submission.
+		/// </summary>
+		[Description("add into")]
+		AddInto,
+	}
+
+	/// <summary>
+	/// Stores integration information for a file revision
+	/// </summary>
+	public class P4IntegrationRecord
+	{
+		/// <summary>
+		/// The integration action performed for this file
+		/// </summary>
+		public readonly P4IntegrateAction Action;
+
+		/// <summary>
+		/// The partner file for this integration
+		/// </summary>
+		public readonly string OtherFile;
+
+		/// <summary>
+		/// Min revision of the partner file for this integration
+		/// </summary>
+		public readonly int StartRevisionNumber;
+
+		/// <summary>
+		/// Max revision of the partner file for this integration
+		/// </summary>
+		public readonly int EndRevisionNumber;
+
+		/// <summary>
+		/// Constructor
+		/// </summary>
+		/// <param name="Action">The integration action</param>
+		/// <param name="OtherFile">The partner file involved in the integration</param>
+		/// <param name="StartRevisionNumber">Starting revision of the partner file for the integration (exclusive)</param>
+		/// <param name="EndRevisionNumber">Ending revision of the partner file for the integration (inclusive)</param>
+		public P4IntegrationRecord(P4IntegrateAction Action, string OtherFile, int StartRevisionNumber, int EndRevisionNumber)
+		{
+			this.Action = Action;
+			this.OtherFile = OtherFile;
+			this.StartRevisionNumber = StartRevisionNumber;
+			this.EndRevisionNumber = EndRevisionNumber;
+		}
+
+		/// <summary>
+		/// Summarize this record for display in the debugger
+		/// </summary>
+		/// <returns>Formatted integration record</returns>
+		public override string ToString()
+		{
+			if(StartRevisionNumber + 1 == EndRevisionNumber)
+			{
+				return String.Format("{0} {1}#{2}", Action, OtherFile, EndRevisionNumber);
+			}
+			else
+			{
+				return String.Format("{0} {1}#{2},#{3}", Action, OtherFile, StartRevisionNumber + 1, EndRevisionNumber);
+			}
+		}
+	}
+
+	/// <summary>
+	/// Stores a revision record for a file
+	/// </summary>
+	public class P4RevisionRecord
+	{
+		/// <summary>
+		/// The revision number of this file
+		/// </summary>
+		public readonly int RevisionNumber;
+
+		/// <summary>
+		/// The changelist responsible for this revision of the file
+		/// </summary>
+		public readonly int ChangeNumber;
+
+		/// <summary>
+		/// Action performed to the file in this revision
+		/// </summary>
+		public readonly P4Action Action;
+
+		/// <summary>
+		/// Type of the file
+		/// </summary>
+		public readonly string Type;
+
+		/// <summary>
+		/// Timestamp of this modification
+		/// </summary>
+		public readonly DateTime DateTime;
+
+		/// <summary>
+		/// Author of the changelist
+		/// </summary>
+		public readonly string UserName;
+
+		/// <summary>
+		/// Client that submitted this changelist
+		/// </summary>
+		public readonly string ClientName;
+
+		/// <summary>
+		/// Size of the file, or -1 if not specified
+		/// </summary>
+		public readonly int FileSize;
+
+		/// <summary>
+		/// Digest of the file, or null if not specified
+		/// </summary>
+		public readonly string Digest;
+
+		/// <summary>
+		/// Description of this changelist
+		/// </summary>
+		public readonly string Description;
+
+		/// <summary>
+		/// Integration records for this revision
+		/// </summary>
+		public readonly P4IntegrationRecord[] Integrations;
+
+		/// <summary>
+		/// Constructor
+		/// </summary>
+		/// <param name="RevisionNumber">Revision number of the file</param>
+		/// <param name="ChangeNumber">Number of the changelist that submitted this revision</param>
+		/// <param name="Action">Action performed to the file in this changelist</param>
+		/// <param name="Type">Type of the file</param>
+		/// <param name="DateTime">Timestamp for the change</param>
+		/// <param name="UserName">User that submitted the change</param>
+		/// <param name="ClientName">Client that submitted the change</param>
+		/// <param name="FileSize">Size of the file, or -1 if not specified</param>
+		/// <param name="Digest">Digest of the file, or null if not specified</param>
+		/// <param name="Description">Description of the changelist</param>
+		/// <param name="Integrations">Integrations performed to the file</param>
+		public P4RevisionRecord(int RevisionNumber, int ChangeNumber, P4Action Action, string Type, DateTime DateTime, string UserName, string ClientName, int FileSize, string Digest, string Description, P4IntegrationRecord[] Integrations)
+		{
+			this.RevisionNumber = RevisionNumber;
+			this.ChangeNumber = ChangeNumber;
+			this.Action = Action;
+			this.Type = Type;
+			this.DateTime = DateTime;
+			this.UserName = UserName;
+			this.ClientName = ClientName;
+			this.Description = Description;
+			this.FileSize = FileSize;
+			this.Digest = Digest;
+			this.Integrations = Integrations;
+		}
+
+		/// <summary>
+		/// Format this record for display in the debugger
+		/// </summary>
+		/// <returns>Summary of this revision</returns>
+		public override string ToString()
+		{
+			return String.Format("#{0} change {1} {2} on {3} by {4}@{5}", RevisionNumber, ChangeNumber, Action, DateTime, UserName, ClientName);
+		}
+	}
+
+	/// <summary>
+	/// Record output by the filelog command
+	/// </summary>
+	public class P4FileRecord
+	{
+		/// <summary>
+		/// Path to the file in the depot
+		/// </summary>
+		public string DepotPath;
+
+		/// <summary>
+		/// Revisions of this file
+		/// </summary>
+		public P4RevisionRecord[] Revisions;
+
+		/// <summary>
+		/// Constructor
+		/// </summary>
+		/// <param name="DepotPath">The depot path of the file</param>
+		/// <param name="Revisions">Revisions of the file</param>
+		public P4FileRecord(string DepotPath, P4RevisionRecord[] Revisions)
+		{
+			this.DepotPath = DepotPath;
+			this.Revisions = Revisions;
+		}
+
+		/// <summary>
+		/// Return the depot path of the file for display in the debugger
+		/// </summary>
+		/// <returns>Path to the file</returns>
+		public override string ToString()
+		{
+			return DepotPath;
+		}
+	}
+
+	/// <summary>
+	/// Options for the filelog command
+	/// </summary>
+	[Flags]
+	public enum P4FileLogOptions
+	{
+		/// <summary>
+		/// No options
+		/// </summary>
+		None = 0,
+
+		/// <summary>
+		/// Display file content history instead of file name history.
+		/// </summary>
+		ContentHistory = 1,
+
+		/// <summary>
+		/// Follow file history across branches.
+		/// </summary>
+		FollowAcrossBranches = 2,
+
+		/// <summary>
+		/// List long output, with the full text of each changelist description.
+		/// </summary>
+		FullDescriptions = 4,
+
+		/// <summary>
+		/// List long output, with the full text of each changelist description truncated at 250 characters.
+		/// </summary>
+		LongDescriptions = 8,
+
+		/// <summary>
+		/// When used with the ContentHistory option, do not follow content of promoted task streams. 
+		/// </summary>
+		DoNotFollowPromotedTaskStreams = 16,
+
+		/// <summary>
+		/// Display a shortened form of output by ignoring non-contributory integrations
+		/// </summary>
+		IgnoreNonContributoryIntegrations = 32,
+	}
+
+	/// <summary>
+	/// Type of a Perforce stream
+	/// </summary>
+	public enum P4StreamType
+	{
+		/// <summary>
+		/// A mainline stream
+		/// </summary>
+		Mainline,
+
+		/// <summary>
+		/// A development stream
+		/// </summary>
+		Development,
+
+		/// <summary>
+		/// A release stream
+		/// </summary>
+		Release,
+
+		/// <summary>
+		/// A virtual stream
+		/// </summary>
+		Virtual,
+
+		/// <summary>
+		/// A task stream
+		/// </summary>
+		Task,
+	}
+
+	/// <summary>
+	/// Options for a stream definition
+	/// </summary>
+	[Flags]
+	public enum P4StreamOptions
+	{
+		/// <summary>
+		/// The stream is locked
+		/// </summary>
+		Locked = 1,
+
+		/// <summary>
+		/// Only the owner may submit to the stream
+		/// </summary>
+		OwnerSubmit = 4,
+
+		/// <summary>
+		/// Integrations from this stream to its parent are expected
+		/// </summary>
+		ToParent = 4,
+
+		/// <summary>
+		/// Integrations from this stream from its parent are expected
+		/// </summary>
+		FromParent = 8,
+
+		/// <summary>
+		/// Undocumented?
+		/// </summary>
+		MergeDown = 16,
+	}
+
+	/// <summary>
+	/// Contains information about a stream, as returned by the 'p4 streams' command
+	/// </summary>
+	[DebuggerDisplay("{Stream}")]
+	public class P4StreamRecord
+	{
+		/// <summary>
+		/// Path to the stream
+		/// </summary>
+		public string Stream;
+
+		/// <summary>
+		/// Last time the stream definition was updated
+		/// </summary>
+		public DateTime Update;
+
+		/// <summary>
+		/// Last time the stream definition was accessed
+		/// </summary>
+		public DateTime Access;
+
+		/// <summary>
+		/// Owner of this stream
+		/// </summary>
+		public string Owner;
+
+		/// <summary>
+		/// Name of the stream. This may be modified after the stream is initially created, but it's underlying depot path will not change.
+		/// </summary>
+		public string Name;
+
+		/// <summary>
+		/// The parent stream
+		/// </summary>
+		public string Parent;
+
+		/// <summary>
+		/// Type of the stream
+		/// </summary>
+		public P4StreamType Type;
+
+		/// <summary>
+		/// User supplied description of the stream
+		/// </summary>
+		public string Description;
+
+		/// <summary>
+		/// Options for the stream definition
+		/// </summary>
+		public P4StreamOptions Options;
+
+		/// <summary>
+		/// Whether this stream is more stable than the parent stream
+		/// </summary>
+		public Nullable<bool> FirmerThanParent;
+
+		/// <summary>
+		/// Whether changes from this stream flow to the parent stream
+		/// </summary>
+		public bool ChangeFlowsToParent;
+
+		/// <summary>
+		/// Whether changes from this stream flow from the parent stream
+		/// </summary>
+		public bool ChangeFlowsFromParent;
+
+		/// <summary>
+		/// The mainline branch associated with this stream
+		/// </summary>
+		public string BaseParent;
+
+		/// <summary>
+		/// Constructor
+		/// </summary>
+		/// <param name="Stream">Path to the stream</param>
+		/// <param name="Update">Last time the stream definition was updated</param>
+		/// <param name="Access">Last time the stream definition was accessed</param>
+		/// <param name="Owner">Owner of this stream</param>
+		/// <param name="Name">Name of the stream. This may be modified after the stream is initially created, but it's underlying depot path will not change.</param>
+		/// <param name="Parent">The parent stream</param>
+		/// <param name="Type">Type of the stream</param>
+		/// <param name="Description">User supplied description of the stream</param>
+		/// <param name="Options">Options for the stream definition</param>
+		/// <param name="FirmerThanParent">Whether this stream is more stable than the parent stream</param>
+		/// <param name="ChangeFlowsToParent">Whether changes from this stream flow to the parent stream</param>
+		/// <param name="ChangeFlowsFromParent">Whether changes from this stream flow from the parent stream</param>
+		/// <param name="BaseParent">The mainline branch associated with this stream</param>
+		public P4StreamRecord(string Stream, DateTime Update, DateTime Access, string Owner, string Name, string Parent, P4StreamType Type, string Description, P4StreamOptions Options, Nullable<bool> FirmerThanParent, bool ChangeFlowsToParent, bool ChangeFlowsFromParent, string BaseParent)
+		{
+			this.Stream = Stream;
+			this.Update = Update;
+			this.Owner = Owner;
+			this.Name = Name;
+			this.Parent = Parent;
+			this.Type = Type;
+			this.Description = Description;
+			this.Options = Options;
+			this.FirmerThanParent = FirmerThanParent;
+			this.ChangeFlowsToParent = ChangeFlowsToParent;
+			this.ChangeFlowsFromParent = ChangeFlowsFromParent;
+			this.BaseParent = BaseParent;
+		}
+
+		/// <summary>
+		/// Return the path of this stream for display in the debugger
+		/// </summary>
+		/// <returns>Path to this stream</returns>
+		public override string ToString()
+		{
+			return Stream;
+		}
+	}
+
+	/// <summary>
+	/// Error severity codes. Taken from the p4java documentation.
+	/// </summary>
+	public enum P4SeverityCode
+	{
+		Empty = 0,
+		Info = 1,
+		Warning = 2,
+		Failed = 3,
+		Fatal = 4,
+	}
+
+	/// <summary>
+	/// Generic error codes that can be returned by the Perforce server. Taken from the p4java documentation.
+	/// </summary>
+	public enum P4GenericCode
+	{
+		None = 0,
+		Usage = 1,
+		Unknown = 2,
+		Context = 3,
+		Illegal = 4,
+		NotYet = 5,
+		Protect = 6,
+		Empty = 17,
+		Fault = 33,
+		Client = 34,
+		Admin = 35,
+		Config = 36,
+		Upgrade = 37,
+		Comm = 38,
+		TooBig = 39, 
+	}
+
+	/// <summary>
+	/// Represents a error return value from Perforce.
+	/// </summary>
+	public class P4ReturnCode
+	{
+		/// <summary>
+		/// The value of the "code" field returned by the server
+		/// </summary>
+		public string Code;
+
+		/// <summary>
+		/// The severity of this error
+		/// </summary>
+		public P4SeverityCode Severity;
+
+		/// <summary>
+		/// The generic error code associated with this message
+		/// </summary>
+		public P4GenericCode Generic;
+
+		/// <summary>
+		/// The message text
+		/// </summary>
+		public string Message;
+
+		/// <summary>
+		/// Constructor
+		/// </summary>
+		/// <param name="Code">The value of the "code" field returned by the server</param>
+		/// <param name="Severity">The severity of this error</param>
+		/// <param name="Generic">The generic error code associated with this message</param>
+		/// <param name="Message">The message text</param>
+		public P4ReturnCode(string Code, P4SeverityCode Severity, P4GenericCode Generic, string Message)
+		{
+			this.Code = Code;
+			this.Severity = Severity;
+			this.Generic = Generic;
+			this.Message = Message;
+		}
+
+		/// <summary>
+		/// Formats this error for display in the debugger
+		/// </summary>
+		/// <returns>String representation of this object</returns>
+		public override string ToString()
+		{
+			return String.Format("{0}: {1} (Generic={2})", Code, Message, Generic);
+		}
+	}
+
 	public partial class CommandUtils
 	{
 		#region Environment Setup
@@ -370,8 +968,6 @@ namespace AutomationTool
 		/// </summary>
 		static internal void InitP4Environment()
 		{
-			CheckP4Enabled();
-
 			// Temporary connection - will use only the currently set env vars to connect to P4
 			var DefaultConnection = new P4Connection(User: null, Client: null);
 			PerforceEnvironment = Automation.IsBuildMachine ? new P4Environment(DefaultConnection, CmdEnv) : new LocalP4Environment(DefaultConnection, CmdEnv);
@@ -382,8 +978,6 @@ namespace AutomationTool
 		/// </summary>
 		static internal void InitDefaultP4Connection()
 		{
-			CheckP4Enabled();
-
 			PerforceConnection = new P4Connection(User: P4Env.User, Client: P4Env.Client, ServerAndPort: P4Env.P4Port);
 		}
 
@@ -429,17 +1023,6 @@ namespace AutomationTool
 		}
 		private static bool? bP4CLRequired;
 
-		/// <summary>
-		/// Throws an exception when P4 is disabled. This should be called in every P4 function.
-		/// </summary>
-		internal static void CheckP4Enabled()
-		{
-			if (P4Enabled == false)
-			{
-				throw new AutomationException("P4 is not enabled.");
-			}
-		}
-
 		/// <summary>
 		/// Checks whether commands are allowed to submit files into P4.
 		/// </summary>
@@ -598,14 +1181,6 @@ namespace AutomationTool
 			}
 		}
 
-		/// <summary>
-		/// Throws an exception when P4 is disabled. This should be called in every P4 function.
-		/// </summary>
-		internal static void CheckP4Enabled()
-		{
-			CommandUtils.CheckP4Enabled();
-		}
-
 		/// <summary>
 		/// Shortcut to Run but with P4.exe as the program name.
 		/// </summary>
@@ -615,7 +1190,6 @@ namespace AutomationTool
 		/// <returns>Exit code</returns>
         public IProcessResult P4(string CommandLine, string Input = null, bool AllowSpew = true, bool WithClient = true, bool SpewIsVerbose = false)
 		{
-			CheckP4Enabled();
 			CommandUtils.ERunOptions RunOptions = AllowSpew ? CommandUtils.ERunOptions.AllowSpew : CommandUtils.ERunOptions.NoLoggingOfRunCommand;
 			if( SpewIsVerbose )
 			{
@@ -634,7 +1208,6 @@ namespace AutomationTool
 		/// <returns>True if succeeded, otherwise false.</returns>
         public bool P4Output(out string Output, string CommandLine, string Input = null, bool AllowSpew = true, bool WithClient = true)
 		{
-			CheckP4Enabled();
 			Output = "";
 
             var Result = P4(CommandLine, Input, AllowSpew, WithClient);
@@ -644,46 +1217,218 @@ namespace AutomationTool
 		}
 
 		/// <summary>
-		/// Calls p4 command and writes the output to a logfile.
+		/// Calls p4 and returns the output.
+		/// </summary>
+		/// <param name="Output">Output of the comman.</param>
+		/// <param name="CommandLine">Commandline for p4.</param>
+		/// <param name="Input">Stdin input.</param>
+		/// <param name="AllowSpew">Whether the command should spew.</param>
+		/// <returns>True if succeeded, otherwise false.</returns>
+        public bool P4Output(out string[] OutputLines, string CommandLine, string Input = null, bool AllowSpew = true, bool WithClient = true)
+		{
+			string Output;
+			bool bResult = P4Output(out Output, CommandLine, Input, AllowSpew, WithClient);
+
+			List<string> Lines = new List<string>();
+			for(int Idx = 0; Idx < Output.Length; )
+			{
+				int EndIdx = Output.IndexOf('\n', Idx);
+				if(EndIdx == -1)
+				{
+					Lines.Add(Output.Substring(Idx));
+					break;
+				}
+
+				if(EndIdx > Idx && Output[EndIdx - 1] == '\r')
+				{
+					Lines.Add(Output.Substring(Idx, EndIdx - Idx - 1));
+				}
+				else
+				{
+					Lines.Add(Output.Substring(Idx, EndIdx - Idx));
+				}
+
+				Idx = EndIdx + 1;
+			}
+			OutputLines = Lines.ToArray();
+
+			return bResult;
+		}
+
+		/// <summary>
+		/// Calls p4 command and writes the output to a logfile.
+		/// </summary>
+		/// <param name="CommandLine">Commandline to pass to p4.</param>
+		/// <param name="Input">Stdin input.</param>
+		/// <param name="AllowSpew">Whether the command is allowed to spew.</param>
+        public void LogP4(string CommandLine, string Input = null, bool AllowSpew = true, bool WithClient = true, bool SpewIsVerbose = false)
+		{
+			string Output;
+            if (!LogP4Output(out Output, CommandLine, Input, AllowSpew, WithClient, SpewIsVerbose:SpewIsVerbose))
+			{
+				throw new P4Exception("p4.exe {0} failed.", CommandLine);
+			}
+		}
+
+		/// <summary>
+		/// Calls p4 and returns the output and writes it also to a logfile.
+		/// </summary>
+		/// <param name="Output">Output of the comman.</param>
+		/// <param name="CommandLine">Commandline for p4.</param>
+		/// <param name="Input">Stdin input.</param>
+		/// <param name="AllowSpew">Whether the command should spew.</param>
+		/// <returns>True if succeeded, otherwise false.</returns>
+        public bool LogP4Output(out string Output, string CommandLine, string Input = null, bool AllowSpew = true, bool WithClient = true, bool SpewIsVerbose = false)
+		{
+			Output = "";
+
+			if (String.IsNullOrEmpty(LogPath))
+			{
+				CommandUtils.LogError("P4Utils.SetupP4() must be called before issuing Peforce commands");
+				return false;
+			}
+
+            var Result = P4(CommandLine, Input, AllowSpew, WithClient, SpewIsVerbose:SpewIsVerbose);
+
+			CommandUtils.WriteToFile(LogPath, CommandLine + "\n");
+			CommandUtils.WriteToFile(LogPath, Result.Output);
+			Output = Result.Output;
+			return Result.ExitCode == 0;
+		}
+
+		/// <summary>
+		/// Execute a Perforce command and parse the output as marshalled Python objects. This is more robustly defined than the text-based tagged output
+		/// format, because it avoids ambiguity when returned fields can have newlines.
+		/// </summary>
+		/// <param name="CommandLine">Command line to execute Perforce with</param>
+		/// <param name="TaggedOutput">List that receives the output records</param>
+		/// <param name="WithClient">Whether to include client information on the command line</param>
+		public List<Dictionary<string, string>> P4TaggedOutput(string CommandLine, bool WithClient = true)
+		{
+			// Execute Perforce, consuming the binary output into a memory stream
+			MemoryStream MemoryStream = new MemoryStream();
+			using (Process Process = new Process())
+			{
+				Process.StartInfo.FileName = HostPlatform.Current.P4Exe;
+				Process.StartInfo.Arguments = String.Format("-G {0} {1}", WithClient? GlobalOptions : GlobalOptionsWithoutClient, CommandLine);
+
+				Process.StartInfo.RedirectStandardError = true;
+				Process.StartInfo.RedirectStandardOutput = true;
+				Process.StartInfo.RedirectStandardInput = false;
+				Process.StartInfo.UseShellExecute = false;
+				Process.StartInfo.CreateNoWindow = true;
+
+				Process.Start();
+
+				Process.StandardOutput.BaseStream.CopyTo(MemoryStream);
+				Process.WaitForExit();
+			}
+
+			// Move back to the start of the memory stream
+			MemoryStream.Position = 0;
+
+			// Parse the records
+			List<Dictionary<string, string>> Records = new List<Dictionary<string, string>>();
+			using (BinaryReader Reader = new BinaryReader(MemoryStream, Encoding.UTF8))
+			{
+				while(Reader.BaseStream.Position < Reader.BaseStream.Length)
+				{
+					// Check that a dictionary follows
+					byte Temp = Reader.ReadByte();
+					if(Temp != '{')
+					{
+						throw new P4Exception("Unexpected data while parsing marshalled output - expected '{'");
+					}
+
+					// Read all the fields in the record
+					Dictionary<string, string> Record = new Dictionary<string, string>();
+					for(;;)
+					{
+						// Read the next field type. Perforce only outputs string records. A '0' character indicates the end of the dictionary.
+						byte KeyFieldType = Reader.ReadByte();
+						if(KeyFieldType == '0')
+						{
+							break;
+						}
+						else if(KeyFieldType != 's')
+						{
+							throw new P4Exception("Unexpected key field type while parsing marshalled output ({0}) - expected 's'", (int)KeyFieldType);
+						}
+
+						// Read the key
+						int KeyLength = Reader.ReadInt32();
+						string Key = Encoding.UTF8.GetString(Reader.ReadBytes(KeyLength));
+
+						// Read the value type.
+						byte ValueFieldType = Reader.ReadByte();
+						if(ValueFieldType == 'i')
+						{
+							// An integer
+							string Value = Reader.ReadInt32().ToString();
+							Record.Add(Key, Value);
+						}
+						else if(ValueFieldType == 's')
+						{
+							// A string
+							int ValueLength = Reader.ReadInt32();
+							string Value = Encoding.UTF8.GetString(Reader.ReadBytes(ValueLength));
+							Record.Add(Key, Value);
+						}
+						else
+						{
+							throw new P4Exception("Unexpected value field type while parsing marshalled output ({0}) - expected 's'", (int)ValueFieldType);
+						}
+					}
+					Records.Add(Record);
+				}
+			}
+			return Records;
+		}
+
+		/// <summary>
+		/// Checks that the raw record data includes the given return code, or creates a ReturnCode value if it doesn't
 		/// </summary>
-		/// <param name="CommandLine">Commandline to pass to p4.</param>
-		/// <param name="Input">Stdin input.</param>
-		/// <param name="AllowSpew">Whether the command is allowed to spew.</param>
-        public void LogP4(string CommandLine, string Input = null, bool AllowSpew = true, bool WithClient = true, bool SpewIsVerbose = false)
+		/// <param name="RawRecord">The raw record data</param>
+		/// <param name="ExpectedCode">The expected code value</param>
+		/// <param name="OtherReturnCode">Output variable for receiving the return code if it doesn't match</param>
+		public static bool VerifyReturnCode(Dictionary<string, string> RawRecord, string ExpectedCode, out P4ReturnCode OtherReturnCode)
 		{
-			CheckP4Enabled();
-			string Output;
-            if (!LogP4Output(out Output, CommandLine, Input, AllowSpew, WithClient, SpewIsVerbose:SpewIsVerbose))
+			// Parse the code field
+			string Code;
+			if(!RawRecord.TryGetValue("code", out Code))
 			{
-				throw new P4Exception("p4.exe {0} failed.", CommandLine);
+				Code = "unknown";
 			}
-		}
-
-		/// <summary>
-		/// Calls p4 and returns the output and writes it also to a logfile.
-		/// </summary>
-		/// <param name="Output">Output of the comman.</param>
-		/// <param name="CommandLine">Commandline for p4.</param>
-		/// <param name="Input">Stdin input.</param>
-		/// <param name="AllowSpew">Whether the command should spew.</param>
-		/// <returns>True if succeeded, otherwise false.</returns>
-        public bool LogP4Output(out string Output, string CommandLine, string Input = null, bool AllowSpew = true, bool WithClient = true, bool SpewIsVerbose = false)
-		{
-			CheckP4Enabled();
-			Output = "";
 
-			if (String.IsNullOrEmpty(LogPath))
+			// Check whether it matches what we expect
+			if(Code == ExpectedCode)
 			{
-				CommandUtils.LogError("P4Utils.SetupP4() must be called before issuing Peforce commands");
-				return false;
+				OtherReturnCode = null;
+				return true;
 			}
+			else
+			{
+				string Severity;
+				if(!RawRecord.TryGetValue("severity", out Severity))
+				{
+					Severity = ((int)P4SeverityCode.Empty).ToString();
+				}
 
-            var Result = P4(CommandLine, Input, AllowSpew, WithClient, SpewIsVerbose:SpewIsVerbose);
+				string Generic;
+				if(!RawRecord.TryGetValue("generic", out Generic))
+				{
+					Generic = ((int)P4GenericCode.None).ToString();
+				}
 
-			CommandUtils.WriteToFile(LogPath, CommandLine + "\n");
-			CommandUtils.WriteToFile(LogPath, Result.Output);
-			Output = Result.Output;
-			return Result.ExitCode == 0;
+				string Message;
+				if(!RawRecord.TryGetValue("data", out Message))
+				{
+					Message = "No description available.";
+				}
+
+				OtherReturnCode = new P4ReturnCode(Code, (P4SeverityCode)int.Parse(Severity), (P4GenericCode)int.Parse(Generic), Message.TrimEnd());
+				return false;
+			}
 		}
 
 		/// <summary>
@@ -727,7 +1472,12 @@ namespace AutomationTool
             {
                 return (A.CL < B.CL) ? -1 : (A.CL > B.CL) ? 1 : 0;
             }
-        }
+
+			public override string ToString()
+			{
+				return String.Format("CL {0}: {1}", CL, Summary);
+			}
+		}
         static Dictionary<string, string> UserToEmailCache = new Dictionary<string, string>();
         public string UserToEmail(string User)
         {
@@ -773,7 +1523,6 @@ namespace AutomationTool
                 return true;
             }
             ChangeRecords = new List<ChangeRecord>();
-            CheckP4Enabled();
             try
             {
                 // Change 1999345 on 2014/02/16 by buildmachine@BuildFarm_BUILD-23_buildmachine_++depot+UE4 'GUBP Node Shadow_LabelPromotabl'
@@ -914,6 +1663,11 @@ namespace AutomationTool
 				public string File;
 				public int Revision;
 				public string ChangeType;
+
+				public override string ToString()
+				{
+					return String.Format("{0}#{1} ({2})", File, Revision, ChangeType);
+				}
 			}
 			public List<DescribeFile> Files = new List<DescribeFile>();
             
@@ -921,7 +1675,12 @@ namespace AutomationTool
             {
                 return (A.CL < B.CL) ? -1 : (A.CL > B.CL) ? 1 : 0;
             }
-        }
+
+			public override string ToString()
+			{
+				return String.Format("CL {0}: {1}", CL, Summary);
+			}
+		}
 
 		/// <summary>
 		/// Wraps P4 describe
@@ -960,7 +1719,6 @@ namespace AutomationTool
         public bool DescribeChangelists(List<int> Changelists, out List<DescribeRecord> DescribeRecords, bool AllowSpew = true)
         {
 			DescribeRecords = new List<DescribeRecord>();
-            CheckP4Enabled();
             try
             {
 				// Change 234641 by This.User@WORKSPACE-C2Q-67_Dev on 2008/05/06 10:32:32
@@ -1147,7 +1905,6 @@ namespace AutomationTool
 		/// <param name="CommandLine">CommandLine to pass on to the command.</param>
         public void Sync(string CommandLine, bool AllowSpew = true, bool SpewIsVerbose = false)
 		{
-			CheckP4Enabled();
 			LogP4("sync " + CommandLine, null, AllowSpew, SpewIsVerbose:SpewIsVerbose);
 		}
 
@@ -1159,7 +1916,6 @@ namespace AutomationTool
 		/// <param name="CommandLine">Commandline for the command.</param>
 		public void Unshelve(int FromCL, int ToCL, string CommandLine = "", bool SpewIsVerbose = false)
 		{
-			CheckP4Enabled();
 			LogP4("unshelve " + String.Format("-s {0} ", FromCL) + String.Format("-c {0} ", ToCL) + CommandLine, SpewIsVerbose: SpewIsVerbose);
 		}
 
@@ -1171,7 +1927,6 @@ namespace AutomationTool
         /// <param name="CommandLine">Commandline for the command.</param>
         public void Shelve(int FromCL, string CommandLine = "", bool AllowSpew = true)
         {
-            CheckP4Enabled();
             LogP4("shelve " + String.Format("-r -c {0} ", FromCL) + CommandLine, AllowSpew: AllowSpew);
         }
 
@@ -1182,8 +1937,6 @@ namespace AutomationTool
         /// <param name="CommandLine">Commandline for the command.</param>
         public void DeleteShelvedFiles(int FromCL, bool AllowSpew = true)
         {
-            CheckP4Enabled();
-
 			string Output;
             if (!LogP4Output(out Output, String.Format("shelve -d -c {0}", FromCL), AllowSpew: AllowSpew) && !Output.StartsWith("No shelved files in changelist to delete."))
 			{
@@ -1198,7 +1951,6 @@ namespace AutomationTool
 		/// <param name="CommandLine">Commandline for the command.</param>
 		public void Edit(int CL, string CommandLine)
 		{
-			CheckP4Enabled();
 			LogP4("edit " + String.Format("-c {0} ", CL) + CommandLine);
 		}
 
@@ -1211,7 +1963,6 @@ namespace AutomationTool
 		{
 			try
 			{
-				CheckP4Enabled();
 				string Output;
 				if (!LogP4Output(out Output, "edit " + String.Format("-c {0} ", CL) + CommandLine, null, true))
 				{
@@ -1236,7 +1987,6 @@ namespace AutomationTool
 		/// <param name="CommandLine">Commandline for the command.</param>
 		public void Add(int CL, string CommandLine)
 		{
-			CheckP4Enabled();
 			LogP4("add " + String.Format("-c {0} ", CL) + CommandLine);
 		}
 
@@ -1247,7 +1997,6 @@ namespace AutomationTool
 		/// <param name="CommandLine">Commandline for the command.</param>
 		public void Reconcile(int CL, string CommandLine, bool AllowSpew = true)
 		{
-			CheckP4Enabled();
 			LogP4("reconcile " + String.Format("-c {0} -ead -f ", CL) + CommandLine, AllowSpew: AllowSpew);
 		}
 
@@ -1258,7 +2007,6 @@ namespace AutomationTool
         /// <param name="CommandLine">Commandline for the command.</param>
         public void ReconcilePreview(string CommandLine)
         {
-            CheckP4Enabled();
             LogP4("reconcile " + String.Format("-ead -n ") + CommandLine);
         }
 
@@ -1270,7 +2018,6 @@ namespace AutomationTool
 		/// <param name="CommandLine">Commandline for the command.</param>
 		public void ReconcileNoDeletes(int CL, string CommandLine, bool AllowSpew = true)
 		{
-			CheckP4Enabled();
 			LogP4("reconcile " + String.Format("-c {0} -ea ", CL) + CommandLine, AllowSpew: AllowSpew);
 		}
 
@@ -1282,7 +2029,6 @@ namespace AutomationTool
 		/// <param name="CommandLine">Commandline for the command.</param>
 		public void Resolve(int CL, string CommandLine)
 		{
-			CheckP4Enabled();
 			LogP4("resolve -ay " + String.Format("-c {0} ", CL) + CommandLine);
 		}
 
@@ -1292,7 +2038,6 @@ namespace AutomationTool
 		/// <param name="CommandLine">Commandline for the command.</param>
 		public void Revert(string CommandLine, bool AllowSpew = true)
 		{
-			CheckP4Enabled();
 			LogP4("revert " + CommandLine, AllowSpew: AllowSpew);
 		}
 
@@ -1303,7 +2048,6 @@ namespace AutomationTool
 		/// <param name="CommandLine">Commandline for the command.</param>
 		public void Revert(int CL, string CommandLine = "", bool AllowSpew = true)
 		{
-			CheckP4Enabled();
 			LogP4("revert " + String.Format("-c {0} ", CL) + CommandLine, AllowSpew: AllowSpew);
 		}
 
@@ -1313,7 +2057,6 @@ namespace AutomationTool
 		/// <param name="CL">Changelist to revert the unmodified files from.</param>
 		public void RevertUnchanged(int CL)
 		{
-			CheckP4Enabled();
 			// caution this is a really bad idea if you hope to force submit!!!
 			LogP4("revert -a " + String.Format("-c {0} ", CL));
 		}
@@ -1324,7 +2067,6 @@ namespace AutomationTool
 		/// <param name="CL">Changelist to revert.</param>
 		public void RevertAll(int CL, bool SpewIsVerbose = false)
 		{
-			CheckP4Enabled();
 			LogP4("revert " + String.Format("-c {0} //...", CL), SpewIsVerbose: SpewIsVerbose);
 		}
 
@@ -1337,7 +2079,6 @@ namespace AutomationTool
 		/// <param name="RevertIfFail">If true, if the submit fails, revert the CL.</param>
 		public void Submit(int CL, out int SubmittedCL, bool Force = false, bool RevertIfFail = false)
 		{
-			CheckP4Enabled();
 			if (!CommandUtils.AllowSubmit)
 			{
 				throw new P4Exception("Submit is not allowed currently. Please use the -Submit switch to override that.");
@@ -1505,11 +2246,18 @@ namespace AutomationTool
 		/// <param name="Owner">Owner of the changelist.</param>
 		/// <param name="Description">Description of the changelist.</param>
 		/// <returns>Id of the created changelist.</returns>
-		public int CreateChange(string Owner = null, string Description = null, bool AllowSpew = false)
+		public int CreateChange(string Owner = null, string Description = null, string User = null, string Type = null, bool AllowSpew = false)
 		{
-			CheckP4Enabled();
 			var ChangeSpec = "Change: new" + "\n";
 			ChangeSpec += "Client: " + ((Owner != null) ? Owner : "") + "\n";
+			if(User != null)
+			{
+				ChangeSpec += "User: " + User + "\n";
+			}
+			if(Type != null)
+			{
+				ChangeSpec += "Type: " + Type + "\n";
+			}
 			ChangeSpec += "Description: " + ((Description != null) ? Description.Replace("\n", "\n\t") : "(none)") + "\n";
 			string CmdOutput;
 			int CL = 0;
@@ -1549,8 +2297,6 @@ namespace AutomationTool
 		/// <param name="SpewIsVerbose"></param>
 		public void UpdateChange(int CL, string NewOwner, string NewDescription, bool SpewIsVerbose = false)
 		{
-			CheckP4Enabled();
-
 			string CmdOutput;
 			if(!LogP4Output(out CmdOutput, String.Format("change -o {0}", CL), SpewIsVerbose: SpewIsVerbose))
 			{
@@ -1584,7 +2330,6 @@ namespace AutomationTool
 		/// <param name="RevertFiles">Indicates whether files in that changelist should be reverted.</param>
 		public void DeleteChange(int CL, bool RevertFiles = true, bool SpewIsVerbose = false, bool AllowSpew = true)
 		{
-			CheckP4Enabled();
 			if (RevertFiles)
 			{
 				RevertAll(CL, SpewIsVerbose: SpewIsVerbose);
@@ -1612,8 +2357,6 @@ namespace AutomationTool
 		/// <returns>True if the changelist was deleted, false otherwise.</returns>
 		public bool TryDeleteEmptyChange(int CL)
 		{
-			CheckP4Enabled();
-
 			string CmdOutput;
 			if (LogP4Output(out CmdOutput, String.Format("change -d {0}", CL)))
 			{
@@ -1637,7 +2380,6 @@ namespace AutomationTool
 		/// <returns>Specification of the changelist.</returns>
 		public string ChangeOutput(int CL, bool AllowSpew = true)
 		{
-			CheckP4Enabled();
 			string CmdOutput;
 			if (LogP4Output(out CmdOutput, String.Format("change -o {0}", CL), AllowSpew: AllowSpew))
 			{
@@ -1654,7 +2396,6 @@ namespace AutomationTool
 		/// <returns>Returns whether the changelist exists.</returns>
 		public bool ChangeExists(int CL, out bool Pending, bool AllowSpew = true)
 		{
-			CheckP4Enabled();
 			string CmdOutput = ChangeOutput(CL, AllowSpew);
 			Pending = false;
 			if (CmdOutput.Length > 0)
@@ -1695,7 +2436,6 @@ namespace AutomationTool
 		/// <returns>List of the files contained in the changelist.</returns>
 		public List<string> ChangeFiles(int CL, out bool Pending, bool AllowSpew = true)
 		{
-			CheckP4Enabled();
 			var Result = new List<string>();
 
 			if (ChangeExists(CL, out Pending, AllowSpew))
@@ -1750,7 +2490,6 @@ namespace AutomationTool
         /// <returns>Specification of the changelist.</returns>
         public string OpenedOutput()
         {
-            CheckP4Enabled();
             string CmdOutput;
             if (LogP4Output(out CmdOutput, "opened"))
             {
@@ -1764,7 +2503,6 @@ namespace AutomationTool
 		/// <param name="LabelName">Label to delete.</param>
         public void DeleteLabel(string LabelName, bool AllowSpew = true)
 		{
-			CheckP4Enabled();
 			var CommandLine = "label -d " + LabelName;
 
 			// NOTE: We don't throw exceptions when trying to delete a label
@@ -1787,7 +2525,6 @@ namespace AutomationTool
 		/// <param name="Time">Time of the label creation</param>
 		public void CreateLabel(string Name, string Options, string View, string Owner = null, string Description = null, string Date = null, string Time = null)
 		{
-			CheckP4Enabled();
 			var LabelSpec = "Label: " + Name + "\n";
 			LabelSpec += "Owner: " + ((Owner != null) ? Owner : "") + "\n";
 			LabelSpec += "Description: " + ((Description != null) ? Description : "") + "\n";
@@ -1816,7 +2553,6 @@ namespace AutomationTool
 		/// <param name="AllowSpew">Whether the command is allowed to spew.</param>
 		public void Tag(string LabelName, string FilePath, bool AllowSpew = true)
 		{
-			CheckP4Enabled();
 			LogP4("tag -l " + LabelName + " " + FilePath, null, AllowSpew);
 		}
 
@@ -1827,7 +2563,6 @@ namespace AutomationTool
 		/// <param name="AllowSpew">Whether the command is allowed to spew.</param>
 		public void LabelSync(string LabelName, bool AllowSpew = true, string FileToLabel = "")
 		{
-			CheckP4Enabled();
 			string Quiet = "";
 			if (!AllowSpew)
 			{
@@ -1851,7 +2586,6 @@ namespace AutomationTool
 		/// <param name="AllowSpew">Whether the command is allowed to spew.</param>
 		public void LabelToLabelSync(string FromLabelName, string ToLabelName, bool AllowSpew = true)
 		{
-			CheckP4Enabled();
 			string Quiet = "";
 			if (!AllowSpew)
 			{
@@ -1867,7 +2601,6 @@ namespace AutomationTool
 		/// <returns>Whether there is an label with files.</returns>
 		public bool LabelExistsAndHasFiles(string Name)
 		{
-			CheckP4Enabled();
 			string Output;
 			return LogP4Output(out Output, "files -m 1 //...@" + Name);
 		}
@@ -1881,7 +2614,6 @@ namespace AutomationTool
 		/// <returns>Returns whether the label description could be retrieved.</returns>
 		public bool LabelDescription(string Name, out string Description, bool AllowSpew = true)
 		{
-			CheckP4Enabled();
 			string Output;
 			Description = "";
 			if (LogP4Output(out Output, "label -o " + Name, AllowSpew: AllowSpew))
@@ -1976,8 +2708,6 @@ namespace AutomationTool
 		/// <returns>The head CL number.</returns>
 		public int GetLatestCLNumber()
 		{
-			CheckP4Enabled();
-
 			string Output;
 			if (!LogP4Output(out Output, "changes -s submitted -m1") || string.IsNullOrWhiteSpace(Output))
 			{
@@ -2053,18 +2783,103 @@ namespace AutomationTool
         /// <returns>The file's first reported path on disk or null if no mapping was found</returns>
         public string DepotToLocalPath(string DepotFile, bool AllowSpew = true)
         {
-			P4WhereRecord[] Records = Where(DepotFile, AllowSpew);
-			if (Records != null)
+			//  P4 where outputs missing entries 
+			string Command = String.Format("-z tag fstat \"{0}\"", DepotFile);
+
+			// Run the command.
+			string[] Lines;
+			if (!P4Output(out Lines, Command, AllowSpew: AllowSpew))
+			{
+				throw new P4Exception("p4.exe {0} failed.", Command);
+			}
+
+			// Find the line containing the client file prefix
+			const string ClientFilePrefix = "... clientFile ";
+			foreach(string Line in Lines)
+			{
+				if(Line.StartsWith(ClientFilePrefix))
+				{
+					return Line.Substring(ClientFilePrefix.Length);
+				}
+			}
+			return null;
+        }
+
+        /// <summary>
+        /// Given a set of file paths in the depot, returns the local disk mapping for the current view
+        /// </summary>
+		/// <param name="DepotFiles">The full file paths in depot naming form</param>
+        /// <returns>The file's first reported path on disk or null if no mapping was found</returns>
+        public string[] DepotToLocalPaths(string[] DepotFiles, bool AllowSpew = true)
+        {
+			const int BatchSize = 20;
+
+			// Parse the output from P4
+			List<string> Lines = new List<string>();
+			for(int Idx = 0; Idx < DepotFiles.Length; Idx += BatchSize)
+			{
+				// Build the argument list
+				StringBuilder Command = new StringBuilder("-z tag fstat ");
+				for(int ArgIdx = Idx; ArgIdx < Idx + BatchSize && ArgIdx < DepotFiles.Length; ArgIdx++)
+				{
+					Command.AppendFormat(" {0}", CommandUtils.MakePathSafeToUseWithCommandLine(DepotFiles[ArgIdx]));
+				}
+
+				// Run the command.
+				string[] Output;
+				if (!P4Output(out Output, Command.ToString(), AllowSpew: AllowSpew))
+				{
+					throw new P4Exception("p4.exe {0} failed.", Command);
+				}
+
+				// Append it to the combined output
+				Lines.AddRange(Output);
+			}
+
+			// Parse all the error lines. These may occur out of sequence due to stdout/stderr buffering.
+			for(int LineIdx = 0; LineIdx < Lines.Count; LineIdx++)
+			{
+				if(Lines[LineIdx].Length > 0 && !Lines[LineIdx].StartsWith("... "))
+				{
+					throw new AutomationException("Unexpected output from p4.exe fstat: {0}", Lines[LineIdx]);
+				}
+			}
+
+			// Parse the output lines
+			string[] LocalFiles = new string[DepotFiles.Length];
+			for(int FileIdx = 0, LineIdx = 0; FileIdx < DepotFiles.Length; FileIdx++)
 			{
-				foreach (P4WhereRecord Record in Records)
+				string DepotFile = DepotFiles[FileIdx];
+				if(LineIdx == Lines.Count)
 				{
-					if (!Record.bUnmap)
+					throw new AutomationException("Unexpected end of output looking for file record for {0}", DepotFile);
+				}
+				else
+				{
+					// We've got a file record; try to parse the matching fields
+					for(; LineIdx < Lines.Count && Lines[LineIdx].Length > 0; LineIdx++)
+					{
+						const string DepotFilePrefix = "... depotFile ";
+						if(Lines[LineIdx].StartsWith(DepotFilePrefix) && Lines[LineIdx].Substring(DepotFilePrefix.Length) != DepotFile)
+						{
+							throw new AutomationException("Expected file record for '{0}'; received output '{1}'", DepotFile, Lines[LineIdx]);
+						}
+
+						const string ClientFilePrefix = "... clientFile ";
+						if(Lines[LineIdx].StartsWith(ClientFilePrefix))
+						{
+							LocalFiles[FileIdx] = Lines[LineIdx].Substring(ClientFilePrefix.Length);
+						}
+					}
+
+					// Skip any blank lines
+					while(LineIdx < Lines.Count && Lines[LineIdx].Length == 0)
 					{
-						return Record.Path;
+						LineIdx++;
 					}
 				}
 			}
-			return null;
+			return LocalFiles;
         }
 
 		/// <summary>
@@ -2077,8 +2892,6 @@ namespace AutomationTool
 		/// <returns>List of records describing the file's mapping. Usually just one, but may be more.</returns>
 		public P4WhereRecord[] Where(string DepotFile, bool AllowSpew = true)
 		{
-			CheckP4Enabled();
-
 			//  P4 where outputs missing entries 
 			string Command = String.Format("-z tag where \"{0}\"", DepotFile);
 
@@ -2166,7 +2979,6 @@ namespace AutomationTool
 		/// <returns>File stats (invalid if the file does not exist in P4)</returns>
 		public P4FileStat FStat(string Filename)
 		{
-			CheckP4Enabled();
 			string Output;
 			string Command = "fstat " + CommandUtils.MakePathSafeToUseWithCommandLine(Filename);
 			if (!LogP4Output(out Output, Command))
@@ -2309,7 +3121,6 @@ namespace AutomationTool
 		/// <returns>True if the client exists.</returns>
 		public bool DoesClientExist(string ClientName, bool Quiet = false)
 		{
-			CheckP4Enabled();
 			if(!Quiet)
 			{
 				CommandUtils.LogLog("Checking if client {0} exists", ClientName);
@@ -2326,8 +3137,6 @@ namespace AutomationTool
 		/// <returns></returns>
 		public P4ClientInfo GetClientInfo(string ClientName, bool Quiet = false)
 		{
-			CheckP4Enabled();
-
 			if(!Quiet)
 			{
 				CommandUtils.LogLog("Getting info for client {0}", ClientName);
@@ -2431,8 +3240,6 @@ namespace AutomationTool
 		/// <returns>List of clients owned by the user.</returns>
 		public P4ClientInfo[] GetClientsForUser(string UserName, string PathUnderClientRoot = null)
 		{
-			CheckP4Enabled();
-
 			var ClientList = new List<P4ClientInfo>();
 
 			// Get all clients for this user
@@ -2489,10 +3296,9 @@ namespace AutomationTool
         /// </summary>
         /// <param name="Name">Client name.</param>
         /// <param name="Force">Forces the operation (-f)</param>
-        public void DeleteClient(string Name, bool Force = false)
+        public void DeleteClient(string Name, bool Force = false, bool AllowSpew = true)
         {
-            CheckP4Enabled();
-            LogP4(String.Format("client -d {0} {1}", (Force ? "-f" : ""), Name), WithClient: false);
+            LogP4(String.Format("client -d {0} {1}", (Force ? "-f" : ""), Name), WithClient: false, AllowSpew: AllowSpew);
         }
 
         /// <summary>
@@ -2534,7 +3340,6 @@ namespace AutomationTool
 		/// <returns>List of sub-directories of the specified directories.</returns>
 		public List<string> Dirs(string CommandLine)
 		{
-			CheckP4Enabled();
 			var DirsCmdLine = String.Format("dirs {0}", CommandLine);
 			var P4Result = P4(DirsCmdLine, AllowSpew: false);
 			if (P4Result.ExitCode != 0)
@@ -2560,7 +3365,6 @@ namespace AutomationTool
 		/// <returns>List of files in the specified directory.</returns>
 		public List<string> Files(string CommandLine)
 		{
-			CheckP4Enabled();
 			string FilesCmdLine = String.Format("files {0}", CommandLine);
 			IProcessResult P4Result = P4(FilesCmdLine, AllowSpew: false);
 			if (P4Result.ExitCode != 0)
@@ -2657,6 +3461,362 @@ namespace AutomationTool
 			return Changelists;
 		}
 
+		/// <summary>
+		/// Execute the 'filelog' command
+		/// </summary>
+		/// <param name="Options">Options for the command</param>
+		/// <param name="FileSpecs">List of file specifications to query</param>
+		/// <returns>List of file records</returns>
+		public P4FileRecord[] FileLog(P4FileLogOptions Options, params string[] FileSpecs)
+		{
+			return FileLog(-1, -1, Options, FileSpecs);
+		}
+
+		/// <summary>
+		/// Execute the 'filelog' command
+		/// </summary>
+		/// <param name="MaxChanges">Number of changelists to show. Ignored if zero or negative.</param>
+		/// <param name="Options">Options for the command</param>
+		/// <param name="FileSpecs">List of file specifications to query</param>
+		/// <returns>List of file records</returns>
+		public P4FileRecord[] FileLog(int MaxChanges, P4FileLogOptions Options, params string[] FileSpecs)
+		{
+			return FileLog(-1, MaxChanges, Options, FileSpecs);
+		}
+
+		/// <summary>
+		/// Execute the 'filelog' command
+		/// </summary>
+		/// <param name="ChangeNumber">Show only files modified by this changelist. Ignored if zero or negative.</param>
+		/// <param name="MaxChanges">Number of changelists to show. Ignored if zero or negative.</param>
+		/// <param name="Options">Options for the command</param>
+		/// <param name="FileSpecs">List of file specifications to query</param>
+		/// <returns>List of file records</returns>
+		public P4FileRecord[] FileLog(int ChangeNumber, int MaxChanges, P4FileLogOptions Options, params string[] FileSpecs)
+		{
+			P4FileRecord[] Records;
+			P4ReturnCode ReturnCode = TryFileLog(ChangeNumber, MaxChanges, Options, FileSpecs, out Records);
+			if(ReturnCode != null)
+			{
+				if(ReturnCode.Generic == P4GenericCode.Empty)
+				{
+					return new P4FileRecord[0];
+				}
+				else
+				{
+					throw new P4Exception(ReturnCode.ToString());
+				}
+			}
+			return Records;
+		}
+
+		/// <summary>
+		/// Execute the 'filelog' command
+		/// </summary>
+		/// <param name="ChangeNumber">Show only files modified by this changelist. Ignored if zero or negative.</param>
+		/// <param name="MaxChanges">Number of changelists to show. Ignored if zero or negative.</param>
+		/// <param name="Options">Options for the command</param>
+		/// <param name="FileSpecs">List of file specifications to query</param>
+		/// <returns>List of file records</returns>
+		public P4ReturnCode TryFileLog(int ChangeNumber, int MaxChanges, P4FileLogOptions Options, string[] FileSpecs, out P4FileRecord[] OutRecords)
+		{
+			// Build the argument list
+			List<string> Arguments = new List<string>();
+			if(ChangeNumber > 0)
+			{
+				Arguments.Add(String.Format("-c {0}", ChangeNumber));
+			}
+			if((Options & P4FileLogOptions.ContentHistory) != 0)
+			{
+				Arguments.Add("-h");
+			}
+			if((Options & P4FileLogOptions.FollowAcrossBranches) != 0)
+			{
+				Arguments.Add("-i");
+			}
+			if((Options & P4FileLogOptions.FullDescriptions) != 0)
+			{
+				Arguments.Add("-l");
+			}
+			if((Options & P4FileLogOptions.LongDescriptions) != 0)
+			{
+				Arguments.Add("-L");
+			}
+			if(MaxChanges > 0)
+			{
+				Arguments.Add(String.Format("-m {0}", MaxChanges));
+			}
+			if((Options & P4FileLogOptions.DoNotFollowPromotedTaskStreams) != 0)
+			{
+				Arguments.Add("-p");
+			}
+			if((Options & P4FileLogOptions.IgnoreNonContributoryIntegrations) != 0)
+			{
+				Arguments.Add("-s");
+			}
+
+			// Always include times to simplify parsing
+			Arguments.Add("-t");
+
+			// Add the file arguments
+			foreach(string FileSpec in FileSpecs)
+			{
+				Arguments.Add(CommandUtils.MakePathSafeToUseWithCommandLine(FileSpec));
+			}
+
+			// Format the full command line
+			string CommandLine = String.Format("filelog {0}", String.Join(" ", Arguments));
+
+			// Get the output
+			List<Dictionary<string, string>> RawRecords = P4TaggedOutput(CommandLine);
+
+			// Parse all the output
+			List<P4FileRecord> Records = new List<P4FileRecord>();
+			foreach(Dictionary<string, string> RawRecord in RawRecords)
+			{
+				// Make sure the record has the correct return value
+				P4ReturnCode OtherReturnCode;
+				if(!VerifyReturnCode(RawRecord, "stat", out OtherReturnCode))
+				{
+					OutRecords = null;
+					return OtherReturnCode;
+				}
+
+				// Get the depot path for this revision
+				string DepotPath = RawRecord["depotFile"];
+
+				// Parse the revisions
+				List<P4RevisionRecord> Revisions = new List<P4RevisionRecord>();
+				for(;;)
+				{
+					string RevisionSuffix = String.Format("{0}", Revisions.Count);
+
+					string RevisionNumberText;
+					if(!RawRecord.TryGetValue("rev" + RevisionSuffix, out RevisionNumberText))
+					{
+						break;
+					}
+
+					int RevisionNumber = int.Parse(RevisionNumberText);
+					int RevisionChangeNumber = int.Parse(RawRecord["change" + RevisionSuffix]);
+					P4Action Action = ParseActionText(RawRecord["action" + RevisionSuffix]);
+					DateTime DateTime = UnixEpoch + TimeSpan.FromSeconds(long.Parse(RawRecord["time" + RevisionSuffix]));
+					string Type = RawRecord["type" + RevisionSuffix];
+					string UserName = RawRecord["user" + RevisionSuffix];
+					string ClientName = RawRecord["client" + RevisionSuffix];
+					int FileSize = RawRecord.ContainsKey("fileSize" + RevisionSuffix)? int.Parse(RawRecord["fileSize" + RevisionSuffix]) : -1;
+					string Digest = RawRecord.ContainsKey("digest" + RevisionSuffix)? RawRecord["digest" + RevisionSuffix] : null;
+					string Description = RawRecord["desc" + RevisionSuffix];
+
+					// Parse all the following integration info
+					List<P4IntegrationRecord> Integrations = new List<P4IntegrationRecord>();
+					for(;;)
+					{
+						string IntegrationSuffix = String.Format("{0},{1}", Revisions.Count, Integrations.Count);
+
+						string HowText;
+						if(!RawRecord.TryGetValue("how" + IntegrationSuffix, out HowText))
+						{
+							break;
+						}
+
+						P4IntegrateAction IntegrateAction = ParseIntegrateActionText(HowText);
+						string OtherFile = RawRecord["file" + IntegrationSuffix];
+						string StartRevisionText = RawRecord["srev" + IntegrationSuffix];
+						string EndRevisionText = RawRecord["erev" + IntegrationSuffix];
+						int StartRevisionNumber = (StartRevisionText == "#none")? 0 : int.Parse(StartRevisionText.Substring(1));
+						int EndRevisionNumber = int.Parse(EndRevisionText.Substring(1));
+
+						Integrations.Add(new P4IntegrationRecord(IntegrateAction, OtherFile, StartRevisionNumber, EndRevisionNumber));
+					}
+
+					// Add the revision
+					Revisions.Add(new P4RevisionRecord(RevisionNumber, RevisionChangeNumber, Action, Type, DateTime, UserName, ClientName, FileSize, Digest, Description, Integrations.ToArray()));
+				}
+
+				// Add the file record
+				Records.Add(new P4FileRecord(DepotPath, Revisions.ToArray()));
+			}
+			OutRecords = Records.ToArray();
+			return null;
+		}
+
+		static readonly DateTime UnixEpoch = new DateTime(1970, 1, 1, 0, 0, 0, DateTimeKind.Utc);
+
+		/// <summary>
+		/// Enumerates all streams in a depot
+		/// </summary>
+		/// <param name="StreamPath">The path for streams to enumerate (eg. "//UE4/...")</param>
+		/// <returns>List of streams matching the given criteria</returns>
+		public List<P4StreamRecord> Streams(string StreamPath)
+		{
+			return Streams(StreamPath, -1, null, false);
+		}
+
+		/// <summary>
+		/// Enumerates all streams in a depot
+		/// </summary>
+		/// <param name="StreamPath">The path for streams to enumerate (eg. "//UE4/...")</param>
+		/// <param name="MaxResults">Maximum number of results to return</param>
+		/// <param name="Filter">Additional filter to be applied to the results</param>
+		/// <param name="bUnloaded">Whether to enumerate unloaded workspaces</param>
+		/// <returns>List of streams matching the given criteria</returns>
+		public List<P4StreamRecord> Streams(string StreamPath, int MaxResults, string Filter, bool bUnloaded)
+		{
+			// Build the command line
+			StringBuilder CommandLine = new StringBuilder("streams");
+			if(bUnloaded)
+			{
+				CommandLine.Append(" -U");
+			}
+			if(Filter != null)
+			{
+				CommandLine.AppendFormat("-F \"{0}\"", Filter);
+			}
+			if(MaxResults > 0)
+			{
+				CommandLine.AppendFormat("-m {0}", MaxResults);
+			}
+			CommandLine.AppendFormat(" \"{0}\"", StreamPath);
+
+			// Execute the command
+			List<Dictionary<string, string>> RawRecords = P4TaggedOutput(CommandLine.ToString(), false);
+
+			// Parse the output
+			List<P4StreamRecord> Records = new List<P4StreamRecord>();
+			foreach(Dictionary<string, string> RawRecord in RawRecords)
+			{
+				// Make sure the record has the correct return value
+				P4ReturnCode OtherReturnCode;
+				if(!VerifyReturnCode(RawRecord, "stat", out OtherReturnCode))
+				{
+					throw new P4Exception(OtherReturnCode.ToString());
+				}
+
+				// Parse the fields
+				string Stream = RawRecord["Stream"];
+				DateTime Update = UnixEpoch + TimeSpan.FromSeconds(long.Parse(RawRecord["Update"]));
+				DateTime Access = UnixEpoch + TimeSpan.FromSeconds(long.Parse(RawRecord["Access"]));
+				string Owner = RawRecord["Owner"];
+				string Name = RawRecord["Name"];
+				string Parent = RawRecord["Parent"];
+				P4StreamType Type = (P4StreamType)Enum.Parse(typeof(P4StreamType), RawRecord["Type"], true);
+				string Description = RawRecord["desc"];
+				P4StreamOptions Options = ParseStreamOptions(RawRecord["Options"]);
+				Nullable<bool> FirmerThanParent = ParseNullableBool(RawRecord["firmerThanParent"]);
+				bool ChangeFlowsToParent = bool.Parse(RawRecord["changeFlowsToParent"]);
+				bool ChangeFlowsFromParent = bool.Parse(RawRecord["changeFlowsFromParent"]);
+				string BaseParent = RawRecord["baseParent"];
+
+				// Add the new stream record
+				Records.Add(new P4StreamRecord(Stream, Update, Access, Owner, Name, Parent, Type, Description, Options, FirmerThanParent, ChangeFlowsToParent, ChangeFlowsFromParent, BaseParent));
+			}
+			return Records;
+		}
+
+		/// <summary>
+		/// Parse a nullable boolean
+		/// </summary>
+		/// <param name="Text">Text to parse. May be "true", "false", or "n/a".</param>
+		/// <returns>The parsed boolean</returns>
+		static Nullable<bool> ParseNullableBool(string Text)
+		{
+			switch(Text)
+			{
+				case "true":
+					return true;
+				case "false":
+					return false;
+				case "n/a":
+					return null;
+				default:
+					throw new P4Exception("Invalid value for nullable bool: {0}", Text);
+			}
+		}
+
+		/// <summary>
+		/// Parse a list of stream option flags
+		/// </summary>
+		/// <param name="Text">Text to parse</param>
+		/// <returns>Flags for the stream options</returns>
+		static P4StreamOptions ParseStreamOptions(string Text)
+		{
+			P4StreamOptions Options = 0;
+			foreach(string Option in Text.Split(' '))
+			{
+				switch(Option)
+				{
+					case "locked":
+						Options |= P4StreamOptions.Locked;
+						break;
+					case "ownersubmit":
+						Options |= P4StreamOptions.OwnerSubmit;
+						break;
+					case "toparent":
+						Options |= P4StreamOptions.ToParent;
+						break;
+					case "fromparent":
+						Options |= P4StreamOptions.FromParent;
+						break;
+					case "mergedown":
+						Options |= P4StreamOptions.MergeDown;
+						break;
+					case "unlocked":
+					case "allsubmit":
+					case "notoparent":
+					case "nofromparent":
+					case "mergeany":
+						break;
+					default:
+						throw new P4Exception("Unknown stream option '{0}'", Option);
+				}
+			}
+			return Options;
+		}
+
+		static Dictionary<string, T> GetEnumLookup<T>()
+		{
+			Dictionary<string, T> Lookup = new Dictionary<string, T>();
+			foreach(T Value in Enum.GetValues(typeof(T)))
+			{
+				foreach(MemberInfo Member in typeof(T).GetMember(Value.ToString()))
+				{
+					string Description = Member.GetCustomAttribute<DescriptionAttribute>().Description;
+					Lookup.Add(Description, Value);
+				}
+			}
+			return Lookup;
+		}
+
+		static Lazy<Dictionary<string, P4Action>> DescriptionToAction = new Lazy<Dictionary<string, P4Action>>(() => GetEnumLookup<P4Action>());
+
+		static P4Action ParseActionText(string ActionText)
+		{
+			P4Action Action;
+			if(!DescriptionToAction.Value.TryGetValue(ActionText, out Action))
+			{
+				throw new P4Exception("Invalid action '{0}'", Action);
+			}
+			return Action;
+		}
+
+		static Lazy<Dictionary<string, P4IntegrateAction>> DescriptionToIntegrationAction = new Lazy<Dictionary<string, P4IntegrateAction>>(() => GetEnumLookup<P4IntegrateAction>());
+
+		static P4IntegrateAction ParseIntegrateActionText(string ActionText)
+		{
+			P4IntegrateAction Action;
+			if(!DescriptionToIntegrationAction.Value.TryGetValue(ActionText, out Action))
+			{
+				throw new P4Exception("Invalid integration action '{0}'", Action);
+			}
+			return Action;
+		}
+
+		static DateTime ParseDateTime(string DateTimeText)
+		{
+			return DateTime.ParseExact(DateTimeText, "yyyy/MM/dd HH:mm:ss", System.Globalization.CultureInfo.InvariantCulture);
+		}
+
 		/// <summary>
 		/// For a given file (and revision, potentially), returns where it was integrated from. Useful in conjunction with files in a P4DescribeRecord, with action = "integrate".
 		/// </summary>
@@ -2755,7 +3915,7 @@ namespace AutomationTool
 			}
 		}
 
-		private static P4Action ParseAction(string Action)
+		static P4Action ParseAction(string Action)
 		{
 			P4Action Result = P4Action.Unknown;
 			var AllActions = GetEnumValuesAndKeywords(typeof(P4Action));
diff --git a/Engine/Source/Programs/GitDependencies/DependencyManifest.cs b/Engine/Source/Programs/GitDependencies/DependencyManifest.cs
--- a/Engine/Source/Programs/GitDependencies/DependencyManifest.cs
+++ b/Engine/Source/Programs/GitDependencies/DependencyManifest.cs
@@ -12,67 +12,127 @@ using System.ComponentModel;
 
 namespace GitDependencies
 {
+	/// <summary>
+	/// Contains information about a binary dependency in the workspace
+	/// </summary>
 	[DebuggerDisplay("{Name}")]
 	public class DependencyFile
 	{
+		/// <summary>
+		/// Name of the file, relative to the root. Directories are separated with forward slashes, without a leading slash.
+		/// </summary>
 		[XmlAttribute]
 		public string Name;
 
+		/// <summary>
+		/// Hash of this file. Used to determine which blob represents its contents.
+		/// </summary>
 		[XmlAttribute]
 		public string Hash;
 
+		/// <summary>
+		/// Whether this file should have the exectuable bit set, on platforms that support it.
+		/// </summary>
 		[XmlAttribute]
 		[DefaultValue(false)]
 		public bool IsExecutable;
 	}
 
+	/// <summary>
+	/// Represents a hashed binary blob of data.
+	/// </summary>
 	[DebuggerDisplay("{Hash}")]
 	public class DependencyBlob
 	{
+		/// <summary>
+		/// Hash of the blob.
+		/// </summary>
 		[XmlAttribute]
 		public string Hash;
 
+		/// <summary>
+		/// Size of the blob.
+		/// </summary>
 		[XmlAttribute]
 		public long Size;
 
+		/// <summary>
+		/// Hash of the pack file containing this blob.
+		/// </summary>
 		[XmlAttribute]
 		public string PackHash;
 
+		/// <summary>
+		/// Offset of this blob's data within the pack file.
+		/// </summary>
 		[XmlAttribute]
 		public long PackOffset;
 	}
 
+	/// <summary>
+	/// A downloadable data file which can contain several blobs.
+	/// </summary>
 	[DebuggerDisplay("{Hash}")]
 	public class DependencyPack
 	{
+		/// <summary>
+		/// Hash of this pack file.
+		/// </summary>
 		[XmlAttribute]
 		public string Hash;
 
+		/// <summary>
+		/// Size of this pack file, when uncompressed.
+		/// </summary>
 		[XmlAttribute]
 		public long Size;
 
+		/// <summary>
+		/// Compressed size of this pack file.
+		/// </summary>
 		[XmlAttribute]
 		public long CompressedSize;
 
+		/// <summary>
+		/// Subdirectory for this pack file on the remote server.
+		/// </summary>
 		[XmlAttribute]
 		public string RemotePath;
 	}
 
+	/// <summary>
+	/// A manifest for dependencies that should be unpacked into the workspace at a given revision.
+	/// </summary>
 	public class DependencyManifest
 	{
+		/// <summary>
+		/// Base URL for downloading pack files from.
+		/// </summary>
 		[XmlAttribute]
 		public string BaseUrl = "http://cdn.unrealengine.com/dependencies";
 
+		/// <summary>
+		/// Whether to ignore proxy servers when downloading files from this remote.
+		/// </summary>
 		[XmlAttribute]
 		[DefaultValue(false)]
 		public bool IgnoreProxy;
 
+		/// <summary>
+		/// Files which should be extracted into the workspace.
+		/// </summary>
 		[XmlArrayItem("File")]
 		public DependencyFile[] Files = new DependencyFile[0];
 
+		/// <summary>
+		/// All the blobs required for entries in the Files array.
+		/// </summary>
 		[XmlArrayItem("Blob")]
 		public DependencyBlob[] Blobs = new DependencyBlob[0];
 
+		/// <summary>
+		/// All the packs required for entries in the Blobs array
+		/// </summary>
 		[XmlArrayItem("Pack")]
 		public DependencyPack[] Packs = new DependencyPack[0];
 	}
