Fix documentation build failure; P4SubmitOption no longer has an enum value for zero, so default setting for new clients must be initialized.

diff --git a/Engine/Source/Programs/AutomationTool/AutomationUtils/P4Utils.cs b/Engine/Source/Programs/AutomationTool/AutomationUtils/P4Utils.cs
--- a/Engine/Source/Programs/AutomationTool/AutomationUtils/P4Utils.cs
+++ b/Engine/Source/Programs/AutomationTool/AutomationUtils/P4Utils.cs
@@ -77,7 +77,7 @@ namespace AutomationTool
 		public DateTime Access;
         public P4LineEnd LineEnd;
         public P4ClientOption Options;
-        public P4SubmitOption SubmitOptions;
+        public P4SubmitOption SubmitOptions = P4SubmitOption.SubmitUnchanged;
         public List<KeyValuePair<string, string>> View = new List<KeyValuePair<string, string>>();
 
 		public bool Matches(P4ClientInfo Other)
