UAT: Add more unit definitions to telemetry data.

diff --git a/Engine/Source/Programs/AutomationTool/AutomationUtils/CommandUtils.cs b/Engine/Source/Programs/AutomationTool/AutomationUtils/CommandUtils.cs
--- a/Engine/Source/Programs/AutomationTool/AutomationUtils/CommandUtils.cs
+++ b/Engine/Source/Programs/AutomationTool/AutomationUtils/CommandUtils.cs
@@ -2251,7 +2251,12 @@ namespace AutomationTool
 	/// </summary>
 	public enum TelemetryUnits
 	{
+		Count,
+		Milliseconds,
+		Seconds,
 		Minutes,
+		Bytes,
+		Megabytes,
 	}
 
 	/// <summary>
