When compiling through UAT, invoke xgConsole.exe with the /no_watchdog_thread argument if the version number is >= 1659. Works around hangs related to the UCRT in the VS2015 compiler.

diff --git a/Engine/Source/Programs/AutomationTool/AutomationUtils/UE4Build.cs b/Engine/Source/Programs/AutomationTool/AutomationUtils/UE4Build.cs
--- a/Engine/Source/Programs/AutomationTool/AutomationUtils/UE4Build.cs
+++ b/Engine/Source/Programs/AutomationTool/AutomationUtils/UE4Build.cs
@@ -1444,6 +1444,18 @@ namespace AutomationTool
 						{
 							Args += " /StopOnErrors";
 						}
+
+						// A bug in the UCRT can cause XGE to hang on VS2015 builds. Figure out if this hang is likely to effect this build and workaround it if able.
+						string XGEVersion = Microsoft.Win32.Registry.GetValue(@"HKEY_LOCAL_MACHINE\SOFTWARE\Wow6432Node\Xoreax\IncrediBuild\Builder", "Version", null) as string;
+						if (XGEVersion != null)
+						{
+							// Per Xoreax support, subtract 1001000 from the registry value to get the build number of the installed XGE.
+							int XGEBuildNumber;
+							if (Int32.TryParse(XGEVersion, out XGEBuildNumber) && XGEBuildNumber - 1001000 >= 1659)
+							{
+								Args += " /no_watchdog_thread";
+							}
+						}
 					}
 
 					if (!bCanUseParallelExecutor && String.IsNullOrEmpty(XGETool))
