Disabled Gauntlet script module for Linux builds (awaiting dependency refactor)

diff --git a/Engine/Source/Programs/AutomationTool/AutomationUtils/LinuxHostPlatform.cs b/Engine/Source/Programs/AutomationTool/AutomationUtils/LinuxHostPlatform.cs
--- a/Engine/Source/Programs/AutomationTool/AutomationUtils/LinuxHostPlatform.cs
+++ b/Engine/Source/Programs/AutomationTool/AutomationUtils/LinuxHostPlatform.cs
@@ -107,6 +107,11 @@ namespace AutomationTool
 
 		public override bool IsScriptModuleSupported(string ModuleName)
 		{
+			// @todo: add more unsupported modules here
+			if (String.Equals(ModuleName, "Gauntlet", StringComparison.InvariantCultureIgnoreCase))
+			{
+				return false;
+			}
 			return true;
 		}
 
