UAT: Fix UAT failing to load due to stale DLLs copied to Engine/Binaries/DotNET by PortalPublishingTool. Assembly resolver will load these by default, rather than the ones compiled for UAT.

Set output folder for PPT to be under its project folder, so it remains under NotForLicensees protections and does not interfere with other programs, and delete any stale DLLs that have already been copied to Engine/Binaries/DotNET.

diff --git a/Engine/Source/Programs/AutomationTool/AutomationUtils/ScriptCompiler.cs b/Engine/Source/Programs/AutomationTool/AutomationUtils/ScriptCompiler.cs
--- a/Engine/Source/Programs/AutomationTool/AutomationUtils/ScriptCompiler.cs
+++ b/Engine/Source/Programs/AutomationTool/AutomationUtils/ScriptCompiler.cs
@@ -321,6 +321,17 @@ namespace AutomationTool
 		{
 			CommandUtils.LogVerbose("Cleaning up script DLL folder");
 			CommandUtils.DeleteDirectory(GetScriptAssemblyFolder());
+
+			// Bug in PortalPublishingTool caused these DLLs to be copied into Engine/Binaries/DotNET. Delete any files left over.
+			DirectoryReference BinariesDir = DirectoryReference.Combine(CommandUtils.RootDirectory, "Engine", "Binaries", "DotNET");
+			foreach (FileReference FileToDelete in DirectoryReference.EnumerateFiles(BinariesDir, "*.automation.dll"))
+			{
+				CommandUtils.DeleteFile(FileToDelete.FullName);
+			}
+			foreach (FileReference FileToDelete in DirectoryReference.EnumerateFiles(BinariesDir, "*.automation.pdb"))
+			{
+				CommandUtils.DeleteFile(FileToDelete.FullName);
+			}
 		}
 
 		private static string GetScriptAssemblyFolder()
