BuildGraph: Add a parameter to the <Compile> task which allows forcing XGE to be disabled for a compile (AllowXGE="false").

diff --git a/Engine/Source/Programs/AutomationTool/BuildGraph/Tasks/CompileTask.cs b/Engine/Source/Programs/AutomationTool/BuildGraph/Tasks/CompileTask.cs
--- a/Engine/Source/Programs/AutomationTool/BuildGraph/Tasks/CompileTask.cs
+++ b/Engine/Source/Programs/AutomationTool/BuildGraph/Tasks/CompileTask.cs
@@ -40,6 +40,12 @@ namespace AutomationTool
 		[TaskParameter(Optional = true)]
 		public string Arguments;
 
+		/// <summary>
+		/// Whether to allow using XGE for compilation
+		/// </summary>
+		[TaskParameter(Optional = true)]
+		public bool AllowXGE = true;
+
 		/// <summary>
 		/// Whether to allow using the parallel executor for this compile
 		/// </summary>
@@ -68,6 +74,11 @@ namespace AutomationTool
 		/// </summary>
 		Dictionary<UE4Build.BuildTarget, string> TargetToTagName = new Dictionary<UE4Build.BuildTarget,string>();
 
+		/// <summary>
+		/// Whether to allow using XGE for this job
+		/// </summary>
+		bool bAllowXGE = true;
+
 		/// <summary>
 		/// Whether to allow using the parallel executor for this job
 		/// </summary>
@@ -97,13 +108,18 @@ namespace AutomationTool
 
 			if(Targets.Count > 0)
 			{
-				if(!bAllowParallelExecutor || !CompileTask.Parameters.AllowParallelExecutor)
+				if (bAllowXGE != CompileTask.Parameters.AllowXGE)
+				{
+					return false;
+				}
+				if (!bAllowParallelExecutor || !CompileTask.Parameters.AllowParallelExecutor)
 				{
 					return false;
 				}
 			}
 
 			CompileTaskParameters Parameters = CompileTask.Parameters;
+			bAllowXGE &= Parameters.AllowXGE;
 			bAllowParallelExecutor &= Parameters.AllowParallelExecutor;
 
 			UE4Build.BuildTarget Target = new UE4Build.BuildTarget { TargetName = Parameters.Target, Platform = Parameters.Platform, Config = Parameters.Configuration, UBTArgs = "-nobuilduht " + (Parameters.Arguments ?? "") };
@@ -135,7 +151,7 @@ namespace AutomationTool
 			try
 			{
 				bool bCanUseParallelExecutor = (BuildHostPlatform.Current.Platform == UnrealTargetPlatform.Win64 && bAllowParallelExecutor);	// parallel executor is only available on Windows as of 2016-09-22
-				Builder.Build(Agenda, InDeleteBuildProducts: null, InUpdateVersionFiles: false, InForceNoXGE: false, InUseParallelExecutor: bCanUseParallelExecutor, InTargetToManifest: TargetToManifest);
+				Builder.Build(Agenda, InDeleteBuildProducts: null, InUpdateVersionFiles: false, InForceNoXGE: !bAllowXGE, InUseParallelExecutor: bCanUseParallelExecutor, InTargetToManifest: TargetToManifest);
 			}
 			catch (CommandUtils.CommandFailedException)
 			{
