fix for not properly staging the plist to the staging directory

diff --git a/Engine/Source/Programs/AutomationTool/IOS/IOSPlatform.Automation.cs b/Engine/Source/Programs/AutomationTool/IOS/IOSPlatform.Automation.cs
--- a/Engine/Source/Programs/AutomationTool/IOS/IOSPlatform.Automation.cs
+++ b/Engine/Source/Programs/AutomationTool/IOS/IOSPlatform.Automation.cs
@@ -837,8 +837,8 @@ public class IOSPlatform : Platform
 			// copy the plist (only if code signing, as it's protected by the code sign blob in the executable and can't be modified independently)
 			if (GetCodeSignDesirability(Params))
 			{
-//				string SourcePath = CombinePaths((SC.IsCodeBasedProject ? SC.ProjectRoot : SC.LocalRoot + "/Engine"), "Intermediate", PlatformName);
-//				string TargetPListFile = Path.Combine(SourcePath, (SC.IsCodeBasedProject ? SC.ShortProjectName : "UE4Game") + "-Info.plist");
+				string SourcePath = CombinePaths((SC.IsCodeBasedProject ? SC.ProjectRoot : SC.LocalRoot + "/Engine"), "Intermediate", PlatformName);
+				string TargetPListFile = Path.Combine(SourcePath, (SC.IsCodeBasedProject ? SC.ShortProjectName : "UE4Game") + "-Info.plist");
 //				if (!File.Exists(TargetPListFile))
 				{
 					// ensure the plist, entitlements, and provision files are properly copied
@@ -868,6 +868,9 @@ public class IOSPlatform : Platform
                             out bSupportsPortrait,
                             out bSupportsLandscape);
 
+                    // copy the plist to the stage dir
+                    SC.StageFiles(StagedFileType.NonUFS, SourcePath, Path.GetFileName(TargetPListFile), false, null, "", false, false, "Info.plist");
+
                     // copy the icons/launch screens from the engine
                     {
                         string DataPath = CombinePaths(SC.LocalRoot, "Engine", "Build", "IOS", "Resources", "Graphics");
