UAT: Include more information in error message regarding invalid names in staged filesystem references. Trying to track down CIS error.

diff --git a/Engine/Source/Programs/AutomationTool/AutomationUtils/StagedFileSystemReference.cs b/Engine/Source/Programs/AutomationTool/AutomationUtils/StagedFileSystemReference.cs
--- a/Engine/Source/Programs/AutomationTool/AutomationUtils/StagedFileSystemReference.cs
+++ b/Engine/Source/Programs/AutomationTool/AutomationUtils/StagedFileSystemReference.cs
@@ -40,15 +40,15 @@ public abstract class StagedFileSystemReference
 				{
 					if (Idx == LastIdx + 1)
 					{
-						throw new ArgumentException("Empty path fragment in staged filesystem reference");
+						throw new ArgumentException("Empty path fragment in staged filesystem reference ({0})", Name);
 					}
 					if (Idx == LastIdx + 2 && Name[Idx - 1] == '.')
 					{
-						throw new ArgumentException("Pseudo-directories are not permitted in filesystem references");
+						throw new ArgumentException("Pseudo-directories are not permitted in filesystem references ({0})", Name);
 					}
 					if (Idx == LastIdx + 3 && Name[Idx - 1] == '.' && Name[Idx - 2] == '.')
 					{
-						throw new ArgumentException("Pseudo-directories are not permitted in filesystem references");
+						throw new ArgumentException("Pseudo-directories are not permitted in filesystem references ({0})", Name);
 					}
 					LastIdx = Idx;
 				}
