- Edigrate of the "create copy up CL" tool that splits up text files based on secret platforms (source CL is 3487400 in Dev-Console)

diff --git a/Engine/Source/Programs/AutomationTool/Scripts/StreamCopyDescription.Automation.cs b/Engine/Source/Programs/AutomationTool/Scripts/StreamCopyDescription.Automation.cs
--- a/Engine/Source/Programs/AutomationTool/Scripts/StreamCopyDescription.Automation.cs
+++ b/Engine/Source/Programs/AutomationTool/Scripts/StreamCopyDescription.Automation.cs
@@ -4,6 +4,7 @@ using System;
 using System.Collections.Generic;
 using System.Diagnostics;
 using System.IO;
+using System.Text;
 using System.Text.RegularExpressions;
 using System.Threading;
 using System.Reflection;
@@ -19,6 +20,30 @@ namespace AutomationTool
 	[RequireP4]
 	public class StreamCopyDescription : BuildCommand
 	{
+		/// <summary>
+		/// Hunts down the list of confidential platforms, based on ConfidentialPlatform.ini in the Engine/Config directory (this way we don't 
+		/// need to hardcode the platform names!)
+		/// </summary>
+		/// <returns></returns>
+		private string[] GetPlatformNames()
+		{
+			List<string> Platforms = new List<string>();
+			
+			// two special platform names (for non-special, and combined CLs that affect multiple confidential platforms)
+			Platforms.Add("");
+			Platforms.Add("Multi");
+
+			foreach (string Dir in Directory.EnumerateDirectories(Path.Combine(EngineDirectory.FullName, "Config")))
+			{
+				if (File.Exists(Path.Combine(Dir, "ConfidentialPlatform.ini")))
+				{
+					Platforms.Add(Path.GetFileName(Dir));
+				}
+			}
+
+			return Platforms.ToArray();
+		}
+
 		/// <summary>
 		/// Execute the command
 		/// </summary>
@@ -28,15 +53,26 @@ namespace AutomationTool
 			string Stream = ParseParamValue("Stream", P4Env.BuildRootP4);
 			string Changes = ParseParamValue("Changes", null);
 			string Lockdown = ParseParamValue("Lockdown", "Nick.Penwarden");
+			bool bPlatformSplit = !ParseParam("CombinePlatforms");
 
 			// Get changes which haven't been copied up from the current stream
-			string Descriptions;
+			string AllDescriptions;
 			string SourceStream;
+			string VerbosityFlag = bPlatformSplit ? "" : "-l";
 			int LastCl;
+
+			string[] Platforms = GetPlatformNames();
+			StringBuilder[] DescriptionBuilders = new StringBuilder[Platforms.Length];
+
+			for (int Index = 0; Index < Platforms.Length; Index++)
+			{
+				DescriptionBuilders[Index] = new StringBuilder();
+			}
+
 			if (Changes == null)
 			{
-				IProcessResult Result = P4.P4(String.Format("interchanges -l -S {0}", Stream), AllowSpew: false);
-				Descriptions = Result.Output.Replace("\r\n", "\n");
+				IProcessResult Result = P4.P4(String.Format("interchanges {0} -S {1}", VerbosityFlag, Stream), AllowSpew: false);
+				AllDescriptions = Result.Output;
 				SourceStream = Stream;
 
 				// Get the last submitted change in the source stream
@@ -49,68 +85,135 @@ namespace AutomationTool
 			}
 			else
 			{
-				IProcessResult Result = P4.P4(String.Format("changes -l {0}", Changes), AllowSpew: false);
-				Descriptions = Result.Output.Replace("\r\n", "\n");
+				IProcessResult Result = P4.P4(String.Format("changes {0} {1}", VerbosityFlag, Changes), AllowSpew: false);
+				AllDescriptions = Result.Output;
 				SourceStream = Regex.Replace(Changes, @"(\/(?:\/[^\/]*){2}).*", "$1");
 				LastCl = Int32.Parse(Regex.Replace(Changes, ".*,", ""));
 			}
 
-			// Clean any workspace names that may reveal internal information
-			Descriptions = Regex.Replace(Descriptions, "(Change[^@]*)@.*", "$1", RegexOptions.Multiline);
-
-			// Remove changes by the build machine
-			Descriptions = Regex.Replace(Descriptions, "[^\n]*buildmachine\n(\n|\t[^\n]*\n)*", "");
-
-			// Remove all the tags we don't care about
-			Descriptions = Regex.Replace(Descriptions, "^[ \t]*#(rb|fyi|codereview|lockdown)\\s.*$", "", RegexOptions.Multiline);
-
-			// Empty out lines which just contain whitespace
-			Descriptions = Regex.Replace(Descriptions, "^[ \t]+$", "", RegexOptions.Multiline);
-
-			// Remove multiple consecutive blank lines
-			Descriptions = Regex.Replace(Descriptions, "\n\n+", "\n\n");
+			if (bPlatformSplit)
+			{
+				string[] Lines = AllDescriptions.Split("\n".ToCharArray());
+				foreach (string Line in Lines)
+				{
+					// @todo replace with regexes!!
+					string[] Tokens = Line.Split(" ".ToCharArray());
+					if (Tokens.Length > 2 && Tokens[0] == "Change")
+					{
+						IProcessResult Result = P4.P4(String.Format("describe -s {0}", Tokens[1]), AllowSpew: false);
 
-			// Only include one newline at the end of each description
-			Descriptions = Regex.Replace(Descriptions, "\n+Change", "\n\nChange");
+						// Affected files ... is the splitting point
+						int AffectedFilesPos = Result.Output.IndexOf("Affected files ...");
+						string Description = Result.Output.Substring(0, AffectedFilesPos);
+						string Files = Result.Output.Substring(AffectedFilesPos);
 
-			// Remove merge-only changelists
-			Descriptions = Regex.Replace(Descriptions, "(?<=(^|\\n))Change .*\\s*Merging .* to .*\\s*\\n(?=(Change|$))", "");
+						// look for the NDA platforms in the list of files (skipping over the "no" and "multi" platforms
+						int WhichPlatform = 0;
+						for (int Index = 2; Index < Platforms.Length; Index++)
+						{
+							// we search by directory in the form of /Platform/
+							if (Files.Contains("/" + Platforms[Index] + "/"))
+							{
+								// if we contained multiple files, then we put into the Multi file, and someone will have to manually deal with it!!
+								if (WhichPlatform == 0)
+								{
+									WhichPlatform = Index;
+								}
+								else
+								{
+									WhichPlatform = 1;
+								}
+							}
+						}
 
-			// Figure out the target stream
-			IProcessResult StreamResult = P4.P4(String.Format("stream -o {0}", Stream), AllowSpew: false);
-			if (StreamResult.ExitCode != 0)
-			{
-				throw new AutomationException("Couldn't get stream description for {0}", Stream);
+						// add this description to the proper platform
+						DescriptionBuilders[WhichPlatform].AppendLine(Description);
+					}
+				}
 			}
-			string Target = P4Spec.FromString(StreamResult.Output).GetField("Parent");
-			if(Target == null)
+			else
 			{
-				throw new AutomationException("Couldn't get parent stream for {0}", Stream);
+				DescriptionBuilders[0].Append(AllDescriptions);
 			}
 
-			// Write the output file
-			string OutputDirName = Path.Combine(CommandUtils.CmdEnv.LocalRoot, "Engine", "Intermediate");
-			CommandUtils.CreateDirectory(OutputDirName);
-			string OutputFileName = Path.Combine(OutputDirName, "Changes.txt");
-			using (StreamWriter Writer = new StreamWriter(OutputFileName))
+
+			for (int PlatformIndex = 0; PlatformIndex < Platforms.Length; PlatformIndex++)
 			{
-				Writer.WriteLine("Copying {0} to {1} (Source: {2} @ {3})", Stream, Target.Trim(), SourceStream, LastCl);
-				Writer.WriteLine("#lockdown {0}", Lockdown);
-				Writer.WriteLine();
-				Writer.WriteLine("==========================");
-				Writer.WriteLine("MAJOR FEATURES + CHANGES");
-				Writer.WriteLine("==========================");
-				Writer.WriteLine();
-
-				foreach (string Line in Descriptions.Split('\n'))
+				string Desc = DescriptionBuilders[PlatformIndex].ToString().Replace("\r\n", "\n");
+
+				if (Desc == "")
 				{
-					Writer.WriteLine(Line);
+					continue;
 				}
-			}
-			Log("Written {0}.", OutputFileName);
 
-			// Open it with the default text editor
-			Process.Start(OutputFileName);
+				// Clean any workspace names that may reveal internal information
+				Desc = Regex.Replace(Desc, "(Change[^@]*)@.*", "$1", RegexOptions.Multiline);
+
+				// Remove changes by the build machine
+				Desc = Regex.Replace(Desc, "[^\n]*buildmachine\n(\n|\t[^\n]*\n)*", "");
+
+				// Remove all the tags we don't care about
+				Desc = Regex.Replace(Desc, "^[ \t]*#(rb|fyi|codereview|lockdown)\\s.*$", "", RegexOptions.Multiline);
+
+				// Empty out lines which just contain whitespace
+				Desc = Regex.Replace(Desc, "^[ \t]+$", "", RegexOptions.Multiline);
+
+				// Remove multiple consecutive blank lines
+				Desc = Regex.Replace(Desc, "\n\n+", "\n\n");
+
+				// Only include one newline at the end of each description
+				Desc = Regex.Replace(Desc, "\n+Change", "\n\nChange");
+
+				// Remove merge-only changelists
+				Desc = Regex.Replace(Desc, "(?<=(^|\\n))Change .*\\s*Merging .* to .*\\s*\\n(?=(Change|$))", "");
+
+				// Figure out the target stream
+				IProcessResult StreamResult = P4.P4(String.Format("stream -o {0}", Stream), AllowSpew: false);
+				if (StreamResult.ExitCode != 0)
+				{
+					throw new AutomationException("Couldn't get stream description for {0}", Stream);
+				}
+				string Target = P4Spec.FromString(StreamResult.Output).GetField("Parent");
+				if (Target == null)
+				{
+					throw new AutomationException("Couldn't get parent stream for {0}", Stream);
+				}
+
+				// Write the output file
+				string OutputDirName = Path.Combine(CommandUtils.CmdEnv.LocalRoot, "Engine", "Intermediate");
+				CommandUtils.CreateDirectory(OutputDirName);
+				string OutputFileName = Path.Combine(OutputDirName, string.Format("Changes{0}.txt", Platforms[PlatformIndex]));
+				using (StreamWriter Writer = new StreamWriter(OutputFileName))
+				{
+					if (PlatformIndex == 1)
+					{
+						Writer.WriteLine("XXXXXXXXXXXXXXXXXXXXXXXXXXXXXXXXXXXXXXXXXXXXXXXXXXXXXXXXXXXXXXXXXXXXXXXXXXXXXXXXXXXXXXXXXXX");
+						Writer.WriteLine("CHANGES WITH MULTIPLE PLATFORMS!!! YOU MUST COPY THESE INTO THE OTHER ONES AS MAKES SENSE!!");
+						Writer.WriteLine("XXXXXXXXXXXXXXXXXXXXXXXXXXXXXXXXXXXXXXXXXXXXXXXXXXXXXXXXXXXXXXXXXXXXXXXXXXXXXXXXXXXXXXXXXXX");
+						Writer.WriteLine();
+						Writer.WriteLine();
+					}
+					else
+					{
+						Writer.WriteLine("Copying {0} to {1} (Source: {2} @ {3})", Stream, Target.Trim(), SourceStream, LastCl);
+						Writer.WriteLine("#lockdown {0}", Lockdown);
+						Writer.WriteLine();
+						Writer.WriteLine("=====================================");
+						Writer.WriteLine("{0} MAJOR FEATURES + CHANGES", Platforms[PlatformIndex]);
+						Writer.WriteLine("=====================================");
+						Writer.WriteLine();
+					}
+
+					foreach (string Line in Desc.Split('\n'))
+					{
+						Writer.WriteLine(Line);
+					}
+				}
+				Log("Written {0}.", OutputFileName);
+
+				// Open it with the default text editor
+				Process.Start(OutputFileName);
+			}
 		}
 	}
 }
