UAT: Fix exception when staging targets that use another platform's source data (concatenating path fragments ending in slashes was resulting in duplicate slashes).

diff --git a/Engine/Source/Programs/AutomationTool/AutomationUtils/StagedFileSystemReference.cs b/Engine/Source/Programs/AutomationTool/AutomationUtils/StagedFileSystemReference.cs
--- a/Engine/Source/Programs/AutomationTool/AutomationUtils/StagedFileSystemReference.cs
+++ b/Engine/Source/Programs/AutomationTool/AutomationUtils/StagedFileSystemReference.cs
@@ -76,7 +76,7 @@ public abstract class StagedFileSystemReference
 		StringBuilder NewName = new StringBuilder(BaseDirectory);
 		foreach (string Fragment in Fragments)
 		{
-			if (NewName.Length > 0)
+			if (NewName.Length > 0 && (NewName[NewName.Length - 1] != '/' && NewName[NewName.Length - 1] != '\\'))
 			{
 				NewName.Append('/');
 			}
