BuildGraph: Fix badges and reports containing dependencies on nodes behind triggers, causing jobs to never finish on the build system.

diff --git a/Engine/Source/Programs/AutomationTool/BuildGraph/Graph.cs b/Engine/Source/Programs/AutomationTool/BuildGraph/Graph.cs
--- a/Engine/Source/Programs/AutomationTool/BuildGraph/Graph.cs
+++ b/Engine/Source/Programs/AutomationTool/BuildGraph/Graph.cs
@@ -489,7 +489,7 @@ namespace AutomationTool
 				JsonWriter.WriteArrayStart("Badges");
 				foreach (Badge Badge in Badges)
 				{
-					Node[] Dependencies = Badge.Nodes.Where(x => NodesToExecute.Contains(x)).ToArray();
+					Node[] Dependencies = Badge.Nodes.Where(x => NodesToExecute.Contains(x) && x.ControllingTrigger == Trigger).ToArray();
 					if (Dependencies.Length > 0)
 					{
 						// Reduce that list to the smallest subset of direct dependencies
@@ -520,7 +520,7 @@ namespace AutomationTool
 				JsonWriter.WriteArrayStart("Reports");
 				foreach (Report Report in NameToReport.Values)
 				{
-					Node[] Dependencies = Report.Nodes.Where(x => NodesToExecute.Contains(x)).ToArray();
+					Node[] Dependencies = Report.Nodes.Where(x => NodesToExecute.Contains(x) && x.ControllingTrigger == Trigger).ToArray();
 					if (Dependencies.Length > 0)
 					{
 						// Reduce that list to the smallest subset of direct dependencies
