UAT: Fix exception running commands with the [DoesNotNeedP4CL] attribute.

diff --git a/Engine/Source/Programs/AutomationTool/AutomationUtils/P4Environment.cs b/Engine/Source/Programs/AutomationTool/AutomationUtils/P4Environment.cs
--- a/Engine/Source/Programs/AutomationTool/AutomationUtils/P4Environment.cs
+++ b/Engine/Source/Programs/AutomationTool/AutomationUtils/P4Environment.cs
@@ -242,8 +242,14 @@ namespace AutomationTool
 				Log.TraceInformation("  User: {0}", User);
 				Log.TraceInformation("  Client: {0}", Client);
 				Log.TraceInformation("  Branch: {0}", Branch);
-				Log.TraceInformation("  Last Change: {0}", Changelist);
-				Log.TraceInformation("  Last Code Change: {0}", CodeChangelist);
+				if (ChangelistInternal != -1)
+				{
+					Log.TraceInformation("  Last Change: {0}", Changelist);
+				}
+				if (CodeChangelistInternal != -1)
+				{
+					Log.TraceInformation("  Last Code Change: {0}", CodeChangelist);
+				}
 			}
 
 			// Write all the environment variables to the log
