Gauntlet

 - Split Gauntlet into Gauntlet & GauntletExtras. The latter contains things NotForLicensees such as Google sheets integration, Orion tests (which contain credentials..) etc.
 - Swapped Mac/Linux blacklisting of Gauntlet for Gauntlet Extras.
 - Changed some ITestNode interfaces from properties to functions for consistency and ease of implementation
 - Added base test node that simplifies some of the ITestNode interface and reduces amount of code in UnrealTestNode
 - Simplified reporting of completion/status for tests deriving from UnrealTestNode. These tests should now implement GetUnrealTestResult or SetUnrealTestResult
 - Moved creation and provisioning of Epic accounts into RunOrionTests (GauntletExtras)
 - Added support for passing isolated params to tests via -tests=TestName(param1,param2=f00)
 - Fixed tests getting empty params

#tests preflighted orion tests and editor+monolithics

diff --git a/Engine/Source/Programs/AutomationTool/AutomationUtils/LinuxHostPlatform.cs b/Engine/Source/Programs/AutomationTool/AutomationUtils/LinuxHostPlatform.cs
--- a/Engine/Source/Programs/AutomationTool/AutomationUtils/LinuxHostPlatform.cs
+++ b/Engine/Source/Programs/AutomationTool/AutomationUtils/LinuxHostPlatform.cs
@@ -108,7 +108,7 @@ namespace AutomationTool
 		public override bool IsScriptModuleSupported(string ModuleName)
 		{
 			// @todo: add more unsupported modules here
-			if (String.Equals(ModuleName, "Gauntlet", StringComparison.InvariantCultureIgnoreCase))
+			if (String.Equals(ModuleName, "GauntletExtras", StringComparison.InvariantCultureIgnoreCase))
 			{
 				return false;
 			}
diff --git a/Engine/Source/Programs/AutomationTool/AutomationUtils/MacHostPlatform.cs b/Engine/Source/Programs/AutomationTool/AutomationUtils/MacHostPlatform.cs
--- a/Engine/Source/Programs/AutomationTool/AutomationUtils/MacHostPlatform.cs
+++ b/Engine/Source/Programs/AutomationTool/AutomationUtils/MacHostPlatform.cs
@@ -130,7 +130,7 @@ namespace AutomationTool
 		public override bool IsScriptModuleSupported(string ModuleName)
 		{
 			// @todo: add more unsupported modules here
-			if (String.Equals(ModuleName, "Gauntlet", StringComparison.InvariantCultureIgnoreCase))
+			if (String.Equals(ModuleName, "GauntletExtras", StringComparison.InvariantCultureIgnoreCase))
 			{
 				return false;
 			}
