Re-disabled Gauntlet for Linux/Mac.

diff --git a/Engine/Source/Programs/AutomationTool/AutomationUtils/LinuxHostPlatform.cs b/Engine/Source/Programs/AutomationTool/AutomationUtils/LinuxHostPlatform.cs
--- a/Engine/Source/Programs/AutomationTool/AutomationUtils/LinuxHostPlatform.cs
+++ b/Engine/Source/Programs/AutomationTool/AutomationUtils/LinuxHostPlatform.cs
@@ -108,7 +108,7 @@ namespace AutomationTool
 		public override bool IsScriptModuleSupported(string ModuleName)
 		{
 			// @todo: add more unsupported modules here
-			if (String.Equals(ModuleName, "GauntletExtras", StringComparison.InvariantCultureIgnoreCase))
+			if (ModuleName.StartsWith("Gauntlet", StringComparison.OrdinalIgnoreCase))
 			{
 				return false;
 			}
diff --git a/Engine/Source/Programs/AutomationTool/AutomationUtils/MacHostPlatform.cs b/Engine/Source/Programs/AutomationTool/AutomationUtils/MacHostPlatform.cs
--- a/Engine/Source/Programs/AutomationTool/AutomationUtils/MacHostPlatform.cs
+++ b/Engine/Source/Programs/AutomationTool/AutomationUtils/MacHostPlatform.cs
@@ -130,7 +130,7 @@ namespace AutomationTool
 		public override bool IsScriptModuleSupported(string ModuleName)
 		{
 			// @todo: add more unsupported modules here
-			if (String.Equals(ModuleName, "GauntletExtras", StringComparison.InvariantCultureIgnoreCase))
+			if (ModuleName.StartsWith("Gauntlet", StringComparison.OrdinalIgnoreCase))
 			{
 				return false;
 			}
