Changing commandutils to attempt to find the Win8.1sdk signtool before the win10 one.  There's currently a bug related to vs2017/server2012r2/win10sdk signtool that causes it to exit with an undefined error when signing from a service account.  Using the win8.1sdk circumvents this issue.

diff --git a/Engine/Source/Programs/AutomationTool/AutomationUtils/CommandUtils.cs b/Engine/Source/Programs/AutomationTool/AutomationUtils/CommandUtils.cs
--- a/Engine/Source/Programs/AutomationTool/AutomationUtils/CommandUtils.cs
+++ b/Engine/Source/Programs/AutomationTool/AutomationUtils/CommandUtils.cs
@@ -2573,8 +2573,8 @@ namespace AutomationTool
 		{
 			string[] PossibleSignToolNames =
 			{
-				"C:/Program Files (x86)/Windows Kits/10/bin/x86/SignTool.exe",
-				"C:/Program Files (x86)/Windows Kits/8.1/bin/x86/SignTool.exe"
+				"C:/Program Files (x86)/Windows Kits/8.1/bin/x86/SignTool.exe",
+				"C:/Program Files (x86)/Windows Kits/10/bin/x86/SignTool.exe"
 			};
 
 			foreach(string PossibleSignToolName in PossibleSignToolNames)
@@ -2816,7 +2816,7 @@ namespace AutomationTool
 				//@TODO: Verbosity choosing
 				//  /v will spew lots of info
 				//  /q does nothing on success and minimal output on failure
-				string CodeSignArgs = String.Format("sign{0} /a /n \"{1}\" /t {2} /v {3}", SpecificStoreArg, SigningIdentity, TimestampServer[NumTrials % TimestampServer.Length], FilesToSign);
+				string CodeSignArgs = String.Format("sign{0} /a /n \"{1}\" /t {2} /debug {3}", SpecificStoreArg, SigningIdentity, TimestampServer[NumTrials % TimestampServer.Length], FilesToSign);
 
 				IProcessResult Result = CommandUtils.Run(SignToolName, CodeSignArgs, null, CommandUtils.ERunOptions.AllowSpew);
 				++NumTrials;
