Updating .NET target versions of a bunch of csproj's to 4.5
Updated script to check for .netcore projects, ignore monocs projects, and whitelisted UnrealVS as it requires 4.6 to be targeted for vs2017

diff --git a/Engine/Source/Programs/AutomationTool/Scripts/CheckCsprojDotNetVersion.cs b/Engine/Source/Programs/AutomationTool/Scripts/CheckCsprojDotNetVersion.cs
--- a/Engine/Source/Programs/AutomationTool/Scripts/CheckCsprojDotNetVersion.cs
+++ b/Engine/Source/Programs/AutomationTool/Scripts/CheckCsprojDotNetVersion.cs
@@ -26,20 +26,32 @@ class CheckCsprojDotNetVersion : BuildCommand
 		Regex FrameworkRegex = new Regex("<TargetFrameworkVersion>v(\\d\\.\\d\\.?\\d?)<\\/TargetFrameworkVersion>");
 		Regex PossibleAppConfigRegex = new Regex("<TargetFrameworkProfile>(.+)<\\/TargetFrameworkProfile>");
 		Regex AppConfigRegex = new Regex("<supportedRuntime version=\"v(\\d\\.\\d\\.?\\d?)\" sku=\"\\.NETFramework,Version=v(\\d\\.\\d\\.?\\d?),Profile=(.+)\"\\/>");
-		foreach (FileReference CsProj in DirectoryReference.EnumerateFiles(EngineDir, "*.csproj", SearchOption.AllDirectories))
-		{
-			if (CsProj.ContainsName(new FileSystemName("ThirdParty"), EngineDir) ||
-				(CsProj.ContainsName(new FileSystemName("UE4TemplateProject"), EngineDir) && CsProj.GetFileName().Equals("ProjectTemplate.csproj")))
-			{
-				continue;
-			}
+        Regex DotNetCoreRegex = new Regex("<TargetFramework>(netcoreapp2.0|netstandard2.0)<\\/TargetFramework>");
+        foreach (FileReference CsProj in DirectoryReference.EnumerateFiles(EngineDir, "*.csproj", SearchOption.AllDirectories))
+        {
+            if (CsProj.ContainsName(new FileSystemName("ThirdParty"), EngineDir) ||
+                (CsProj.ContainsName(new FileSystemName("UE4TemplateProject"), EngineDir) && CsProj.GetFileName().Equals("ProjectTemplate.csproj")) ||
+                CsProj.GetFileNameWithoutExtension().ToLower().Contains("_mono") ||
+                CsProj.GetFileNameWithoutExtension().ToLower().Contains("unrealvs"))
+
+            {
+                continue;
+            }
+
+            // read in the file
+            string Contents = File.ReadAllText(CsProj.FullName);
+            Match Match = DotNetCoreRegex.Match(Contents);
+            // Check if we're a _NETCore app, ignore these.
+            if (Match.Success)
+            {
+                continue;
+            }
+
 
-			// read in the file
-			string Contents = File.ReadAllText(CsProj.FullName);
-			Match m = FrameworkRegex.Match(Contents);
-			if (m.Success)
+            Match = FrameworkRegex.Match(Contents);
+            if (Match.Success)
 			{
-				string TargetedVersion = m.Groups[1].Value;
+				string TargetedVersion = Match.Groups[1].Value;
 				// make sure we match, throw warning otherwise
 				if (!DesiredTargetVersion.Equals(TargetedVersion, StringComparison.InvariantCultureIgnoreCase))
 				{
@@ -49,8 +61,8 @@ class CheckCsprojDotNetVersion : BuildCommand
 			// if we don't have a TargetFrameworkVersion, check for the existence of TargetFrameworkProfile.
 			else
 			{
-				m = PossibleAppConfigRegex.Match(Contents);
-				if (!m.Success)
+                Match = PossibleAppConfigRegex.Match(Contents);
+				if (!Match.Success)
 				{
 					CommandUtils.Log("No TargetFrameworkVersion or TargetFrameworkProfile found for project {0}, is it a mono project? If not, does it compile properly?", CsProj);
 					continue;
@@ -58,7 +70,7 @@ class CheckCsprojDotNetVersion : BuildCommand
 
 				// look for the app config
 				FileReference AppConfigFile = FileReference.Combine(CsProj.Directory, "app.config");
-				string Profile = m.Groups[1].Value;
+				string Profile = Match.Groups[1].Value;
 				if (!FileReference.Exists(AppConfigFile))
 				{
 					CommandUtils.Log("Found TargetFrameworkProfile but no associated app.config containing the version for project {0}.", CsProj);
@@ -67,8 +79,8 @@ class CheckCsprojDotNetVersion : BuildCommand
 
 				// read in the app config
 				Contents = File.ReadAllText(AppConfigFile.FullName);
-				m = AppConfigRegex.Match(Contents);
-				if (!m.Success)
+                Match = AppConfigRegex.Match(Contents);
+				if (!Match.Success)
 				{
 					CommandUtils.Log("Couldn't find a supportedRuntime match for the version in the app.config for project {0}.", CsProj);
 					continue;
@@ -77,9 +89,9 @@ class CheckCsprojDotNetVersion : BuildCommand
 				// Version1 is the one that appears right after supportedRuntime
 				// Version2 is the one in the sku
 				// ProfileString should match the TargetFrameworkProfile from the csproj
-				string Version1String = m.Groups[1].Value;
-				string Version2String = m.Groups[2].Value;
-				string ProfileString = m.Groups[3].Value;
+				string Version1String = Match.Groups[1].Value;
+				string Version2String = Match.Groups[2].Value;
+				string ProfileString = Match.Groups[3].Value;
 
 				// not sure how this is possible, but check for it anyway
 				if (!ProfileString.Equals(Profile, StringComparison.InvariantCultureIgnoreCase))
diff --git a/Engine/Source/Programs/IncludeTool/IncludeTool/App.config b/Engine/Source/Programs/IncludeTool/IncludeTool/App.config
--- a/Engine/Source/Programs/IncludeTool/IncludeTool/App.config
+++ b/Engine/Source/Programs/IncludeTool/IncludeTool/App.config
@@ -1,6 +1,6 @@
-﻿<?xml version="1.0" encoding="utf-8" ?>
+<?xml version="1.0" encoding="utf-8"?>
 <configuration>
     <startup> 
-        <supportedRuntime version="v4.0" sku=".NETFramework,Version=v4.5.2" />
+        <supportedRuntime version="v4.0" sku=".NETFramework,Version=v4.5"/>
     </startup>
-</configuration>
\ No newline at end of file
+</configuration>
diff --git a/Engine/Source/Programs/IncludeTool/IncludeTool/IncludeTool.csproj b/Engine/Source/Programs/IncludeTool/IncludeTool/IncludeTool.csproj
--- a/Engine/Source/Programs/IncludeTool/IncludeTool/IncludeTool.csproj
+++ b/Engine/Source/Programs/IncludeTool/IncludeTool/IncludeTool.csproj
@@ -9,7 +9,7 @@
     <AppDesignerFolder>Properties</AppDesignerFolder>
     <RootNamespace>IncludeTool</RootNamespace>
     <AssemblyName>IncludeTool</AssemblyName>
-    <TargetFrameworkVersion>v4.5.2</TargetFrameworkVersion>
+    <TargetFrameworkVersion>v4.5</TargetFrameworkVersion>
     <FileAlignment>512</FileAlignment>
     <TargetFrameworkProfile />
   </PropertyGroup>
diff --git a/Engine/Source/Programs/MemoryProfiler2/App.config b/Engine/Source/Programs/MemoryProfiler2/App.config
--- a/Engine/Source/Programs/MemoryProfiler2/App.config
+++ b/Engine/Source/Programs/MemoryProfiler2/App.config
@@ -1,7 +1,7 @@
 <?xml version="1.0"?>
 <configuration>
 	<startup>
-		<supportedRuntime version="v4.0" sku=".NETFramework,Version=v4.0,Profile=Client"/>
+		<supportedRuntime version="v4.0" sku=".NETFramework,Version=v4.5"/>
 	</startup>
 	<runtime>
 		<generatePublisherEvidence enabled="false"/>
diff --git a/Engine/Source/Programs/MemoryProfiler2/MemoryProfiler2.csproj b/Engine/Source/Programs/MemoryProfiler2/MemoryProfiler2.csproj
--- a/Engine/Source/Programs/MemoryProfiler2/MemoryProfiler2.csproj
+++ b/Engine/Source/Programs/MemoryProfiler2/MemoryProfiler2.csproj
@@ -23,9 +23,10 @@
     <OldToolsVersion>3.5</OldToolsVersion>
     <UpgradeBackupLocation>
     </UpgradeBackupLocation>
-    <TargetFrameworkVersion>v4.0</TargetFrameworkVersion>
+    <TargetFrameworkVersion>v4.5</TargetFrameworkVersion>
     <IsWebBootstrapper>false</IsWebBootstrapper>
-    <TargetFrameworkProfile>Client</TargetFrameworkProfile>
+    <TargetFrameworkProfile>
+    </TargetFrameworkProfile>
     <IsWebBootstrapper>false</IsWebBootstrapper>
     <PublishUrl>publish\</PublishUrl>
     <Install>true</Install>
@@ -51,6 +52,7 @@
     <ErrorReport>prompt</ErrorReport>
     <WarningLevel>4</WarningLevel>
     <PlatformTarget>AnyCPU</PlatformTarget>
+    <Prefer32Bit>false</Prefer32Bit>
   </PropertyGroup>
   <PropertyGroup Condition=" '$(Configuration)|$(Platform)' == 'Release|AnyCPU' ">
     <DebugType>pdbonly</DebugType>
@@ -60,6 +62,7 @@
     <ErrorReport>prompt</ErrorReport>
     <WarningLevel>4</WarningLevel>
     <PlatformTarget>AnyCPU</PlatformTarget>
+    <Prefer32Bit>false</Prefer32Bit>
   </PropertyGroup>
   <PropertyGroup>
     <ApplicationIcon>Icon.ico</ApplicationIcon>
