UAT: Rewrite ListThirdPartySoftware command to fix current blank output, and to correctly support foreign projects and projects which do not have the same name as the target.

diff --git a/Engine/Source/Programs/AutomationTool/Scripts/ListThirdPartySoftware.Automation.cs b/Engine/Source/Programs/AutomationTool/Scripts/ListThirdPartySoftware.Automation.cs
--- a/Engine/Source/Programs/AutomationTool/Scripts/ListThirdPartySoftware.Automation.cs
+++ b/Engine/Source/Programs/AutomationTool/Scripts/ListThirdPartySoftware.Automation.cs
@@ -10,132 +10,124 @@ using System.Text.RegularExpressions;
 using Tools.DotNETCommon;
 
 [Help("Lists TPS files associated with any source used to build a specified target(s). Grabs TPS files associated with source modules, content, and engine shaders.")]
-[Help("Target", "one or more Project|Config|Platform values describing all the targets to get associated TPS files for.")]
-[Help("Debug", "Log more detailed info for debugging")]
+[Help("Target", "One or more UBT command lines to enumerate associated TPS files for (eg. UE4Game Win64 Development).")]
 class ListThirdPartySoftware : BuildCommand
 {
 	public override void ExecuteBuild()
 	{
 		CommandUtils.Log("************************* List Third Party Software");
-		bool DebugRun = ParseParam("Debug");
 
-		// first get a list of all TPS files available (as a lookup of Dir -> List of TPS files in Dir).
-		var AllTPSFiles = 
-			Directory.EnumerateFiles(CombinePaths(CmdEnv.LocalRoot), "*.tps", SearchOption.AllDirectories)
-			.ToLookup(TPSFile => Path.GetDirectoryName(TPSFile), StringComparer.InvariantCultureIgnoreCase);
-
-		ParseParamValues(Params, "Target") // Grab all the Target params.
-		.Select(TargetArg => TargetArg.Split(new[] { '|' }, 3)) // split the arg up by |
-		.Where(Target => Target.Length == 3) // ensure it has three parts
-		.Select(Target => new // strongly name the parts of the target and 
+		// Parse the list of targets to list TPS for. Each target is specified by -Target="Name|Configuration|Platform" on the command line.
+		HashSet<FileReference> TpsFiles = new HashSet<FileReference>();
+		foreach(string Target in ParseParamValues(Params, "Target"))
 		{
-			Name          = Target[0],
-			Configuration = Target[1],
-			Platform      = Target[2],
-			// create a list of unsupported platform strings for that target, which we'll use later to cull out irrelvant TPS files
-			UnsupportedSubstrings =
-				Utils.MakeListOfUnsupportedPlatforms(new List<UnrealTargetPlatform> { (UnrealTargetPlatform)Enum.Parse(typeof(UnrealTargetPlatform), Target[2]) })
-				.Select(PlatformName => Path.DirectorySeparatorChar + PlatformName + Path.DirectorySeparatorChar) // convert to /Platform/
-				.ToList(), // project to list
-		})
-		.SelectMany(Target =>
-			// run UBT on each target to get the build folders, add a few hard-coded paths, then search for TPS files in relevant folders to them
-			EnumerateBuildFolders(Target.Name, Target.Configuration, Target.Platform, DebugRun)
-			// tack on the target content directory, which we could infer from one of a few places
-			.Concat(new[] 
+			// Get the path to store the exported JSON target data
+			FileReference OutputFile = FileReference.Combine(CommandUtils.EngineDirectory, "Intermediate", "Build", "ThirdParty.json");
+
+			IProcessResult Result = Run(UE4Build.GetUBTExecutable(), String.Format("{0} -jsonexport=\"{1}\" -skipbuild", Target.Replace('|', ' '), OutputFile.FullName), Options: ERunOptions.Default);
+			if(Result.ExitCode != 0)
+			{
+				throw new AutomationException("Failed to run UBT");
+			}
+
+			// Read the exported target info back in
+			JsonObject Object = JsonObject.Read(OutputFile);
+
+			// Get the project file if there is one
+			FileReference ProjectFile = null;
+			string ProjectFileName;
+			if(Object.TryGetStringField("ProjectFile", out ProjectFileName))
+			{
+				ProjectFile = new FileReference(ProjectFileName);
+			}
+
+			// Get the default paths to search
+			HashSet<DirectoryReference> DirectoriesToScan = new HashSet<DirectoryReference>();
+			DirectoriesToScan.Add(DirectoryReference.Combine(CommandUtils.EngineDirectory, "Shaders"));
+			DirectoriesToScan.Add(DirectoryReference.Combine(CommandUtils.EngineDirectory, "Content"));
+			if(ProjectFile != null)
+			{
+				DirectoriesToScan.Add(DirectoryReference.Combine(ProjectFile.Directory, "Content"));
+			}
+
+			// Get the variables to be expanded in any runtime dependencies variables
+			Dictionary<string, string> Variables = new Dictionary<string, string>();
+			Variables.Add("EngineDir", CommandUtils.EngineDirectory.FullName);
+			if(ProjectFile != null)
+			{
+				Variables.Add("ProjectDir", ProjectFile.Directory.FullName);
+			}
+
+			// Add all the paths for each module, and its runtime dependencies
+			JsonObject Modules = Object.GetObjectField("Modules");
+			foreach(string ModuleName in Modules.KeyNames)
+			{
+				JsonObject Module = Modules.GetObjectField(ModuleName);
+				DirectoriesToScan.Add(new DirectoryReference(Module.GetStringField("Directory")));
+
+				foreach(JsonObject RuntimeDependency in Module.GetObjectArrayField("RuntimeDependencies"))
 				{
-					Path.Combine(Path.GetDirectoryName(Target.Name), "Content"), // The target could be a .uproject, so look for a content folder there
-					Path.Combine(CmdEnv.LocalRoot, Target.Name, "Content"), // the target could be a root project, so look for a content folder there
-				}.Where(Folder => Directory.Exists(Folder)) // only accept Content folders that actually exist.
-			)
-			// tack on hard-coded engine shaders and content folder
-			.Concat(new List<string> { "Shaders", "Content", }.Select(Folder => Path.Combine(CmdEnv.LocalRoot, "Engine", Folder)))
-			// canonicalize the path and convert separators. We do string compares below, so we need separators to be in consistent formats.
-			.Select(FolderToScan => CombinePaths(FolderToScan))
-			// scan each folder for TPS files that appear relevant to that folder
-			.SelectMany(FolderToScan =>
-				AllTPSFiles
-				// A relevant folder is one that is in a subfolders or parent folder to that folder
-				.Where(TPSFileDir => FolderToScan.StartsWith(TPSFileDir.Key, StringComparison.InvariantCultureIgnoreCase) || TPSFileDir.Key.StartsWith(FolderToScan, StringComparison.InvariantCultureIgnoreCase))
-				// flatten the list of all the TPS files in any folders we found (remember we were just looking at folders above, which could have a list of TPS files)
-				.SelectMany(TPSFileDir => TPSFileDir)
-				// filter out TPS from unwanted platforms for the current target.
-				// This is a bit hard to follow, but we are only looking for unsupported platform names INSIDE the folder to search.
-				// for instance, if the folder is Foo/PS4/Plugins/PadSupport and we are building for windows, we don't want to filter it out even though it has PS4 in the name.
-				// but if the TPS file is Foo/PS4/Plugins/PadSupport/Android/TPS.TPS, then we do want to filter out because of the Android folder.
-				// this basically checks if an unsupported platform string is found AFTER the original folder we are looking in.
-				.Where(TpsFile => Target.UnsupportedSubstrings.Max(PlatformDir => TpsFile.IndexOf(PlatformDir, StringComparison.InvariantCultureIgnoreCase)) < FolderToScan.Length)
-			)
-		)
-		.Distinct() // remove duplicate TPS files.
-		// pull out the redirect if present (but leave a breadcrumb to where the redirect came from)
-		.Select(TPSFile => 
-			// read all the lines in the file
-			File.ReadAllLines(TPSFile)
-			// look for a line with a redirect
-			.Where(Line => Line.IndexOf("Redirect:", StringComparison.InvariantCultureIgnoreCase) >= 0)
-			// grab the redirect file and canonicalize it, tacking on some info on where the redirect was from.
-			.Select(Line => CombinePaths(Path.GetFullPath(Path.Combine(Path.GetDirectoryName(TPSFile), Line.Split(new[] { ':' },2)[1].Trim()))) + " (redirect from " + TPSFile + ")")
-			// take the first one found, or just use the TPS file directly if no redirect was present.
-			.SingleOrDefault() ?? TPSFile)
-		.OrderBy(TPSFile => TPSFile) // sort them.
-		.ToList().ForEach(TPSFile => UnrealBuildTool.Log.WriteLine(0, string.Empty, LogEventType.Console, TPSFile)); // log them.
-	}
+					string RuntimeDependencyPath = RuntimeDependency.GetStringField("Path");
+					RuntimeDependencyPath = Utils.ExpandVariables(RuntimeDependencyPath, Variables);
+					DirectoriesToScan.Add(new FileReference(RuntimeDependencyPath).Directory);
+				}
+			}
 
-	IEnumerable<string> EnumerateBuildFolders(string Target, string Configuration, string Platform, bool DebugRun)
-	{
-		FileReference OutputFile = FileReference.Combine(CommandUtils.EngineDirectory, "Intermediate", "Build", "ThirdParty.json");
+			// Remove any directories that are under other directories, and sort the output list
+			List<DirectoryReference> SortedDirectoriesToScan = new List<DirectoryReference>();
+			foreach(DirectoryReference DirectoryToScan in DirectoriesToScan.OrderBy(x => x.FullName))
+			{
+				if(SortedDirectoriesToScan.Count == 0 || !DirectoryToScan.IsUnderDirectory(SortedDirectoriesToScan[SortedDirectoriesToScan.Count - 1]))
+				{
+					SortedDirectoriesToScan.Add(DirectoryToScan);
+				}
+			}
 
-		IProcessResult Result = Run(UE4Build.GetUBTExecutable(), String.Format("{0} {1} {2} -jsonexport=\"{3}\" -skipbuild", Target, Configuration, Platform, OutputFile.FullName), Options: DebugRun ? ERunOptions.Default : ERunOptions.NoLoggingOfRunCommand);
-		if(Result.ExitCode != 0)
-		{
-			throw new AutomationException("Failed to run UBT");
+			// Get the platforms to exclude
+			List<UnrealTargetPlatform> SupportedPlatforms = new List<UnrealTargetPlatform> { (UnrealTargetPlatform)Enum.Parse(typeof(UnrealTargetPlatform), Object.GetStringField("Platform")) };
+			FileSystemName[] ExcludePlatformNames = Utils.MakeListOfUnsupportedPlatforms(SupportedPlatforms).Select(x => new FileSystemName(x)).ToArray();
+
+			// Find all the TPS files under the engine directory which match
+			foreach(DirectoryReference DirectoryToScan in SortedDirectoriesToScan)
+			{
+				foreach(FileReference TpsFile in DirectoryReference.EnumerateFiles(DirectoryToScan, "*.tps", SearchOption.AllDirectories))
+				{
+					if(!TpsFile.ContainsAnyNames(ExcludePlatformNames, DirectoryToScan))
+					{
+						TpsFiles.Add(TpsFile);
+					}
+				}
+			}
 		}
 
-		JsonObject Object = JsonObject.Read(OutputFile);
+		// Also add any redirects
+		List<string> OutputMessages = new List<string>();
+		foreach(FileReference TpsFile in TpsFiles)
+		{
+			string Message = TpsFile.FullName;
+
+			string[] Lines = FileReference.ReadAllLines(TpsFile);
+			foreach(string Line in Lines)
+			{
+				const string RedirectPrefix = "Redirect:";
 
-		// local function that takes a RuntimeDependency path and resolves it (replacing Env vars that we support)
-		Func<string, DirectoryReference> ResolveRuntimeDependencyFolder = (string DependencyPath) =>
+				int Idx = Line.IndexOf(RedirectPrefix, StringComparison.InvariantCultureIgnoreCase);
+				if(Idx >= 0)
+				{
+					FileReference RedirectTpsFile = FileReference.Combine(TpsFile.Directory, Line.Substring(Idx + RedirectPrefix.Length).Trim());
+					Message = String.Format("{0} (redirect from {1})", RedirectTpsFile.FullName, TpsFile.FullName);
+					break;
+				}
+			}
+
+			OutputMessages.Add(Message);
+		}
+		OutputMessages.Sort();
+
+		// Print them all out
+		foreach(string OutputMessage in OutputMessages)
 		{
-			return new DirectoryReference(Path.GetDirectoryName(
-				// Regex to replace the env vars we support $(EngineDir|ProjectDir), ignoring case
-				Regex.Replace(DependencyPath, @"\$\((?<Type>Engine|Project)Dir\)", M => 
-					M.Groups["Type"].Value.Equals("Engine", StringComparison.InvariantCultureIgnoreCase)
-						? CommandUtils.EngineDirectory.FullName 
-						: new FileReference(Object.GetStringField("ProjectFile")).Directory.FullName,
-				RegexOptions.IgnoreCase)));
-		};
-
-		JsonObject Modules = Object.GetObjectField("Modules");
-
-		// Create a set of directories used for each binary
-		List<DirectoryReference> DirectoriesToScan = Modules
-			// get all directories that the module uses (source folder and any runtime dependencies)
-			.KeyNames.Select(KeyName => Modules.GetObjectField(KeyName))
-			.SelectMany(Module => Module
-				// resolve any runtime dependency folders and add them.
-				.GetObjectArrayField("RuntimeDependencies").Select(Dependency => ResolveRuntimeDependencyFolder(Dependency.GetStringField("Path")))
-				// Add on the module source directory
-				.Concat(new[] { new DirectoryReference(Module.GetStringField("Directory")) }))
-			// remove any duplicate folders since some modules may be from the same plugin
-			.Distinct()
-			// Project to a list as we need to do an O(n^2) operation below.
-			.ToList();
-
-		List<string> FinalDirs = DirectoriesToScan.Where(RemovalCandidate =>
-			// O(n^2) search to remove subfolders of any we are already searching.
-			// look for directories that aren't subdirectories of any other directory in the list.
-			!DirectoriesToScan.Any(DirectoryToScan =>
-				// != check because this inner loop will eventually check against itself
-				RemovalCandidate != DirectoryToScan &&
-				RemovalCandidate.IsUnderDirectory(DirectoryToScan)))
-			// grab the full name
-			.Select(Dir=>Dir.FullName)
-			// sort the final output
-			.OrderBy(Dir=> Dir)
-			// log the folders
-			.ToList();
-
-		return FinalDirs;
+			CommandUtils.Log(OutputMessage);
+		}
 	}
 }
