Fix crash packaging Paragon in nightly build, due to command parameter to ProjectParams constructor being null.

diff --git a/Engine/Source/Programs/AutomationTool/AutomationUtils/ProjectParams.cs b/Engine/Source/Programs/AutomationTool/AutomationUtils/ProjectParams.cs
--- a/Engine/Source/Programs/AutomationTool/AutomationUtils/ProjectParams.cs
+++ b/Engine/Source/Programs/AutomationTool/AutomationUtils/ProjectParams.cs
@@ -940,11 +940,14 @@ namespace AutomationTool
             }
 
 			// Gather up any '-ini:' arguments and save them. We'll pass these along to other tools that may be spawned in a new process as part of the command.
-			foreach (string Param in Command.Params)
+			if(Command != null)
 			{
-				if (Param.StartsWith("ini:", StringComparison.InvariantCultureIgnoreCase))
+				foreach (string Param in Command.Params)
 				{
-					this.ConfigOverrideParams.Add(Param);
+					if (Param.StartsWith("ini:", StringComparison.InvariantCultureIgnoreCase))
+					{
+						this.ConfigOverrideParams.Add(Param);
+					}
 				}
 			}
 
