Copying Orion-Staging @ 3748653 (Orion/Dev-General @ 3722124) to //UE4/Main

diff --git a/Engine/Source/Editor/PropertyEditor/Private/DetailPropertyRow.cpp b/Engine/Source/Editor/PropertyEditor/Private/DetailPropertyRow.cpp
--- a/Engine/Source/Editor/PropertyEditor/Private/DetailPropertyRow.cpp
+++ b/Engine/Source/Editor/PropertyEditor/Private/DetailPropertyRow.cpp
@@ -44,7 +44,7 @@ FDetailPropertyRow::FDetailPropertyRow(TSharedPtr<FPropertyNode> InPropertyNode,
 		{
 			// We are showing an entirely different object inline.  Generate a layout for it now.
 			ExternalObjectLayout = MakeShared<FDetailLayoutData>();
-			InParentCategory->GetDetailsView()->UpdateSinglePropertyMap(InExternalRootNode, *ExternalObjectLayout);
+			InParentCategory->GetDetailsView()->UpdateSinglePropertyMap(InExternalRootNode, *ExternalObjectLayout, true);
 		}
 
 		if (PropertyNode->GetPropertyKeyNode().IsValid())
@@ -459,7 +459,7 @@ void FDetailPropertyRow::MakeExternalPropertyRowCustomization(TSharedPtr<FStruct
 	}
 }
 
-void FDetailPropertyRow::MakeExternalPropertyRowCustomization(const TArray<UObject*>& InObjects, FName PropertyName, TSharedRef<FDetailCategoryImpl> ParentCategory, struct FDetailLayoutCustomization& OutCustomization)
+void FDetailPropertyRow::MakeExternalPropertyRowCustomization(const TArray<UObject*>& InObjects, FName PropertyName, TSharedRef<FDetailCategoryImpl> ParentCategory, struct FDetailLayoutCustomization& OutCustomization, TOptional<bool> bAllowChildrenOverride, TOptional<bool> bCreateCategoryNodesOverride)
 {
 	TSharedRef<FObjectPropertyNode> RootPropertyNode = MakeShared<FObjectPropertyNode>();
 
@@ -477,6 +477,15 @@ void FDetailPropertyRow::MakeExternalPropertyRowCustomization(const TArray<UObje
 	InitParams.bForceHiddenPropertyVisibility = FPropertySettings::Get().ShowHiddenProperties();
 	InitParams.bCreateCategoryNodes = PropertyName == NAME_None;
 
+	if (bAllowChildrenOverride.IsSet())
+	{
+		InitParams.bAllowChildren = bAllowChildrenOverride.GetValue();
+	}
+	if (bCreateCategoryNodesOverride.IsSet())
+	{
+		InitParams.bCreateCategoryNodes = bCreateCategoryNodesOverride.GetValue();
+	}
+
 	RootPropertyNode->InitNode(InitParams);
 
 	ParentCategory->GetParentLayoutImpl().AddExternalRootPropertyNode(RootPropertyNode);
diff --git a/Engine/Source/Editor/PropertyEditor/Private/DetailPropertyRow.h b/Engine/Source/Editor/PropertyEditor/Private/DetailPropertyRow.h
--- a/Engine/Source/Editor/PropertyEditor/Private/DetailPropertyRow.h
+++ b/Engine/Source/Editor/PropertyEditor/Private/DetailPropertyRow.h
@@ -98,7 +98,7 @@ public:
 	EVisibility GetPropertyVisibility() const { return PropertyVisibility.Get(); }
 
 	static void MakeExternalPropertyRowCustomization(TSharedPtr<FStructOnScope> StructData, FName PropertyName, TSharedRef<FDetailCategoryImpl> ParentCategory, struct FDetailLayoutCustomization& OutCustomization);
-	static void MakeExternalPropertyRowCustomization(const TArray<UObject*>& InObjects, FName PropertyName, TSharedRef<FDetailCategoryImpl> ParentCategory, struct FDetailLayoutCustomization& OutCustomization);
+	static void MakeExternalPropertyRowCustomization(const TArray<UObject*>& InObjects, FName PropertyName, TSharedRef<FDetailCategoryImpl> ParentCategory, struct FDetailLayoutCustomization& OutCustomization, TOptional<bool> bAllowChildrenOverride = TOptional<bool>(), TOptional<bool> bCreateCategoryNodesOverride = TOptional<bool>());
 
 private:
 	/**
diff --git a/Engine/Source/Editor/PropertyEditor/Private/IDetailsViewPrivate.h b/Engine/Source/Editor/PropertyEditor/Private/IDetailsViewPrivate.h
--- a/Engine/Source/Editor/PropertyEditor/Private/IDetailsViewPrivate.h
+++ b/Engine/Source/Editor/PropertyEditor/Private/IDetailsViewPrivate.h
@@ -119,7 +119,7 @@ public:
 	virtual TSharedPtr<class IDetailRootObjectCustomization> GetRootObjectCustomization() const = 0;
 
 	/** Runs the details customization update on a root property node */
-	virtual void UpdateSinglePropertyMap(TSharedPtr<FComplexPropertyNode> InRootPropertyNode, struct FDetailLayoutData& LayoutData) = 0;
+	virtual void UpdateSinglePropertyMap(TSharedPtr<FComplexPropertyNode> InRootPropertyNode, struct FDetailLayoutData& LayoutData, bool bIsExternal) = 0;
 
 	virtual const FCustomPropertyTypeLayoutMap& GetCustomPropertyTypeLayoutMap() const = 0;
 
diff --git a/Engine/Source/Editor/PropertyEditor/Private/PropertyNode.cpp b/Engine/Source/Editor/PropertyEditor/Private/PropertyNode.cpp
--- a/Engine/Source/Editor/PropertyEditor/Private/PropertyNode.cpp
+++ b/Engine/Source/Editor/PropertyEditor/Private/PropertyNode.cpp
@@ -114,7 +114,14 @@ void FPropertyNode::InitNode( const FPropertyNodeInitParams& InitParams )
 	//default to copying from the parent
 	if (ParentNode)
 	{
-		SetNodeFlags(EPropertyNodeFlags::ShowCategories, !!ParentNode->HasNodeFlags(EPropertyNodeFlags::ShowCategories));
+		if (ParentNode->HasNodeFlags(EPropertyNodeFlags::ShowCategories) != 0)
+		{
+			SetNodeFlags(EPropertyNodeFlags::ShowCategories, true);
+		}
+		else
+		{
+			SetNodeFlags(EPropertyNodeFlags::ShowCategories, false);
+		}
 
 		// We are advanced if our parent is advanced or our property is marked as advanced
 		SetNodeFlags(EPropertyNodeFlags::IsAdvanced, ParentNode->HasNodeFlags(EPropertyNodeFlags::IsAdvanced) || bAdvanced );
diff --git a/Engine/Source/Editor/PropertyEditor/Private/PropertyRowGenerator.cpp b/Engine/Source/Editor/PropertyEditor/Private/PropertyRowGenerator.cpp
--- a/Engine/Source/Editor/PropertyEditor/Private/PropertyRowGenerator.cpp
+++ b/Engine/Source/Editor/PropertyEditor/Private/PropertyRowGenerator.cpp
@@ -457,7 +457,7 @@ void FPropertyRowGenerator::UpdateSinglePropertyMap(TSharedPtr<FComplexPropertyN
 	// Reset everything
 	LayoutData.ClassToPropertyMap.Empty();
 
-	TSharedPtr<FDetailLayoutBuilderImpl> DetailLayout = MakeShareable(new FDetailLayoutBuilderImpl(InRootPropertyNode, LayoutData.ClassToPropertyMap, PropertyUtilities, nullptr));
+	TSharedPtr<FDetailLayoutBuilderImpl> DetailLayout = MakeShareable(new FDetailLayoutBuilderImpl(InRootPropertyNode, LayoutData.ClassToPropertyMap, PropertyUtilities, nullptr, false));
 	LayoutData.DetailLayout = DetailLayout;
 
 	TSharedPtr<FComplexPropertyNode> RootPropertyNode = InRootPropertyNode;
diff --git a/Engine/Source/Editor/PropertyEditor/Private/SDetailsViewBase.cpp b/Engine/Source/Editor/PropertyEditor/Private/SDetailsViewBase.cpp
--- a/Engine/Source/Editor/PropertyEditor/Private/SDetailsViewBase.cpp
+++ b/Engine/Source/Editor/PropertyEditor/Private/SDetailsViewBase.cpp
@@ -341,16 +341,16 @@ void SDetailsViewBase::UpdatePropertyMaps()
 	for(int32 RootNodeIndex = 0; RootNodeIndex < RootPropertyNodes.Num(); ++RootNodeIndex)
 	{
 		FDetailLayoutData& LayoutData = DetailLayouts[RootNodeIndex];
-		UpdateSinglePropertyMap(RootPropertyNodes[RootNodeIndex], LayoutData);
+		UpdateSinglePropertyMap(RootPropertyNodes[RootNodeIndex], LayoutData, false);
 	}
 }
 
-void SDetailsViewBase::UpdateSinglePropertyMap(TSharedPtr<FComplexPropertyNode> InRootPropertyNode, FDetailLayoutData& LayoutData)
+void SDetailsViewBase::UpdateSinglePropertyMap(TSharedPtr<FComplexPropertyNode> InRootPropertyNode, FDetailLayoutData& LayoutData, bool bIsExternal)
 {
 	// Reset everything
 	LayoutData.ClassToPropertyMap.Empty();
 
-	TSharedPtr<FDetailLayoutBuilderImpl> DetailLayout = MakeShareable(new FDetailLayoutBuilderImpl(InRootPropertyNode, LayoutData.ClassToPropertyMap, PropertyUtilities.ToSharedRef(), SharedThis(this)));
+	TSharedPtr<FDetailLayoutBuilderImpl> DetailLayout = MakeShareable(new FDetailLayoutBuilderImpl(InRootPropertyNode, LayoutData.ClassToPropertyMap, PropertyUtilities.ToSharedRef(), SharedThis(this), bIsExternal));
 	LayoutData.DetailLayout = DetailLayout;
 
 	TSharedPtr<FComplexPropertyNode> RootPropertyNode = InRootPropertyNode;
diff --git a/Engine/Source/Editor/PropertyEditor/Private/SDetailsViewBase.h b/Engine/Source/Editor/PropertyEditor/Private/SDetailsViewBase.h
--- a/Engine/Source/Editor/PropertyEditor/Private/SDetailsViewBase.h
+++ b/Engine/Source/Editor/PropertyEditor/Private/SDetailsViewBase.h
@@ -157,7 +157,7 @@ public:
 	TSharedPtr<class FAssetThumbnailPool> GetThumbnailPool() const override;
 	TSharedPtr<IPropertyUtilities> GetPropertyUtilities() override;
 	void CreateColorPickerWindow(const TSharedRef< class FPropertyEditor >& PropertyEditor, bool bUseAlpha) override;
-	virtual void UpdateSinglePropertyMap(TSharedPtr<FComplexPropertyNode> InRootPropertyNode, FDetailLayoutData& LayoutData) override;
+	virtual void UpdateSinglePropertyMap(TSharedPtr<FComplexPropertyNode> InRootPropertyNode, FDetailLayoutData& LayoutData, bool bIsExternal) override;
 	virtual FNotifyHook* GetNotifyHook() const override { return DetailsViewArgs.NotifyHook; }
 	virtual const FCustomPropertyTypeLayoutMap& GetCustomPropertyTypeLayoutMap() const { return InstancedTypeToLayoutMap; }
 	void SaveExpandedItems( TSharedRef<FPropertyNode> StartNode ) override;
diff --git a/Engine/Source/Editor/PropertyEditor/Public/DetailCategoryBuilder.h b/Engine/Source/Editor/PropertyEditor/Public/DetailCategoryBuilder.h
--- a/Engine/Source/Editor/PropertyEditor/Public/DetailCategoryBuilder.h
+++ b/Engine/Source/Editor/PropertyEditor/Public/DetailCategoryBuilder.h
@@ -162,6 +162,12 @@ public:
 	 */
 	virtual class IDetailLayoutBuilder& GetParentLayout() const = 0;
 
+	/**
+	 * @return true if the parent layout is valid or has been destroyed by a refresh.
+	 */
+	virtual bool IsParentLayoutValid() const = 0;
+
+
 	/**
 	 * @return The localized display name of the category
 	 */
diff --git a/Engine/Source/Editor/PropertyEditor/Public/IDetailChildrenBuilder.h b/Engine/Source/Editor/PropertyEditor/Public/IDetailChildrenBuilder.h
--- a/Engine/Source/Editor/PropertyEditor/Public/IDetailChildrenBuilder.h
+++ b/Engine/Source/Editor/PropertyEditor/Public/IDetailChildrenBuilder.h
@@ -67,9 +67,11 @@ public:
 	 * @param  Objects			The objects to add
 	 * @param  PropertyName		Name of a property inside the object(s) to add.
 	 * @param  UniqueIdName		Optional identifier that uniquely identifies this object among other objects of the same type.  If this is empty, saving and restoring expansion state of this object may not work
+	 * @param  bAllowChildrenOverride Allows customization of how the new root property node is expanded when this is added. 
+	 * @param  bCreateCategoryNodesOverride Allows customization of how the new root node's category is displayed (or not). 
 	 * @return The header row generated for this set of objects by the details panel
 	 */
-	virtual class IDetailPropertyRow* AddExternalObjectProperty(const TArray<UObject*>& Objects, FName PropertyName, FName UniqueIdName = NAME_None) = 0;
+	virtual class IDetailPropertyRow* AddExternalObjectProperty(const TArray<UObject*>& Objects, FName PropertyName, FName UniqueIdName = NAME_None, TOptional<bool> bAllowChildrenOverride = TOptional<bool>(), TOptional<bool> bCreateCategoryNodesOverride = TOptional<bool>()) = 0;
 
 	/**
 	 * Adds a custom structure as a child
diff --git a/Engine/Source/Editor/UnrealEd/Classes/Animation/DebugSkelMeshComponent.h b/Engine/Source/Editor/UnrealEd/Classes/Animation/DebugSkelMeshComponent.h
--- a/Engine/Source/Editor/UnrealEd/Classes/Animation/DebugSkelMeshComponent.h
+++ b/Engine/Source/Editor/UnrealEd/Classes/Animation/DebugSkelMeshComponent.h
@@ -100,6 +100,8 @@ public:
 class FDebugSkelMeshSceneProxy : public FSkeletalMeshSceneProxy
 {
 public:
+	SIZE_T GetTypeHash() const override;
+
 	/**
 	* Constructor.
 	* @param	Component - skeletal mesh primitive being added
diff --git a/Engine/Source/Editor/UnrealEd/Private/Animation/DebugSkelMeshComponent.cpp b/Engine/Source/Editor/UnrealEd/Private/Animation/DebugSkelMeshComponent.cpp
--- a/Engine/Source/Editor/UnrealEd/Private/Animation/DebugSkelMeshComponent.cpp
+++ b/Engine/Source/Editor/UnrealEd/Private/Animation/DebugSkelMeshComponent.cpp
@@ -954,6 +954,12 @@ FDebugSkelMeshSceneProxy::FDebugSkelMeshSceneProxy(const UDebugSkelMeshComponent
 	}
 }
 
+SIZE_T FDebugSkelMeshSceneProxy::GetTypeHash() const
+{
+	static size_t UniquePointer;
+	return reinterpret_cast<size_t>(&UniquePointer);
+}
+
 void FDebugSkelMeshSceneProxy::GetDynamicMeshElements(const TArray<const FSceneView*>& Views, const FSceneViewFamily& ViewFamily, uint32 VisibilityMap, FMeshElementCollector& Collector) const
 {
 	if(!DynamicData || DynamicData->bDrawMesh)
diff --git a/Engine/Source/Editor/UnrealEd/Private/AssetThumbnail.cpp b/Engine/Source/Editor/UnrealEd/Private/AssetThumbnail.cpp
--- a/Engine/Source/Editor/UnrealEd/Private/AssetThumbnail.cpp
+++ b/Engine/Source/Editor/UnrealEd/Private/AssetThumbnail.cpp
@@ -30,7 +30,7 @@
 #include "ObjectTools.h"
 
 #include "IAssetTools.h"
-#include "IAssetTypeActions.h"
+#include "AssetTypeActions_Base.h"
 #include "AssetToolsModule.h"
 #include "Styling/SlateIconFinder.h"
 #include "ClassIconFinder.h"
@@ -562,16 +562,29 @@ private:
 		return FText::GetEmpty();
 	}
 
-	FText GetDisplayNameForClass( UClass* Class ) const
+	FText GetDisplayNameForClass( UClass* Class, const FAssetData* AssetData = nullptr) const
 	{
 		FText ClassDisplayName;
 		if ( Class )
 		{
 			FAssetToolsModule& AssetToolsModule = FModuleManager::LoadModuleChecked<FAssetToolsModule>("AssetTools");
 			TWeakPtr<IAssetTypeActions> AssetTypeActions = AssetToolsModule.Get().GetAssetTypeActionsForClass(Class);
+			
 			if ( AssetTypeActions.IsValid() )
 			{
-				ClassDisplayName = AssetTypeActions.Pin()->GetName();
+				if (AssetData != nullptr)
+				{
+					FAssetTypeActions_Base* BaseAssetTypeAction = static_cast<FAssetTypeActions_Base*>(AssetTypeActions.Pin().Get());
+					if (BaseAssetTypeAction != nullptr)
+					{
+						ClassDisplayName = BaseAssetTypeAction->GetDisplayNameFromAssetData(*AssetData);
+					}
+				}
+
+				if (ClassDisplayName.IsEmpty())
+				{
+					ClassDisplayName = AssetTypeActions.Pin()->GetName();
+				}
 			}
 
 			if ( ClassDisplayName.IsEmpty() )
@@ -591,7 +604,7 @@ private:
 
 		if ( Class )
 		{
-			return GetDisplayNameForClass( Class );
+			return GetDisplayNameForClass( Class, &AssetData );
 		}
 
 		return FText::FromString(AssetClass);
diff --git a/Engine/Source/Editor/UnrealEd/Private/CookOnTheFlyServer.cpp b/Engine/Source/Editor/UnrealEd/Private/CookOnTheFlyServer.cpp
--- a/Engine/Source/Editor/UnrealEd/Private/CookOnTheFlyServer.cpp
+++ b/Engine/Source/Editor/UnrealEd/Private/CookOnTheFlyServer.cpp
@@ -2223,7 +2223,7 @@ bool UCookOnTheFlyServer::SaveCookedPackages(TArray<UPackage*>& PackagesToSave,
 				for (int32 RemainingIndex = I; RemainingIndex < NumPackagesToRequeue; ++RemainingIndex)
 				{
 					FName StandardFilename = GetCachedStandardPackageFileFName(PackagesToSave[RemainingIndex]);
-					CookRequests.EnqueueUnique(FFilePlatformRequest(StandardFilename, AllTargetPlatformNames));
+							CookRequests.EnqueueUnique(FFilePlatformRequest(StandardFilename, SaveTargetPlatformNames));
 				}
 				Result |= COSR_WaitingOnCache;
 
@@ -4615,7 +4615,7 @@ void UCookOnTheFlyServer::CleanSandbox(const bool bIterative)
 
 				FString SandboxDirectory = GetSandboxDirectory(Target->PlatformName()); // GetOutputDirectory(Target->PlatformName());
 				IFileManager::Get().DeleteDirectory(*SandboxDirectory, false, true);
-				
+
 				ClearPlatformCookedData(FName(*Target->PlatformName()));
 
 				IniSettingsOutOfDate(Target);
@@ -4980,33 +4980,35 @@ void UCookOnTheFlyServer::CollectFilesToCook(TArray<FName>& FilesInPath, const T
 		}
 	}
 
-	if (!(FilesToCookFlags & ECookByTheBookOptions::DisableUnsolicitedPackages))
+
+
+	const FString ExternalMountPointName(TEXT("/Game/"));
+	if (IsCookingDLC())
 	{
-		const FString ExternalMountPointName(TEXT("/Game/"));
+		// get the dlc and make sure we cook that directory 
+		FString DLCPath = FPaths::Combine(*GetBaseDirectoryForDLC(), TEXT("Content"));
 
-		if (IsCookingDLC())
+		TArray<FString> Files;
+		IFileManager::Get().FindFilesRecursive(Files, *DLCPath, *(FString(TEXT("*")) + FPackageName::GetAssetPackageExtension()), true, false, false);
+		IFileManager::Get().FindFilesRecursive(Files, *DLCPath, *(FString(TEXT("*")) + FPackageName::GetMapPackageExtension()), true, false, false);
+		for (int32 Index = 0; Index < Files.Num(); Index++)
 		{
-			// get the dlc and make sure we cook that directory 
-			FString DLCPath = FPaths::Combine(*GetBaseDirectoryForDLC(), TEXT("Content"));
+			FString StdFile = Files[Index];
+			FPaths::MakeStandardFilename(StdFile);
+			AddFileToCook(FilesInPath, StdFile);
 
-			TArray<FString> Files;
-			IFileManager::Get().FindFilesRecursive(Files, *DLCPath, *(FString(TEXT("*")) + FPackageName::GetAssetPackageExtension()), true, false, false);
-			IFileManager::Get().FindFilesRecursive(Files, *DLCPath, *(FString(TEXT("*")) + FPackageName::GetMapPackageExtension()), true, false, false);
-			for (int32 Index = 0; Index < Files.Num(); Index++)
+			// this asset may not be in our currently mounted content directories, so try to mount a new one now
+			FString LongPackageName;
+			if (!FPackageName::IsValidLongPackageName(StdFile) && !FPackageName::TryConvertFilenameToLongPackageName(StdFile, LongPackageName))
 			{
-				FString StdFile = Files[Index];
-				FPaths::MakeStandardFilename(StdFile);
-				AddFileToCook(FilesInPath, StdFile);
-
-				// this asset may not be in our currently mounted content directories, so try to mount a new one now
-				FString LongPackageName;
-				if (!FPackageName::IsValidLongPackageName(StdFile) && !FPackageName::TryConvertFilenameToLongPackageName(StdFile, LongPackageName))
-				{
-					FPackageName::RegisterMountPoint(ExternalMountPointName, DLCPath);
-				}
+				FPackageName::RegisterMountPoint(ExternalMountPointName, DLCPath);
 			}
 		}
+	}
 
+
+	if (!(FilesToCookFlags & ECookByTheBookOptions::DisableUnsolicitedPackages))
+	{
 		for (const FString& CurrEntry : CookDirectories)
 		{
 			TArray<FString> Files;
diff --git a/Engine/Source/Editor/UnrealEd/Private/EditorEngine.cpp b/Engine/Source/Editor/UnrealEd/Private/EditorEngine.cpp
--- a/Engine/Source/Editor/UnrealEd/Private/EditorEngine.cpp
+++ b/Engine/Source/Editor/UnrealEd/Private/EditorEngine.cpp
@@ -3922,7 +3922,8 @@ void UEditorEngine::BuildReflectionCaptures(UWorld* World)
 				UE_LOG(LogEditor, Warning, TEXT("Unable to build Reflection Capture %s, max number of reflection captures exceeded"), *CaptureComponent->GetPathName());
 			}
 		}
-
+				// Queue an update
+		// Update sky light first because it's considered direct lighting, sky diffuse will be visible in reflection capture indirect specular
 		if (LightingScenario)
 		{
 			// Hide current scenario now that we are done capturing it
@@ -3930,13 +3931,20 @@ void UEditorEngine::BuildReflectionCaptures(UWorld* World)
 		}
 	}
 
+		// Passing in flag to verify all recaptures, no uploads
 	// Restore initial visibility
 	for (int32 LevelIndex = 0; LevelIndex < LightingScenarios.Num(); LevelIndex++)
 	{
 		ULevel* LightingScenario = LightingScenarios[LevelIndex];
 
+			// Recreate capture render state now that we have valid BuildData
+
+
 		if (LightingScenario)
 		{
+			// Hide current scenario now that we are done capturing it
+
+	// Restore initial visibility
 			EditorLevelUtils::SetLevelVisibilityTemporarily(LightingScenario, true);
 		}
 	}
diff --git a/Engine/Source/Editor/UnrealEd/Private/EditorLevelUtils.cpp b/Engine/Source/Editor/UnrealEd/Private/EditorLevelUtils.cpp
--- a/Engine/Source/Editor/UnrealEd/Private/EditorLevelUtils.cpp
+++ b/Engine/Source/Editor/UnrealEd/Private/EditorLevelUtils.cpp
@@ -826,7 +826,7 @@ void UEditorLevelUtils::SetLevelVisibilityTemporarily(ULevel* Level, bool bShoul
 	}
 }
 
-void UEditorLevelUtils::SetLevelVisibility(ULevel* Level, bool bShouldBeVisible, bool bForceLayersVisible)
+void UEditorLevelUtils::SetLevelVisibility(ULevel* Level, bool bShouldBeVisible, bool bForceLayersVisible, ELevelVisibilityDirtyMode ModifyMode)
 {
 	// Nothing to do
 	if (Level == NULL)
@@ -841,7 +841,7 @@ void UEditorLevelUtils::SetLevelVisibility(ULevel* Level, bool bShouldBeVisible,
 	{
 		//create a transaction so we can undo the visibilty toggle
 		const FScopedTransaction Transaction(LOCTEXT("ToggleLevelVisibility", "Toggle Level Visibility"));
-		if (Level->bIsVisible != bShouldBeVisible)
+		if (Level->bIsVisible != bShouldBeVisible && ModifyMode == ELevelVisibilityDirtyMode::ModifyOnChange)
 		{
 			Level->Modify();
 		}
@@ -851,7 +851,11 @@ void UEditorLevelUtils::SetLevelVisibility(ULevel* Level, bool bShouldBeVisible,
 			AActor* CurActor = *PLevelActorIter;
 			if (CurActor && !FActorEditorUtils::IsABuilderBrush(CurActor) && CurActor->bHiddenEdLevel == bShouldBeVisible)
 			{
-				CurActor->Modify();
+				if (ModifyMode == ELevelVisibilityDirtyMode::ModifyOnChange)
+				{
+					CurActor->Modify();
+				}
+				
 				CurActor->bHiddenEdLevel = !bShouldBeVisible;
 				CurActor->RegisterAllComponents();
 				CurActor->MarkComponentsRenderStateDirty();
@@ -862,7 +866,10 @@ void UEditorLevelUtils::SetLevelVisibility(ULevel* Level, bool bShouldBeVisible,
 		UModel* CurLevelModel = Level->Model;
 		if (CurLevelModel)
 		{
-			CurLevelModel->Modify();
+			if (ModifyMode == ELevelVisibilityDirtyMode::ModifyOnChange)
+			{
+				CurLevelModel->Modify();
+			}
 
 			for (TArray<FBspSurf>::TIterator SurfaceIterator(CurLevelModel->Surfs); SurfaceIterator; ++SurfaceIterator)
 			{
@@ -905,11 +912,14 @@ void UEditorLevelUtils::SetLevelVisibility(ULevel* Level, bool bShouldBeVisible,
 		// Handle the case of a streaming level
 		if (StreamingLevel)
 		{
-			// We need to set the RF_Transactional to make a streaming level serialize itself. so store the original ones, set the flag, and put the original flags back when done
-			EObjectFlags cachedFlags = StreamingLevel->GetFlags();
-			StreamingLevel->SetFlags(RF_Transactional);
-			StreamingLevel->Modify();
-			StreamingLevel->SetFlags(cachedFlags);
+			if (ModifyMode == ELevelVisibilityDirtyMode::ModifyOnChange)
+			{
+				// We need to set the RF_Transactional to make a streaming level serialize itself. so store the original ones, set the flag, and put the original flags back when done
+				EObjectFlags cachedFlags = StreamingLevel->GetFlags();
+				StreamingLevel->SetFlags(RF_Transactional);
+				StreamingLevel->Modify();
+				StreamingLevel->SetFlags(cachedFlags);
+			}
 
 			// Set the visibility state for this streaming level.  
 			StreamingLevel->bShouldBeVisibleInEditor = bShouldBeVisible;
@@ -921,7 +931,7 @@ void UEditorLevelUtils::SetLevelVisibility(ULevel* Level, bool bShouldBeVisible,
 		}
 
 		// UpdateLevelStreaming sets Level->bIsVisible directly, so we need to make sure it gets saved to the transaction buffer.
-		if (Level->bIsVisible != bShouldBeVisible)
+		if (Level->bIsVisible != bShouldBeVisible && ModifyMode == ELevelVisibilityDirtyMode::ModifyOnChange)
 		{
 			Level->Modify();
 		}
@@ -975,7 +985,11 @@ void UEditorLevelUtils::SetLevelVisibility(ULevel* Level, bool bShouldBeVisible,
 					// Make the actor layer visible, if it's not already.
 					if (Actor->bHiddenEdLayer)
 					{
-						bModified = Actor->Modify();
+						if (ModifyMode == ELevelVisibilityDirtyMode::ModifyOnChange)
+						{
+							bModified = Actor->Modify();
+						}
+						
 						Actor->bHiddenEdLayer = false;
 					}
 
@@ -986,7 +1000,7 @@ void UEditorLevelUtils::SetLevelVisibility(ULevel* Level, bool bShouldBeVisible,
 				// Set the visibility of each actor in the streaming level
 				if (!FActorEditorUtils::IsABuilderBrush(Actor) && Actor->bHiddenEdLevel == bShouldBeVisible)
 				{
-					if (!bModified)
+					if (!bModified && ModifyMode == ELevelVisibilityDirtyMode::ModifyOnChange)
 					{
 						bModified = Actor->Modify();
 					}
@@ -1011,7 +1025,7 @@ void UEditorLevelUtils::SetLevelVisibility(ULevel* Level, bool bShouldBeVisible,
 	GEngine->BroadcastLevelActorListChanged();
 
 	// If the level is being hidden, deselect actors and surfaces that belong to this level.
-	if (!bShouldBeVisible)
+	if (!bShouldBeVisible && ModifyMode == ELevelVisibilityDirtyMode::ModifyOnChange)
 	{
 		USelection* SelectedActors = GEditor->GetSelectedActors();
 		SelectedActors->Modify();
diff --git a/Engine/Source/Editor/UnrealEd/Private/EditorViewportClient.cpp b/Engine/Source/Editor/UnrealEd/Private/EditorViewportClient.cpp
--- a/Engine/Source/Editor/UnrealEd/Private/EditorViewportClient.cpp
+++ b/Engine/Source/Editor/UnrealEd/Private/EditorViewportClient.cpp
@@ -1451,12 +1451,20 @@ EMouseCursor::Type FEditorViewportClient::GetCursor(FViewport* InViewport,int32
 	
 	// Allow the viewport interaction to override any previously set mouse cursor
 	UViewportWorldInteraction* WorldInteraction = Cast<UViewportWorldInteraction>(GEditor->GetEditorWorldExtensionsManager()->GetEditorWorldExtensions(GetWorld())->FindExtension(UViewportWorldInteraction::StaticClass()));
-	if (WorldInteraction != nullptr && WorldInteraction->ShouldSuppressExistingCursor())
+	if (WorldInteraction != nullptr)
 	{
+		if (WorldInteraction->ShouldForceCursor())
+		{
+			MouseCursor = EMouseCursor::Crosshairs;
+			SetRequiredCursor(false, true);
+			UpdateRequiredCursorVisibility();
+		}
+		else if (WorldInteraction->ShouldSuppressExistingCursor())
+		{
 			MouseCursor = EMouseCursor::None;
-			RequiredCursorVisibiltyAndAppearance.bHardwareCursorVisible = false;
-			RequiredCursorVisibiltyAndAppearance.bSoftwareCursorVisible = false;
+			SetRequiredCursor(false, false);
 			UpdateRequiredCursorVisibility();
+		}
 	}
 
 	CachedMouseX = X;
@@ -2735,8 +2743,7 @@ void FEditorViewportClient::ProcessDoubleClickInViewport( const struct FInputEve
 	// This needs to be set to false to allow the axes to update
 	bWidgetAxisControlledByDrag = false;
 	MouseDeltaTracker->ResetUsedDragModifier();
-	RequiredCursorVisibiltyAndAppearance.bHardwareCursorVisible = true;
-	RequiredCursorVisibiltyAndAppearance.bSoftwareCursorVisible = false;
+	SetRequiredCursor(true, false);
 	ApplyRequiredCursorVisibility();
 }
 
@@ -4284,8 +4291,7 @@ void FEditorViewportClient::UpdateRequiredCursorVisibility()
 
 	if (GetViewportType() == LVT_None)
 	{
-		RequiredCursorVisibiltyAndAppearance.bHardwareCursorVisible = true;
-		RequiredCursorVisibiltyAndAppearance.bSoftwareCursorVisible = false;
+		SetRequiredCursor(true, false);
 		return;
 	}
 
@@ -4298,8 +4304,7 @@ void FEditorViewportClient::UpdateRequiredCursorVisibility()
 			(  GetWidgetMode() == FWidget::WM_TranslateRotateZ && Widget->GetCurrentAxis() != EAxisList::ZRotation &&  Widget->GetCurrentAxis() != EAxisList::None ) ||
 			( GetWidgetMode() == FWidget::WM_2D && Widget->GetCurrentAxis() != EAxisList::Rotate2D &&  Widget->GetCurrentAxis() != EAxisList::None ) ) )
 		{
-			RequiredCursorVisibiltyAndAppearance.bHardwareCursorVisible = false;
-			RequiredCursorVisibiltyAndAppearance.bSoftwareCursorVisible = true;								
+			SetRequiredCursor(false, true);
 			SetRequiredCursorOverride( true , EMouseCursor::CardinalCross );
 			return;
 		}
@@ -4313,13 +4318,11 @@ void FEditorViewportClient::UpdateRequiredCursorVisibility()
 			{
 				// Always turn the hardware cursor on before turning the software cursor off
 				// so the hardware cursor will be be set where the software cursor was
-				RequiredCursorVisibiltyAndAppearance.bHardwareCursorVisible = !bHasMouseMovedSinceClick;
-				RequiredCursorVisibiltyAndAppearance.bSoftwareCursorVisible = bHasMouseMovedSinceClick;
+				SetRequiredCursor(!bHasMouseMovedSinceClick, bHasMouseMovedSinceClick);
 				SetRequiredCursorOverride( true , EMouseCursor::GrabHand );
 				return;
 			}
-			RequiredCursorVisibiltyAndAppearance.bHardwareCursorVisible = false;
-			RequiredCursorVisibiltyAndAppearance.bSoftwareCursorVisible = false;
+			SetRequiredCursor(false, false);
 			return;
 		}
 	}
@@ -4328,8 +4331,7 @@ void FEditorViewportClient::UpdateRequiredCursorVisibility()
 	if (IsUsingAbsoluteTranslation() && !MouseDeltaTracker->UsingDragTool() )
 	{
 		//If we are dragging something we should hide the hardware cursor and show the s/w one
-		RequiredCursorVisibiltyAndAppearance.bHardwareCursorVisible = false;
-		RequiredCursorVisibiltyAndAppearance.bSoftwareCursorVisible = true;
+		SetRequiredCursor(false, true);
 		SetRequiredCursorOverride( true , EMouseCursor::CardinalCross );
 	}
 	else
@@ -4340,8 +4342,7 @@ void FEditorViewportClient::UpdateRequiredCursorVisibility()
 		if (bMouseButtonDown && (RawMouseDelta.SizeSquared() >= MOUSE_CLICK_DRAG_DELTA || IsFlightCameraActive() || ShouldOrbitCamera()) && !MouseDeltaTracker->UsingDragTool())
 		{
 			//current system - do not show cursor when mouse is down
-			RequiredCursorVisibiltyAndAppearance.bHardwareCursorVisible = false;
-			RequiredCursorVisibiltyAndAppearance.bSoftwareCursorVisible = false;
+			SetRequiredCursor(false, false);
 			return;
 		}
 
@@ -4350,11 +4351,16 @@ void FEditorViewportClient::UpdateRequiredCursorVisibility()
 			RequiredCursorVisibiltyAndAppearance.bOverrideAppearance = false;
 		}
 
-		RequiredCursorVisibiltyAndAppearance.bHardwareCursorVisible = true;
-		RequiredCursorVisibiltyAndAppearance.bSoftwareCursorVisible = false;
+		SetRequiredCursor(true, false);
 	}
 }
 
+void FEditorViewportClient::SetRequiredCursor(const bool bHardwareCursorVisible, const bool bSoftwareCursorVisible)
+{
+	RequiredCursorVisibiltyAndAppearance.bHardwareCursorVisible = bHardwareCursorVisible;
+	RequiredCursorVisibiltyAndAppearance.bSoftwareCursorVisible = bSoftwareCursorVisible;
+}
+
 void FEditorViewportClient::ApplyRequiredCursorVisibility( bool bUpdateSoftwareCursorPostion )
 {
 	if( RequiredCursorVisibiltyAndAppearance.bDontResetCursor == true )
diff --git a/Engine/Source/Editor/UnrealEd/Private/Factories/SkeletalMeshImport.cpp b/Engine/Source/Editor/UnrealEd/Private/Factories/SkeletalMeshImport.cpp
--- a/Engine/Source/Editor/UnrealEd/Private/Factories/SkeletalMeshImport.cpp
+++ b/Engine/Source/Editor/UnrealEd/Private/Factories/SkeletalMeshImport.cpp
@@ -561,6 +561,8 @@ ExistingSkelMeshData* SaveExistingSkelMeshData(USkeletalMesh* ExistingSkelMesh,
 
 		ExistingMeshDataPtr->ExistingClothingAssets = ExistingSkelMesh->MeshClothingAssets;
 
+		ExistingMeshDataPtr->ExistingSamplingInfo = ExistingSkelMesh->GetSamplingInfo();
+
 		//Add the last fbx import data
 		UFbxSkeletalMeshImportData* ImportData = Cast<UFbxSkeletalMeshImportData>(ExistingSkelMesh->AssetImportData);
 		if (ImportData && ExistingMeshDataPtr->UseMaterialNameSlotWorkflow)
@@ -911,6 +913,8 @@ void RestoreExistingSkelMeshData(ExistingSkelMeshData* MeshData, USkeletalMesh*
 			ClothingAsset->RefreshBoneMapping(SkeletalMesh);
 		}
 
+		SkeletalMesh->SetSamplingInfo(MeshData->ExistingSamplingInfo);
+
 		//Restore the section change only for the base LOD, other LOD will be restore when setting the LOD.
 		if (MeshData->UseMaterialNameSlotWorkflow)
 		{
diff --git a/Engine/Source/Editor/UnrealEd/Private/FbxAnimUtils.cpp b/Engine/Source/Editor/UnrealEd/Private/FbxAnimUtils.cpp
--- a/Engine/Source/Editor/UnrealEd/Private/FbxAnimUtils.cpp
+++ b/Engine/Source/Editor/UnrealEd/Private/FbxAnimUtils.cpp
@@ -32,9 +32,9 @@ namespace FbxAnimUtils
 		}
 	}
 
-	static FbxNode* FindCurveNodeRecursive(FbxNode* NodeToQuery, const FString& InCurveNodeName)
+	static FbxNode* FindCurveNodeRecursive(FbxNode* NodeToQuery, const FString& InCurveNodeName, FbxNodeAttribute::EType InNodeType)
 	{
-		if (InCurveNodeName == UTF8_TO_TCHAR(NodeToQuery->GetName()) && NodeToQuery->GetNodeAttribute() && NodeToQuery->GetNodeAttribute()->GetAttributeType() == FbxNodeAttribute::eMesh)
+		if (InCurveNodeName == UTF8_TO_TCHAR(NodeToQuery->GetName()) && NodeToQuery->GetNodeAttribute() && NodeToQuery->GetNodeAttribute()->GetAttributeType() == InNodeType)
 		{
 			return NodeToQuery;
 		}
@@ -42,7 +42,7 @@ namespace FbxAnimUtils
 		const int32 NodeCount = NodeToQuery->GetChildCount();
 		for (int32 NodeIndex = 0; NodeIndex < NodeCount; ++NodeIndex)
 		{
-			FbxNode* FoundNode = FindCurveNodeRecursive(NodeToQuery->GetChild(NodeIndex), InCurveNodeName);
+			FbxNode* FoundNode = FindCurveNodeRecursive(NodeToQuery->GetChild(NodeIndex), InCurveNodeName, InNodeType);
 			if (FoundNode != nullptr)
 			{
 				return FoundNode;
@@ -52,10 +52,54 @@ namespace FbxAnimUtils
 		return nullptr;
 	}
 
-	static FbxNode* FindCurveNode(UnFbx::FFbxImporter* FbxImporter, const FString& InCurveNodeName)
+	static FbxNode* FindCurveNode(UnFbx::FFbxImporter* FbxImporter, const FString& InCurveNodeName, FbxNodeAttribute::EType InNodeType)
 	{
 		FbxNode *RootNode = FbxImporter->Scene->GetRootNode();
-		return FindCurveNodeRecursive(RootNode, InCurveNodeName);
+		return FindCurveNodeRecursive(RootNode, InCurveNodeName, InNodeType);
+	}
+
+	bool IsSupportedCurveDataType(int32 InDataType)
+	{
+		switch (InDataType)
+		{
+		case eFbxShort:		//!< 16 bit signed integer.
+		case eFbxUShort:	//!< 16 bit unsigned integer.
+		case eFbxUInt:		//!< 32 bit unsigned integer.
+		case eFbxHalfFloat:	//!< 16 bit floating point.
+		case eFbxInt:		//!< 32 bit signed integer.
+		case eFbxFloat:		//!< Floating point value.
+		case eFbxDouble:	//!< Double width floating point value.
+		case eFbxDouble2:	//!< Vector of two double values.
+		case eFbxDouble3:	//!< Vector of three double values.
+		case eFbxDouble4:	//!< Vector of four double values.
+		case eFbxDouble4x4:	//!< Four vectors of four double values.
+			return true;
+		}
+
+		return false;
+	}
+
+	bool ShouldImportCurve(FbxAnimCurve* Curve, bool bDoNotImportWithZeroValues)
+	{
+		if (Curve && Curve->KeyGetCount() > 0)
+		{
+			if (bDoNotImportWithZeroValues)
+			{
+				for (int32 KeyIndex = 0; KeyIndex < Curve->KeyGetCount(); ++KeyIndex)
+				{
+					if (!FMath::IsNearlyZero(Curve->KeyGetValue(KeyIndex)))
+					{
+						return true;
+					}
+				}
+			}
+			else
+			{
+				return true;
+			}
+		}
+
+		return false;
 	}
 
 	bool ImportCurveTableFromNode(const FString& InFbxFilename, const FString& InCurveNodeName, UCurveTable* InOutCurveTable, float& OutPreRoll)
@@ -82,13 +126,13 @@ namespace FbxAnimUtils
 					FbxAnimLayer* AnimLayer = AnimStack->GetMember<FbxAnimLayer>(0);
 					if (AnimLayer != nullptr)
 					{
-						FbxNode* Node = FindCurveNode(FbxImporter, InCurveNodeName);
-						if (Node != nullptr)
+						// Try the legacy path with curves in blendshapes on a dummy mesh
+						if (FbxNode* MeshNode = FindCurveNode(FbxImporter, InCurveNodeName, FbxNodeAttribute::eMesh))
 						{
 							// We have a node, so clear the curve table
 							InOutCurveTable->RowMap.Empty();
 
-							FbxGeometry* Geometry = (FbxGeometry*)Node->GetNodeAttribute();
+							FbxGeometry* Geometry = (FbxGeometry*)MeshNode->GetNodeAttribute();
 							int32 BlendShapeDeformerCount = Geometry->GetDeformerCount(FbxDeformer::eBlendShape);
 							for (int32 BlendShapeIndex = 0; BlendShapeIndex < BlendShapeDeformerCount; ++BlendShapeIndex)
 							{
@@ -123,7 +167,31 @@ namespace FbxAnimUtils
 									}
 								}
 							}
+							return true;
+						}
+						// Now look for attribute curves on a skeleton node
+						else if (FbxNode* SkeletonNode = FindCurveNode(FbxImporter, InCurveNodeName, FbxNodeAttribute::eSkeleton))
+						{
+							// We have a node, so clear the curve table
+							InOutCurveTable->RowMap.Empty();
 
+							ExtractAttributeCurves(SkeletonNode, false, [&InOutCurveTable, &AnimTimeSpan, &FbxImporter](FbxAnimCurve* InCurve, const FString& InCurveName, const FString& InChannelName, int32 InChannelIndex, int32 InNumChannels)
+							{
+								FString FinalCurveName;
+								if (InNumChannels == 1)
+								{
+									FinalCurveName = InCurveName;
+								}
+								else
+								{
+									FinalCurveName = InCurveName + "_" + InChannelName;
+								}
+
+								FRichCurve& RichCurve = *InOutCurveTable->RowMap.Add(*FinalCurveName, new FRichCurve());
+								RichCurve.Reset();
+
+								FbxImporter->ImportCurve(InCurve, RichCurve, AnimTimeSpan, 1.0f);
+							});
 							return true;
 						}
 					}
@@ -135,4 +203,38 @@ namespace FbxAnimUtils
 
 		return false;
 	}
+
+	void ExtractAttributeCurves(FbxNode* InNode, bool bInDoNotImportCurveWithZero, TFunctionRef<void(FbxAnimCurve* /*InCurve*/, const FString& /*InCurveName*/, const FString& /*InChannelName*/, int32 /*InChannelIndex*/, int32 /*InNumChannels*/)> ImportFunction)
+	{
+		FbxProperty Property = InNode->GetFirstProperty();
+		while (Property.IsValid())
+		{
+			FbxAnimCurveNode* CurveNode = Property.GetCurveNode();
+			// do this if user defined and animated and leaf node
+			if( CurveNode && Property.GetFlag(FbxPropertyFlags::eUserDefined) &&
+				CurveNode->IsAnimated() && FbxAnimUtils::IsSupportedCurveDataType(Property.GetPropertyDataType().GetType()) )
+			{
+				FString CurveName = UTF8_TO_TCHAR(CurveNode->GetName());
+				UE_LOG(LogFbx, Log, TEXT("CurveName : %s"), *CurveName );
+
+				int32 TotalCount = CurveNode->GetChannelsCount();
+				for (int32 ChannelIndex=0; ChannelIndex<TotalCount; ++ChannelIndex)
+				{
+					FbxAnimCurve * AnimCurve = CurveNode->GetCurve(ChannelIndex);
+					FString ChannelName = CurveNode->GetChannelName(ChannelIndex).Buffer();
+
+					if (FbxAnimUtils::ShouldImportCurve(AnimCurve, bInDoNotImportCurveWithZero))
+					{
+						ImportFunction(AnimCurve, CurveName, ChannelName, ChannelIndex, TotalCount);
+					}
+					else
+					{
+						UE_LOG(LogFbx, Log, TEXT("CurveName(%s) is skipped because it only contains invalid values."), *CurveName);
+					}
+				}
+			}
+
+			Property = InNode->GetNextProperty(Property); 
+		}
+	}
 }
diff --git a/Engine/Source/Editor/UnrealEd/Private/FeedbackContextEditor.cpp b/Engine/Source/Editor/UnrealEd/Private/FeedbackContextEditor.cpp
--- a/Engine/Source/Editor/UnrealEd/Private/FeedbackContextEditor.cpp
+++ b/Engine/Source/Editor/UnrealEd/Private/FeedbackContextEditor.cpp
@@ -403,7 +403,7 @@ void FFeedbackContextEditor::StartSlowTask( const FText& Task, bool bShowCancelB
 
 		if( bHaveOpenMenu )
 		{
-			UE_LOG(LogSlate, Warning, TEXT("Prevented a slow task dialog from being summoned while a context menu was open") );
+			UE_LOG(LogSlate, Log, TEXT("Prevented a slow task dialog from being summoned while a context menu was open") );
 		}
 
 		// reset the cancellation flag
diff --git a/Engine/Source/Editor/UnrealEd/Private/HierarchicalLOD.cpp b/Engine/Source/Editor/UnrealEd/Private/HierarchicalLOD.cpp
--- a/Engine/Source/Editor/UnrealEd/Private/HierarchicalLOD.cpp
+++ b/Engine/Source/Editor/UnrealEd/Private/HierarchicalLOD.cpp
@@ -406,7 +406,7 @@ bool FHierarchicalLODBuilder::ShouldGenerateCluster(AActor* Actor, const bool bP
 	{
 		UStaticMeshComponent* SMComponent = Cast<UStaticMeshComponent>(Component);
 		// TODO: support instanced static meshes
-		if (Component && !Component->IsA<UInstancedStaticMeshComponent>())
+		if (SMComponent && !Component->IsA<UInstancedStaticMeshComponent>())
 		{
 			Components.Add(SMComponent);
 		}
diff --git a/Engine/Source/Editor/UnrealEd/Private/SCommonEditorViewportToolbarBase.cpp b/Engine/Source/Editor/UnrealEd/Private/SCommonEditorViewportToolbarBase.cpp
--- a/Engine/Source/Editor/UnrealEd/Private/SCommonEditorViewportToolbarBase.cpp
+++ b/Engine/Source/Editor/UnrealEd/Private/SCommonEditorViewportToolbarBase.cpp
@@ -248,6 +248,7 @@ TSharedRef<SWidget> SCommonEditorViewportToolbarBase::GenerateOptionsMenu() cons
 			}
 		}
 		OptionsMenuBuilder.EndSection();
+		ExtendOptionsMenu(OptionsMenuBuilder);
 	}
 
 	return OptionsMenuBuilder.MakeWidget();
diff --git a/Engine/Source/Editor/UnrealEd/Private/SCurveEditor.cpp b/Engine/Source/Editor/UnrealEd/Private/SCurveEditor.cpp
--- a/Engine/Source/Editor/UnrealEd/Private/SCurveEditor.cpp
+++ b/Engine/Source/Editor/UnrealEd/Private/SCurveEditor.cpp
@@ -493,6 +493,8 @@ void SCurveEditor::Construct(const FArguments& InArgs)
 	}
 
 	FCoreUObjectDelegates::OnObjectPropertyChanged.AddSP(this, &SCurveEditor::OnObjectPropertyChanged);
+
+	bRequireFocusToZoom = false;
 }
 
 FText SCurveEditor::GetIsCurveVisibleToolTip(TSharedPtr<FCurveViewModel> CurveViewModel) const
@@ -1305,6 +1307,21 @@ void SCurveEditor::SetZoomToFit(bool bNewZoomToFitVertical, bool bNewZoomToFitHo
 	bZoomToFitHorizontal = bNewZoomToFitHorizontal;
 }
 
+void SCurveEditor::SetRequireFocusToZoom(bool bInRequireFocusToZoom)
+{
+	bRequireFocusToZoom = bInRequireFocusToZoom;
+}
+
+TOptional<bool> SCurveEditor::OnQueryShowFocus(const EFocusCause InFocusCause) const
+{
+	if (bRequireFocusToZoom)
+	{
+		// Enable showing a focus rectangle when the widget has keyboard focus
+		return TOptional<bool>(true);
+	}
+	return TOptional<bool>();
+}
+
 FCurveOwnerInterface* SCurveEditor::GetCurveOwner() const
 {
 	return CurveOwner;
@@ -1557,8 +1574,12 @@ void SCurveEditor::UpdateCurveToolTip(const FGeometry& InMyGeometry, const FPoin
 
 FReply SCurveEditor::OnMouseWheel(const FGeometry& MyGeometry, const FPointerEvent& MouseEvent)
 {
-	ZoomView(FVector2D(MouseEvent.GetWheelDelta(), MouseEvent.GetWheelDelta()));
-	return FReply::Handled();
+	if (!bRequireFocusToZoom || HasKeyboardFocus())
+	{
+		ZoomView(FVector2D(MouseEvent.GetWheelDelta(), MouseEvent.GetWheelDelta()));
+		return FReply::Handled();
+	}
+	return FReply::Unhandled();
 }
 
 void SCurveEditor::ZoomView(FVector2D Delta)
diff --git a/Engine/Source/Editor/UnrealEd/Private/SkeletalMeshEdit.cpp b/Engine/Source/Editor/UnrealEd/Private/SkeletalMeshEdit.cpp
--- a/Engine/Source/Editor/UnrealEd/Private/SkeletalMeshEdit.cpp
+++ b/Engine/Source/Editor/UnrealEd/Private/SkeletalMeshEdit.cpp
@@ -30,6 +30,7 @@
 #include "UObjectIterator.h"
 #include "ComponentReregisterContext.h"
 #include "Components/SkeletalMeshComponent.h"
+#include "FbxAnimUtils.h"
 
 #define LOCTEXT_NAMESPACE "SkeletalMeshEdit"
 
@@ -1106,31 +1107,6 @@ namespace AnimationTransformDebug
 	}
 }
 
-/**
- * We only support float values, so these are the numbers we can take
- */
-bool IsSupportedCurveDataType(EFbxType DatatType)
-{
-
-	switch (DatatType)
-	{
-	case eFbxShort:		//!< 16 bit signed integer.
-	case eFbxUShort:		//!< 16 bit unsigned integer.
-	case eFbxUInt:		//!< 32 bit unsigned integer.
-	case eFbxHalfFloat:	//!< 16 bit floating point.
-	case eFbxInt:		//!< 32 bit signed integer.
-	case eFbxFloat:		//!< Floating point value.
-	case eFbxDouble:		//!< Double width floating point value.
-	case eFbxDouble2:	//!< Vector of two double values.
-	case eFbxDouble3:	//!< Vector of three double values.
-	case eFbxDouble4:	//!< Vector of four double values.
-	case eFbxDouble4x4:	//!< Four vectors of four double values.
-		return true;
-	}
-
-	return false;
-}
-
 bool UnFbx::FFbxImporter::ImportCurveToAnimSequence(class UAnimSequence * TargetSequence, const FString& CurveName, const FbxAnimCurve* FbxCurve, int32 CurveFlags,const FbxTimeSpan AnimTimeSpan, const float ValueScale/*=1.f*/) const
 {
 	if (TargetSequence && FbxCurve)
@@ -1188,29 +1164,6 @@ bool UnFbx::FFbxImporter::ImportCurveToAnimSequence(class UAnimSequence * Target
 	return false;
 }
 
-bool ShouldImportCurve(FbxAnimCurve* Curve, bool bDoNotImportWithZeroValues)
-{
-	if (Curve && Curve->KeyGetCount() > 0)
-	{
-		if (bDoNotImportWithZeroValues)
-		{
-			for (int32 KeyIndex = 0; KeyIndex < Curve->KeyGetCount(); ++KeyIndex)
-			{
-				if (!FMath::IsNearlyZero(Curve->KeyGetValue(KeyIndex)))
-				{
-					return true;
-				}
-			}
-		}
-		else
-		{
-			return true;
-		}
-	}
-
-	return false;
-}
-
 bool UnFbx::FFbxImporter::ImportAnimation(USkeleton* Skeleton, UAnimSequence * DestSeq, const FString& FileName, TArray<FbxNode*>& SortedLinks, TArray<FbxNode*>& NodeArray, FbxAnimStack* CurAnimStack, const int32 ResampleRate, const FbxTimeSpan AnimTimeSpan)
 {
 	//This destroy all previously imported animation raw data
@@ -1323,7 +1276,7 @@ bool UnFbx::FFbxImporter::ImportAnimation(USkeleton* Skeleton, UAnimSequence * D
 							}
 
 							FbxAnimCurve* Curve = Geometry->GetShapeChannel(BlendShapeIndex, ChannelIndex, (FbxAnimLayer*)CurAnimStack->GetMember(0));
-							if (ShouldImportCurve(Curve, ImportOptions->bDoNotImportCurveWithZero))
+							if (FbxAnimUtils::ShouldImportCurve(Curve, ImportOptions->bDoNotImportCurveWithZero))
 							{
 								FFormatNamedArguments Args;
 								Args.Add(TEXT("BlendShape"), FText::FromString(ChannelName));
@@ -1360,75 +1313,52 @@ bool UnFbx::FFbxImporter::ImportAnimation(USkeleton* Skeleton, UAnimSequence * D
 		int32 CurLinkIndex=0;
 		for(auto Node: SortedLinks)
 		{
-			FbxProperty Property = Node->GetFirstProperty();
-			while (Property.IsValid())
+			FbxAnimUtils::ExtractAttributeCurves(Node, ImportOptions->bDoNotImportCurveWithZero, 
+			[this, &CurLinkIndex, &TotalLinks, &DestSeq, &AnimTimeSpan, &ExistingCurveNames](FbxAnimCurve* InCurve, const FString& InCurveName, const FString& InChannelName, int32 InChannelIndex, int32 InNumChannels)
 			{
-				FbxAnimCurveNode* CurveNode = Property.GetCurveNode();
-				// do this if user defined and animated and leaf node
-				if( CurveNode && Property.GetFlag(FbxPropertyFlags::eUserDefined) &&
-					CurveNode->IsAnimated() && IsSupportedCurveDataType(Property.GetPropertyDataType().GetType()) )
+				FString FinalCurveName;
+				if (InNumChannels == 1)
 				{
-					FString CurveName = UTF8_TO_TCHAR(CurveNode->GetName());
-					UE_LOG(LogFbx, Log, TEXT("CurveName : %s"), *CurveName );
-
-					int32 TotalCount = CurveNode->GetChannelsCount();
-					for (int32 ChannelIndex=0; ChannelIndex<TotalCount; ++ChannelIndex)
-					{
-						FbxAnimCurve * AnimCurve = CurveNode->GetCurve(ChannelIndex);
-						FString ChannelName = CurveNode->GetChannelName(ChannelIndex).Buffer();
+					FinalCurveName = InCurveName;
+				}
+				else
+				{
+					FinalCurveName = InCurveName + "_" + InChannelName;
+				}
 
-						if (ShouldImportCurve(AnimCurve, ImportOptions->bDoNotImportCurveWithZero))
-						{
-							FString FinalCurveName;
-							if (TotalCount == 1)
-							{
-								FinalCurveName = CurveName;
-							}
-							else
-							{
-								FinalCurveName = CurveName + "_" + ChannelName;
-							}
+				FFormatNamedArguments Args;
+				Args.Add(TEXT("CurveName"), FText::FromString(InChannelName));
+				const FText StatusUpate = FText::Format(LOCTEXT("ImportingCustomAttributeCurvesDetail", "Importing Custom Attribute [{CurveName}]"), Args);
+				GWarn->StatusUpdate(CurLinkIndex + 1, TotalLinks, StatusUpate);
 
-							FFormatNamedArguments Args;
-							Args.Add(TEXT("CurveName"), FText::FromString(FinalCurveName));
-							const FText StatusUpate = FText::Format(LOCTEXT("ImportingCustomAttributeCurvesDetail", "Importing Custom Attribute [{CurveName}]"), Args);
-							GWarn->StatusUpdate(CurLinkIndex + 1, TotalLinks, StatusUpate);
+				int32 CurveFlags = AACF_DefaultCurve;
+				if (ImportCurveToAnimSequence(DestSeq, InChannelName, InCurve, CurveFlags, AnimTimeSpan))
+				{
+					USkeleton* SeqSkeleton = DestSeq->GetSkeleton();
 
-							int32 CurveFlags = AACF_DefaultCurve;
-							if (ImportCurveToAnimSequence(DestSeq, FinalCurveName, AnimCurve, CurveFlags, AnimTimeSpan))
+					// first let them override material curve if required
+					if (ImportOptions->bSetMaterialDriveParameterOnCustomAttribute)
+					{
+						// now mark this curve as morphtarget
+						SeqSkeleton->AccumulateCurveMetaData(FName(*InChannelName), true, false);
+					}
+					else
+					{
+						// if not material set by default, apply naming convention for material
+						for (const auto& Suffix : ImportOptions->MaterialCurveSuffixes)
+						{
+							int32 TotalSuffix = Suffix.Len();
+							if (InChannelName.Right(TotalSuffix) == Suffix)
 							{
-								// first let them override material curve if required
-								if (ImportOptions->bSetMaterialDriveParameterOnCustomAttribute)
-								{
-									// now mark this curve as morphtarget
-									MySkeleton->AccumulateCurveMetaData(FName(*FinalCurveName), true, false);
-								}
-								else
-								{
-									// if not material set by default, apply naming convention for material
-									for (const auto& Suffix : ImportOptions->MaterialCurveSuffixes)
-									{
-										int32 TotalSuffix = Suffix.Len();
-										if (CurveName.Right(TotalSuffix) == Suffix)
-										{
-											MySkeleton->AccumulateCurveMetaData(FName(*FinalCurveName), true, false);
-											break;
-										}
-									}
-								}
-
-								ExistingCurveNames.Remove(FinalCurveName);
+								SeqSkeleton->AccumulateCurveMetaData(FName(*InChannelName), true, false);
+								break;
 							}
 						}
-						else
-						{
-							UE_LOG(LogFbx, Log, TEXT("CurveName(%s) is skipped because it only contains invalid values."), *CurveName);
-						}
 					}
-				}
 
-				Property = Node->GetNextProperty(Property); 
-			}
+					ExistingCurveNames.Remove(InChannelName);
+				}	
+			});
 
 			CurLinkIndex++;
 		}
diff --git a/Engine/Source/Editor/UnrealEd/Public/EditorLevelUtils.h b/Engine/Source/Editor/UnrealEd/Public/EditorLevelUtils.h
--- a/Engine/Source/Editor/UnrealEd/Public/EditorLevelUtils.h
+++ b/Engine/Source/Editor/UnrealEd/Public/EditorLevelUtils.h
@@ -17,6 +17,14 @@ class ULevelStreaming;
 
 DECLARE_LOG_CATEGORY_EXTERN(LogLevelTools, Warning, All);
 
+enum class ELevelVisibilityDirtyMode
+{
+	// Use when the user is causing the visibility change.  Will update transaction state and mark the package dirty.
+	ModifyOnChange,
+	// Use when code is causing the visibility change.
+	DontModify
+};
+
 UCLASS(transient)
 class UEditorLevelUtils : public UObject
 {
@@ -144,7 +152,7 @@ public:
 	 * @param	bShouldBeVisible		The level's new visibility state.
 	 * @param	bForceLayersVisible		If true and the level is visible, force the level's layers to be visible.
 	 */
-	static UNREALED_API void SetLevelVisibility(ULevel* Level, bool bShouldBeVisible, bool bForceLayersVisible);
+	static UNREALED_API void SetLevelVisibility(ULevel* Level, bool bShouldBeVisible, bool bForceLayersVisible, ELevelVisibilityDirtyMode ModifyMode = ELevelVisibilityDirtyMode::ModifyOnChange);
 	
 	/** 
 	 * Deselects all BSP surfaces in this level 
diff --git a/Engine/Source/Editor/UnrealEd/Public/EditorViewportClient.h b/Engine/Source/Editor/UnrealEd/Public/EditorViewportClient.h
--- a/Engine/Source/Editor/UnrealEd/Public/EditorViewportClient.h
+++ b/Engine/Source/Editor/UnrealEd/Public/EditorViewportClient.h
@@ -1113,6 +1113,9 @@ protected:
 	/** Setup the cursor visibility state we require and store in RequiredCursorVisibiltyAndAppearance struct */
 	void UpdateRequiredCursorVisibility();
 	
+	/** Sets the required hardware and software cursor. */
+	void SetRequiredCursor(const bool bHardwareCursorVisible, const bool bSoftwareCursorVisible);
+
 	/** 
 	 * Apply the required cursor visibility states from the RequiredCursorVisibiltyAndAppearance struct 
 	 * @param	View				True - Set the position of the software cursor if being made visible. This defaults to FALSE.
diff --git a/Engine/Source/Editor/UnrealEd/Public/FbxAnimUtils.h b/Engine/Source/Editor/UnrealEd/Public/FbxAnimUtils.h
--- a/Engine/Source/Editor/UnrealEd/Public/FbxAnimUtils.h
+++ b/Engine/Source/Editor/UnrealEd/Public/FbxAnimUtils.h
@@ -8,6 +8,12 @@ class UAnimSequence;
 class UCurveTable;
 class USkeletalMesh;
 
+namespace fbxsdk
+{
+	class FbxAnimCurve;
+	class FbxNode;
+}
+
 //Define interface for exporting fbx animation
 namespace FbxAnimUtils
 {
@@ -23,4 +29,22 @@ namespace FbxAnimUtils
 	 */
 	UNREALED_API bool ImportCurveTableFromNode(const FString& InFbxFilename, const FString& InCurveNodeName, UCurveTable* InOutCurveTable, float& OutPreRoll);
 
+	/** Helper function used to determine whether we can import a particular curve type. Note InDataType is EFbxType */
+	UNREALED_API bool IsSupportedCurveDataType(int32 InDataType);
+
+	/** 
+	 * Helper function used to check whether we should import the specified curve.
+	 * @param	Curve							The curve to check
+	 * @param	bDoNotImportWithZeroValues		Whether to skip zero-valued curves
+	 */
+	UNREALED_API bool ShouldImportCurve(fbxsdk::FbxAnimCurve* Curve, bool bDoNotImportWithZeroValues);
+
+	/**
+	 * Helper function that extracts attribute curves from a specified Fbx node.
+	 * @param	InNode							The node to extract curves from
+	 * @param	bInDoNotImportCurveWithZero		Whether to skip zero-valued curves
+	 * @param	ImportFunction					Function called to perform operations on the extracted curve(s)
+	 */
+	UNREALED_API void ExtractAttributeCurves(fbxsdk::FbxNode* InNode, bool bInDoNotImportCurveWithZero, TFunctionRef<void(fbxsdk::FbxAnimCurve* /*InCurve*/, const FString& /*InCurveName*/, const FString& /*InChannelName*/, int32 /*InChannelIndex*/, int32 /*InNumChannels*/)> ImportFunction);
+
 } // namespace FbxAnimUtils
diff --git a/Engine/Source/Editor/UnrealEd/Public/SCommonEditorViewportToolbarBase.h b/Engine/Source/Editor/UnrealEd/Public/SCommonEditorViewportToolbarBase.h
--- a/Engine/Source/Editor/UnrealEd/Public/SCommonEditorViewportToolbarBase.h
+++ b/Engine/Source/Editor/UnrealEd/Public/SCommonEditorViewportToolbarBase.h
@@ -170,6 +170,9 @@ protected:
 
 	void CreateViewMenuExtensions(FMenuBuilder& MenuBuilder);
 
+	/** Extension allowing derived classes to add to the options menu.*/	
+	virtual void ExtendOptionsMenu(FMenuBuilder& OptionsMenuBuilder) const {}
+
 protected:
 	// Returns the info provider for this viewport
 	ICommonEditorViewportToolbarInfoProvider& GetInfoProvider() const;
diff --git a/Engine/Source/Editor/UnrealEd/Public/SCurveEditor.h b/Engine/Source/Editor/UnrealEd/Public/SCurveEditor.h
--- a/Engine/Source/Editor/UnrealEd/Public/SCurveEditor.h
+++ b/Engine/Source/Editor/UnrealEd/Public/SCurveEditor.h
@@ -282,6 +282,11 @@ public:
 	UNREALED_API void ZoomToFitVertical(const bool bZoomToFitAll = false);
 	UNREALED_API void ZoomToFit(const bool bZoomToFitAll = false);
 
+	/* Set flag that allows scrolling up/down over the widget from the outside without it handling the scroll wheel event */
+	UNREALED_API void SetRequireFocusToZoom(bool bInRequireFocusToZoom);
+
+	UNREALED_API virtual TOptional<bool> OnQueryShowFocus(const EFocusCause InFocusCause) const override;
+
 private:
 	/** Used to track a key and the curve that owns it */
 	struct FSelectedCurveKey
@@ -721,6 +726,9 @@ private:
 	/** Flag to allow auto framing */
 	bool bAllowAutoFrame;
 
+	/** Flag to allow scrolling up/down over the widget from the outside without it handling the scroll wheel event */
+	bool bRequireFocusToZoom;
+
 protected:
 	/** Minimum input of view range  */
 	TAttribute<float>	ViewMinInput;
diff --git a/Engine/Source/Editor/ViewportInteraction/Gizmo/VIGizmoHandleMeshComponent.cpp b/Engine/Source/Editor/ViewportInteraction/Gizmo/VIGizmoHandleMeshComponent.cpp
--- a/Engine/Source/Editor/ViewportInteraction/Gizmo/VIGizmoHandleMeshComponent.cpp
+++ b/Engine/Source/Editor/ViewportInteraction/Gizmo/VIGizmoHandleMeshComponent.cpp
@@ -5,9 +5,14 @@
 
 #include "RHIResources.h"
 
-class FGimzoHandleSceneProxy: public FStaticMeshSceneProxy
+class FGimzoHandleSceneProxy final : public FStaticMeshSceneProxy
 {
 public:
+	SIZE_T GetTypeHash() const override
+	{
+		static size_t UniquePointer;
+		return reinterpret_cast<size_t>(&UniquePointer);
+	}
 
 	FGimzoHandleSceneProxy(UGizmoHandleMeshComponent* InComponent)
 		: FStaticMeshSceneProxy(InComponent, false)
diff --git a/Engine/Source/Editor/ViewportInteraction/Public/ViewportWorldInteraction.h b/Engine/Source/Editor/ViewportInteraction/Public/ViewportWorldInteraction.h
--- a/Engine/Source/Editor/ViewportInteraction/Public/ViewportWorldInteraction.h
+++ b/Engine/Source/Editor/ViewportInteraction/Public/ViewportWorldInteraction.h
@@ -319,6 +319,10 @@ public:
 		bShouldSuppressCursor = bInShouldSuppressCursor;
 	};
 
+	/** Getters and setters for whether or not to show the cursor on the viewport */
+	bool ShouldForceCursor() const;
+	void SetForceCursor(const bool bInShouldForceCursor);
+
 	/** Gets the container for all the assets of ViewportInteraction. */
 	const class UViewportInteractionAssetContainer& GetAssetContainer() const;
 
@@ -701,6 +705,9 @@ private:
 	/** Whether or not to show the cursor on the viewport */
 	bool bShouldSuppressCursor;
 
+	/** Whether or not to force the cursor on the viewport */
+	bool bShouldForceCursor;
+
 	/** The current tick number */
 	uint32 CurrentTickNumber;
 
diff --git a/Engine/Source/Editor/ViewportInteraction/ViewportWorldInteraction.cpp b/Engine/Source/Editor/ViewportInteraction/ViewportWorldInteraction.cpp
--- a/Engine/Source/Editor/ViewportInteraction/ViewportWorldInteraction.cpp
+++ b/Engine/Source/Editor/ViewportInteraction/ViewportWorldInteraction.cpp
@@ -3010,6 +3010,16 @@ bool UViewportWorldInteraction::ShouldSuppressExistingCursor() const
 
 }
 
+void UViewportWorldInteraction::SetForceCursor(const bool bInShouldForceCursor)
+{
+	bShouldForceCursor = bInShouldForceCursor;
+}
+
+bool UViewportWorldInteraction::ShouldForceCursor() const
+{
+	return true;
+}
+
 const UViewportInteractionAssetContainer& UViewportWorldInteraction::GetAssetContainer() const
 {
 	return *AssetContainer;
diff --git a/Engine/Source/Programs/AutomationTool/AutomationUtils/ProjectParams.cs b/Engine/Source/Programs/AutomationTool/AutomationUtils/ProjectParams.cs
--- a/Engine/Source/Programs/AutomationTool/AutomationUtils/ProjectParams.cs
+++ b/Engine/Source/Programs/AutomationTool/AutomationUtils/ProjectParams.cs
@@ -227,6 +227,8 @@ namespace AutomationTool
 			this.DLCFile = InParams.DLCFile;
 			this.GenerateRemaster = InParams.GenerateRemaster;
             this.DLCIncludeEngineContent = InParams.DLCIncludeEngineContent;
+			this.DLCActLikePatch = InParams.DLCActLikePatch;
+			this.DLCPakPluginFile = InParams.DLCPakPluginFile;
             this.DiffCookedContentPath = InParams.DiffCookedContentPath;
             this.AdditionalCookerOptions = InParams.AdditionalCookerOptions;
 			this.ClientCookedTargets = InParams.ClientCookedTargets;
@@ -245,7 +247,7 @@ namespace AutomationTool
 			this.IterativeCooking = InParams.IterativeCooking;
 			this.IterateSharedCookedBuild = InParams.IterateSharedCookedBuild;
 			this.IterateSharedBuildUsePrecompiledExe = InParams.IterateSharedBuildUsePrecompiledExe;
-			this.CookAll = InParams.CookAll;
+            this.CookAll = InParams.CookAll;
 			this.CookPartialGC = InParams.CookPartialGC;
 			this.CookInEditor = InParams.CookInEditor; 
 			this.CookOutputDir = InParams.CookOutputDir;
@@ -265,7 +267,7 @@ namespace AutomationTool
             this.UnversionedCookedContent = InParams.UnversionedCookedContent;
 			this.EncryptIniFiles = InParams.EncryptIniFiles;
             this.EncryptPakIndex = InParams.EncryptPakIndex;
-            this.EncryptEverything = InParams.EncryptEverything;
+			this.EncryptEverything = InParams.EncryptEverything;
 			this.SkipCookingEditorContent = InParams.SkipCookingEditorContent;
             this.NumCookersToSpawn = InParams.NumCookersToSpawn;
 			this.FileServer = InParams.FileServer;
@@ -390,7 +392,7 @@ namespace AutomationTool
             bool? UnversionedCookedContent = null,
 			bool? EncryptIniFiles = null,
             bool? EncryptPakIndex = null,
-            bool? EncryptEverything = null,
+			bool? EncryptEverything = null,
 			bool? SkipCookingEditorContent = null,
             int? NumCookersToSpawn = null,
             string AdditionalCookerOptions = null,
@@ -402,9 +404,11 @@ namespace AutomationTool
 			bool? AddPatchLevel = null,
 			bool? StageBaseReleasePaks = null,
 			bool? GenerateRemaster = null,
-			string DLCName = null,
+            string DLCName = null,
             string DiffCookedContentPath = null,
             bool? DLCIncludeEngineContent = null,
+			bool? DLCPakPluginFile = null,
+			bool? DLCActLikePatch = null,
 			bool? CrashReporter = null,
 			bool? DedicatedServer = null,
 			bool? Client = null,
@@ -537,7 +541,7 @@ namespace AutomationTool
             this.AddPatchLevel = GetParamValueIfNotSpecified(Command, AddPatchLevel, this.AddPatchLevel, "AddPatchLevel");
 			this.StageBaseReleasePaks = GetParamValueIfNotSpecified(Command, StageBaseReleasePaks, this.StageBaseReleasePaks, "StageBaseReleasePaks");
 			this.GenerateRemaster = GetParamValueIfNotSpecified(Command, GenerateRemaster, this.GenerateRemaster, "GenerateRemaster");
-			this.AdditionalCookerOptions = ParseParamValueIfNotSpecified(Command, AdditionalCookerOptions, "AdditionalCookerOptions", String.Empty);
+            this.AdditionalCookerOptions = ParseParamValueIfNotSpecified(Command, AdditionalCookerOptions, "AdditionalCookerOptions", String.Empty);
 
 			DLCName = ParseParamValueIfNotSpecified(Command, DLCName, "DLCName", String.Empty);
 			if(!String.IsNullOrEmpty(DLCName))
@@ -554,6 +558,9 @@ namespace AutomationTool
             //this.DLCName = 
             this.DiffCookedContentPath = ParseParamValueIfNotSpecified(Command, DiffCookedContentPath, "DiffCookedContentPath", String.Empty);
             this.DLCIncludeEngineContent = GetParamValueIfNotSpecified(Command, DLCIncludeEngineContent, this.DLCIncludeEngineContent, "DLCIncludeEngineContent");
+			this.DLCPakPluginFile = GetParamValueIfNotSpecified(Command, DLCPakPluginFile, this.DLCPakPluginFile, "DLCPakPluginFile");
+			this.DLCActLikePatch = GetParamValueIfNotSpecified(Command, DLCActLikePatch, this.DLCActLikePatch, "DLCActLikePatch");
+			
 			this.SkipCook = GetParamValueIfNotSpecified(Command, SkipCook, this.SkipCook, "skipcook");
 			if (this.SkipCook)
 			{
@@ -586,7 +593,7 @@ namespace AutomationTool
             this.UnversionedCookedContent = GetParamValueIfNotSpecified(Command, UnversionedCookedContent, this.UnversionedCookedContent, "UnversionedCookedContent");
 			this.EncryptIniFiles = GetParamValueIfNotSpecified(Command, EncryptIniFiles, this.EncryptIniFiles, "EncryptIniFiles");
             this.EncryptPakIndex = GetParamValueIfNotSpecified(Command, EncryptPakIndex, this.EncryptPakIndex, "EncryptPakIndex");
-            this.EncryptEverything = GetParamValueIfNotSpecified(Command, EncryptEverything, this.EncryptEverything, "EncryptEverything");
+			this.EncryptEverything = GetParamValueIfNotSpecified(Command, EncryptEverything, this.EncryptEverything, "EncryptEverything");
 			this.SkipCookingEditorContent = GetParamValueIfNotSpecified(Command, SkipCookingEditorContent, this.SkipCookingEditorContent, "SkipCookingEditorContent");
             if (NumCookersToSpawn.HasValue)
             {
@@ -601,7 +608,7 @@ namespace AutomationTool
 			this.IterativeCooking = GetParamValueIfNotSpecified(Command, IterativeCooking, this.IterativeCooking, new string[] { "iterativecooking", "iterate" });
 			this.IterateSharedCookedBuild = GetParamValueIfNotSpecified(Command, IterateSharedCookedBuild, this.IterateSharedCookedBuild, new string[] { "IterateSharedCookedBuild"});
 			this.IterateSharedBuildUsePrecompiledExe = GetParamValueIfNotSpecified(Command, IterateSharedBuildUsePrecompiledExe, this.IterateSharedBuildUsePrecompiledExe, new string[] { "IterateSharedBuildUsePrecompiledExe" });
-
+			
 			this.SkipCookOnTheFly = GetParamValueIfNotSpecified(Command, SkipCookOnTheFly, this.SkipCookOnTheFly, "skipcookonthefly");
 			this.CookAll = GetParamValueIfNotSpecified(Command, CookAll, this.CookAll, "CookAll");
 			this.CookPartialGC = GetParamValueIfNotSpecified(Command, CookPartialGC, this.CookPartialGC, "CookPartialGC");
@@ -1328,8 +1335,8 @@ namespace AutomationTool
         /// Cook: Create a cooked release version.  Also, the version. e.g. 1.0
         /// </summary>
         public string CreateReleaseVersion;
-        
-		/// <summary>
+
+        /// <summary>
 		/// Cook: While cooking clean up packages as we go along rather then cleaning everything (and potentially having to reload some of it) when we run out of space
 		/// </summary>
 		[Help("CookPartialgc", "while cooking clean up packages as we are done with them rather then cleaning everything up when we run out of space")]
@@ -1367,8 +1374,8 @@ namespace AutomationTool
         /// see also CreateReleaseVersion, BasedOnReleaseVersion
         /// </summary>
         public bool GeneratePatch;
-		
-		/// <summary>
+
+        /// <summary>
 		/// Are we generating a remaster, generate a patch from a previously released version of the game (use CreateReleaseVersion to create a release). 
 		/// this requires BasedOnReleaseVersion
 		/// see also CreateReleaseVersion, BasedOnReleaseVersion
@@ -1383,7 +1390,7 @@ namespace AutomationTool
         public bool StageBaseReleasePaks;
 
         /// Name of dlc to cook and package (if this paramter is supplied cooks the dlc and packages it into the dlc directory)
-		/// </summary>
+        /// </summary>
         public FileReference DLCFile;
 
         /// <summary>
@@ -1392,11 +1399,23 @@ namespace AutomationTool
         /// </summary>
         public bool DLCIncludeEngineContent;
 
-        /// <summary>
-        /// After cook completes diff the cooked content against another cooked content directory.
-        ///  report all errors to the log
-        /// </summary>
-        public string DiffCookedContentPath;
+
+		/// <summary>
+		/// Enable packaging up the uplugin file inside the dlc pak.  This is sometimes desired if you want the plugin to be standalone
+		/// </summary>
+		public bool DLCPakPluginFile;
+
+		/// <summary>
+		/// DLC will remap it's files to the game directory and act like a patch.  This is useful if you want to override content in the main game along side your main game.
+		/// For example having different main game content in different DLCs
+		/// </summary>
+		public bool DLCActLikePatch;
+
+		/// <summary>
+		/// After cook completes diff the cooked content against another cooked content directory.
+		///  report all errors to the log
+		/// </summary>
+		public string DiffCookedContentPath;
 
         /// <summary>
         /// Cook: Additional cooker options to include on the cooker commandline
@@ -1423,15 +1442,15 @@ namespace AutomationTool
 		/// </summary>
 		public bool EncryptEverything;
 
-        /// <summary>
+		/// <summary>
 		/// Encrypt all files which are packaged into the pak file.  Only valid when encryption keys and building pak file specified. 
 		/// </summary>
         public bool EncryptPakIndex;
 
         /// <summary>
-        /// put -debug on the editorexe commandline
-        /// </summary>
-        public bool UseDebugParamForEditorExe;
+		/// put -debug on the editorexe commandline
+		/// </summary>
+		public bool UseDebugParamForEditorExe;
 
         /// <summary>
         /// Cook: Do not include a version number in the cooked content
@@ -2180,7 +2199,7 @@ namespace AutomationTool
             get { return GeneratePatch; }
         }
 
-        /// <summary>
+		/// <summary>
         /// True if we are generating a new patch pak tier
         /// </summary>
         public bool ShouldAddPatchLevel
@@ -2386,7 +2405,11 @@ namespace AutomationTool
 
             if (!HasDLCName && DLCIncludeEngineContent)
             {
-                throw new AutomationException("DLCIncludeEngineContent flag is only valid when cooking dlc.");
+                throw new AutomationException("DLCIncludeEngineContent flag is only valid when building DLC.");
+            }
+			if (!HasDLCName && DLCPakPluginFile )
+			{
+				throw new AutomationException("DLCPakPluginFile flag is only valid when building DLC.");
             }
 
             if ((IsGeneratingPatch || HasDLCName || ShouldAddPatchLevel) && !HasBasedOnReleaseVersion)
@@ -2466,17 +2489,18 @@ namespace AutomationTool
 				CommandUtils.LogLog("UnversionedCookedContent={0}", UnversionedCookedContent);
 				CommandUtils.LogLog("EncryptIniFiles={0}", EncryptIniFiles);
                 CommandUtils.LogLog("EncryptPakIndex={0}", EncryptPakIndex);
-                CommandUtils.LogLog("EncryptEverything={0}", EncryptEverything);
+				CommandUtils.LogLog("EncryptEverything={0}", EncryptEverything);
 				CommandUtils.LogLog("SkipCookingEditorContent={0}", SkipCookingEditorContent);
                 CommandUtils.LogLog("NumCookersToSpawn={0}", NumCookersToSpawn);
-				CommandUtils.LogLog("GeneratePatch={0}", GeneratePatch);
+                CommandUtils.LogLog("GeneratePatch={0}", GeneratePatch);
 				CommandUtils.LogLog("AddPatchLevel={0}", AddPatchLevel);
 				CommandUtils.LogLog("StageBaseReleasePaks={0}", StageBaseReleasePaks);
 				CommandUtils.LogLog("GenerateRemaster={0}", GenerateRemaster);
-				CommandUtils.LogLog("CreateReleaseVersion={0}", CreateReleaseVersion);
+                CommandUtils.LogLog("CreateReleaseVersion={0}", CreateReleaseVersion);
                 CommandUtils.LogLog("BasedOnReleaseVersion={0}", BasedOnReleaseVersion);
                 CommandUtils.LogLog("DLCFile={0}", DLCFile);
                 CommandUtils.LogLog("DLCIncludeEngineContent={0}", DLCIncludeEngineContent);
+				CommandUtils.LogLog("DLCPakPluginFile={0}", DLCPakPluginFile);
                 CommandUtils.LogLog("DiffCookedContentPath={0}", DiffCookedContentPath);
                 CommandUtils.LogLog("AdditionalCookerOptions={0}", AdditionalCookerOptions);
 				CommandUtils.LogLog("DedicatedServer={0}", DedicatedServer);
