Copying Orion-Staging @ 3748653 (Orion/Dev-General @ 3722124) to //UE4/Main

diff --git a/Engine/Source/Programs/AutomationTool/Scripts/CopyBuildToStagingDirectory.Automation.cs b/Engine/Source/Programs/AutomationTool/Scripts/CopyBuildToStagingDirectory.Automation.cs
--- a/Engine/Source/Programs/AutomationTool/Scripts/CopyBuildToStagingDirectory.Automation.cs
+++ b/Engine/Source/Programs/AutomationTool/Scripts/CopyBuildToStagingDirectory.Automation.cs
@@ -215,6 +215,45 @@ public partial class Project : CommandUtils
         }
     }
 
+    private static void StageAdditionalDirectoriesFromConfig(DeploymentContext SC, DirectoryReference ProjectContentRoot, StagedDirectoryReference StageContentRoot, ConfigHierarchy PlatformGameConfig, bool bUFS, string ConfigKeyName)
+    {
+        List<string> ExtraDirs;
+        if (PlatformGameConfig.GetArray("/Script/UnrealEd.ProjectPackagingSettings", ConfigKeyName, out ExtraDirs))
+        {
+            // Each string has the format '(Path="TheDirToStage")'
+            foreach (var PathStr in ExtraDirs)
+            {
+                string RelativePath = null;
+                var PathParts = PathStr.Split('"');
+                if (PathParts.Length == 3)
+                {
+                    RelativePath = PathParts[1];
+                }
+                else if (PathParts.Length == 1)
+                {
+                    RelativePath = PathParts[0];
+                }
+                if (RelativePath != null)
+                {
+                    DirectoryReference InputDir = DirectoryReference.Combine(ProjectContentRoot, RelativePath);
+					// Using Path because RelativePath may contain ".." which StagedDirectoryReference will not accept
+					StagedDirectoryReference OutputDir = new StagedDirectoryReference(Path.GetFullPath(Path.Combine(StageContentRoot.ToString(), RelativePath)));
+					if (bUFS)
+                    {
+                        List<FileReference> Files = SC.FindFilesToStage(InputDir, StageFilesSearch.AllDirectories);
+                        Files.RemoveAll(x => x.HasExtension(".uasset") || x.HasExtension(".umap"));
+
+                        SC.StageFiles(StagedFileType.UFS, InputDir, Files, OutputDir);
+                    }
+                    else
+                    {
+                        SC.StageFiles(StagedFileType.NonUFS, InputDir, StageFilesSearch.AllDirectories, OutputDir);
+                    }
+                }
+            }
+        }
+    }
+
 	public static void CreateStagingManifest(ProjectParams Params, DeploymentContext SC)
 	{
 		if (!Params.Stage)
@@ -233,15 +272,22 @@ public partial class Project : CommandUtils
 			// stage all the previously staged files
 			SC.StageFiles(StagedFileType.NonUFS, DirectoryReference.Combine(SC.ProjectRoot, "Saved", "SharedIterativeBuild", SC.CookPlatform, "Staged"), StageFilesSearch.AllDirectories, StagedDirectoryReference.Root); // remap to the root directory
 		}
-        if (Params.HasDLCName)
-        {
-            DirectoryReference DLCRoot = Params.DLCFile.Directory;
+		if (Params.HasDLCName)
+		{
+			DirectoryReference DLCRoot = Params.DLCFile.Directory;
 
-            // Making a plugin, grab the binaries too
-            SC.StageFile(StagedFileType.NonUFS, Params.DLCFile);
+			// Making a plugin, grab the binaries too
+			if (Params.DLCPakPluginFile)
+			{
+				SC.StageFile(StagedFileType.UFS, Params.DLCFile);
+			}
+			else
+			{
+				SC.StageFile(StagedFileType.NonUFS, Params.DLCFile);
+			}
 
 			DirectoryReference BinariesDir = DirectoryReference.Combine(DLCRoot, "Binaries");
-			if(DirectoryReference.Exists(BinariesDir))
+			if (DirectoryReference.Exists(BinariesDir))
 			{
 				SC.StageFiles(StagedFileType.NonUFS, BinariesDir, "libUE4-*.so", StageFilesSearch.AllDirectories);
 				SC.StageFiles(StagedFileType.NonUFS, BinariesDir, "UE4-*.dll", StageFilesSearch.AllDirectories);
@@ -249,357 +295,395 @@ public partial class Project : CommandUtils
 				SC.StageFiles(StagedFileType.NonUFS, BinariesDir, "UE4Server-*.dll", StageFilesSearch.AllDirectories);
 			}
 
-            // Put all of the cooked dir into the staged dir
-            DirectoryReference PlatformCookDir = String.IsNullOrEmpty(Params.CookOutputDir)? DirectoryReference.Combine(DLCRoot, "Saved", "Cooked", SC.CookPlatform) : DirectoryReference.Combine(new DirectoryReference(Params.CookOutputDir), SC.CookPlatform);
+			// Put all of the cooked dir into the staged dir
+			DirectoryReference PlatformCookDir = String.IsNullOrEmpty(Params.CookOutputDir) ? DirectoryReference.Combine(DLCRoot, "Saved", "Cooked", SC.CookPlatform) : DirectoryReference.Combine(new DirectoryReference(Params.CookOutputDir), SC.CookPlatform);
 
-            // Stage any loose files in the root folder
+			// Stage any loose files in the root folder
 			// TODO: not sure if we should stage the loose files if we have pak files enabled... 
 			SC.StageFiles(StagedFileType.UFS, PlatformCookDir, StageFilesSearch.TopDirectoryOnly, SC.RelativeProjectRootForStage);
 
+			SC.StageFiles(StagedFileType.UFS, DirectoryReference.Combine(DLCRoot, "Config"), "*.ini", StageFilesSearch.AllDirectories);
+
+			if (Params.DLCActLikePatch)
+			{
+				DirectoryReference DLCContent = DirectoryReference.Combine(DLCRoot, "Content");
+				string RelativeDLCContentPath = DLCContent.MakeRelativeTo(SC.LocalRoot);
+				string RelativeRootContentPath = DirectoryReference.Combine(SC.ProjectRoot, "Content").MakeRelativeTo(SC.LocalRoot);
+
+				SC.RemapDirectories.Add(Tuple.Create(new StagedDirectoryReference(RelativeDLCContentPath), new StagedDirectoryReference(RelativeRootContentPath)));
+			}
+
+			StagedDirectoryReference EngineDirectory = new StagedDirectoryReference("Engine");
 			// Stage each sub directory separately so that we can skip Engine if need be
-			foreach(DirectoryReference SubDir in DirectoryReference.EnumerateDirectories(PlatformCookDir))
-            {
+			foreach (DirectoryReference SubDir in DirectoryReference.EnumerateDirectories(PlatformCookDir))
+			{
 				StagedDirectoryReference MountPoint = new StagedDirectoryReference(SubDir.MakeRelativeTo(PlatformCookDir));
-                // Dedicated server cook doesn't save shaders so no Engine dir is created
-                if ((!SC.DedicatedServer) && (!Params.DLCIncludeEngineContent) && MountPoint == new StagedDirectoryReference("Engine"))
-                {
-                    continue;
-                }
-				// SC.StageFilesUFS(CombinePaths(SC.ProjectRoot, "Plugins", DLCName, "Saved", "Cooked", SC.CookPlatform), "*", true, new[] { CommandUtils.CombinePaths("Engine", "*") }, "", true, !Params.UsePak(SC.StageTargetPlatform));
+				// Dedicated server cook doesn't save shaders so no Engine dir is created
+
+				if ((!SC.DedicatedServer) && (!Params.DLCIncludeEngineContent) && MountPoint == EngineDirectory)
+				{
+					continue;
+				}
+
 				SC.StageFiles(StagedFileType.UFS, SubDir, "*", StageFilesSearch.AllDirectories, MountPoint);
-				//SC.StageFilesUFS(SubDir, "*", true, ExcludeWildCards, SC.RelativeProjectRootForStage, true, !Params.UsePak(SC.StageTargetPlatform));
 			}
 
-            return;
-        }
+			FileReference PluginSettingsFile = FileReference.Combine(DLCRoot, "Config", "PluginSettings.ini");
+			if (FileReference.Exists(PluginSettingsFile))
+			{
+				ConfigFile File = new ConfigFile(PluginSettingsFile);
+				ConfigFileSection StageSettings;
+				if (File.TryGetSection("StageSettings", out StageSettings))
+				{
+					foreach (ConfigLine Line in StageSettings.Lines)
+					{
+						if (Line.Key == "AdditionalNonUSFDirectories")
+						{
+							SC.StageFiles(StagedFileType.NonUFS, DirectoryReference.Combine(DLCRoot, Line.Value), "*.*", StageFilesSearch.AllDirectories);
+						}
+					}
+				}
 
-		if (!Params.IterateSharedBuildUsePrecompiledExe)
-		{
-			ThisPlatform.GetFilesToDeployOrStage(Params, SC);
+				if (SC.DedicatedServer)
+				{
+					ConfigFileSection StageSettingsServer;
+					if (File.TryGetSection("StageSettingsServer", out StageSettingsServer))
+					{
+						foreach (ConfigLine Line in StageSettingsServer.Lines)
+						{
+							if (Line.Key == "AdditionalNonUSFDirectories")
+							{
+								SC.StageFiles(StagedFileType.NonUFS, DirectoryReference.Combine(DLCRoot, Line.Value), "*.*", StageFilesSearch.AllDirectories);
+							}
+						}
+					}
+				}
+			}
 
-			// Stage any extra runtime dependencies from the receipts
-			foreach (StageTarget Target in SC.StageTargets)
+			if (Params.UsePak(SC.StageTargetPlatform))
 			{
-				SC.StageRuntimeDependenciesFromReceipt(Target.Receipt, Target.RequireFilesExist, Params.UsePak(SC.StageTargetPlatform));
+				CreatePluginManifest(SC, SC.FilesToStage.UFSFiles, StagedFileType.UFS, Params.DLCFile.GetFileNameWithoutAnyExtensions());
 			}
 		}
-
-		// move the UE4Commandline.txt file to the root of the stage
-		// this file needs to be treated as a UFS file for casing, but NonUFS for being put into the .pak file
-		// @todo: Maybe there should be a new category - UFSNotForPak
-		FileReference CommandLineFile = FileReference.Combine(GetIntermediateCommandlineDir(SC), "UE4CommandLine.txt");
-		if(FileReference.Exists(CommandLineFile))
+		else
 		{
-			StagedFileReference StagedCommandLineFile = new StagedFileReference("UE4CommandLine.txt");
-			if (SC.StageTargetPlatform.DeployLowerCaseFilenames())
+
+			if (!Params.IterateSharedBuildUsePrecompiledExe)
 			{
-				StagedCommandLineFile = StagedCommandLineFile.ToLowerInvariant();
+				ThisPlatform.GetFilesToDeployOrStage(Params, SC);
+
+				// Stage any extra runtime dependencies from the receipts
+				foreach (StageTarget Target in SC.StageTargets)
+				{
+					SC.StageRuntimeDependenciesFromReceipt(Target.Receipt, Target.RequireFilesExist, Params.UsePak(SC.StageTargetPlatform));
+				}
 			}
-			SC.StageFile(StagedFileType.SystemNonUFS, CommandLineFile, StagedCommandLineFile);
-		}
 
-        ConfigHierarchy PlatformGameConfig = ConfigCache.ReadHierarchy(ConfigHierarchyType.Game, DirectoryReference.FromFile(Params.RawProjectPath), SC.StageTargetPlatform.IniPlatformType);
-        DirectoryReference ProjectContentRoot = DirectoryReference.Combine(SC.ProjectRoot, "Content");
-        StagedDirectoryReference StageContentRoot = StagedDirectoryReference.Combine(SC.RelativeProjectRootForStage, "Content");
-        
-        if (!Params.CookOnTheFly && !Params.SkipCookOnTheFly) // only stage the UFS files if we are not using cook on the fly
-        {
+			// move the UE4Commandline.txt file to the root of the stage
+			// this file needs to be treated as a UFS file for casing, but NonUFS for being put into the .pak file
+			// @todo: Maybe there should be a new category - UFSNotForPak
+			FileReference CommandLineFile = FileReference.Combine(GetIntermediateCommandlineDir(SC), "UE4CommandLine.txt");
+			if (FileReference.Exists(CommandLineFile))
+			{
+				StagedFileReference StagedCommandLineFile = new StagedFileReference("UE4CommandLine.txt");
+				if (SC.StageTargetPlatform.DeployLowerCaseFilenames())
+				{
+					StagedCommandLineFile = StagedCommandLineFile.ToLowerInvariant();
+				}
+				SC.StageFile(StagedFileType.SystemNonUFS, CommandLineFile, StagedCommandLineFile);
+			}
 
+			ConfigHierarchy PlatformGameConfig = ConfigCache.ReadHierarchy(ConfigHierarchyType.Game, DirectoryReference.FromFile(Params.RawProjectPath), SC.StageTargetPlatform.IniPlatformType);
+			DirectoryReference ProjectContentRoot = DirectoryReference.Combine(SC.ProjectRoot, "Content");
+			StagedDirectoryReference StageContentRoot = StagedDirectoryReference.Combine(SC.RelativeProjectRootForStage, "Content");
 
-            // Initialize internationalization preset.
-            string InternationalizationPreset = Params.InternationalizationPreset;
+			if (!Params.CookOnTheFly && !Params.SkipCookOnTheFly) // only stage the UFS files if we are not using cook on the fly
+			{
 
-            // Use configuration if otherwise lacking an internationalization preset.
-            if (string.IsNullOrEmpty(InternationalizationPreset))
-            {
-                if (PlatformGameConfig != null)
-                {
-                    PlatformGameConfig.GetString("/Script/UnrealEd.ProjectPackagingSettings", "InternationalizationPreset", out InternationalizationPreset);
-                }
-            }
 
-            // Error if no preset has been provided.
-            if (string.IsNullOrEmpty(InternationalizationPreset))
-            {
-                throw new AutomationException("No internationalization preset was specified for packaging. This will lead to fatal errors when launching. Specify preset via commandline (-I18NPreset=) or project packaging settings (InternationalizationPreset).");
-            }
+				// Initialize internationalization preset.
+				string InternationalizationPreset = Params.InternationalizationPreset;
 
-            // Initialize cultures to stage.
-            List<string> CulturesToStage = null;
+				// Use configuration if otherwise lacking an internationalization preset.
+				if (string.IsNullOrEmpty(InternationalizationPreset))
+				{
+					if (PlatformGameConfig != null)
+					{
+						PlatformGameConfig.GetString("/Script/UnrealEd.ProjectPackagingSettings", "InternationalizationPreset", out InternationalizationPreset);
+					}
+				}
 
-            // Use parameters if provided.
-            if (Params.CulturesToCook != null && Params.CulturesToCook.Count > 0)
-            {
-                CulturesToStage = Params.CulturesToCook;
-            }
+				// Error if no preset has been provided.
+				if (string.IsNullOrEmpty(InternationalizationPreset))
+				{
+					throw new AutomationException("No internationalization preset was specified for packaging. This will lead to fatal errors when launching. Specify preset via commandline (-I18NPreset=) or project packaging settings (InternationalizationPreset).");
+				}
 
-            // Use configuration if otherwise lacking cultures to stage.
-            if (CulturesToStage == null || CulturesToStage.Count == 0)
-            {
-                if (PlatformGameConfig != null)
-                {
-                    PlatformGameConfig.GetArray("/Script/UnrealEd.ProjectPackagingSettings", "CulturesToStage", out CulturesToStage);
-                }
-            }
+				// Initialize cultures to stage.
+				List<string> CulturesToStage = null;
 
-            // Error if no cultures have been provided.
-            if (CulturesToStage == null || CulturesToStage.Count == 0)
-            {
-                throw new AutomationException("No cultures were specified for cooking and packaging. This will lead to fatal errors when launching. Specify culture codes via commandline (-CookCultures=) or using project packaging settings (+CulturesToStage).");
-            }
+				// Use parameters if provided.
+				if (Params.CulturesToCook != null && Params.CulturesToCook.Count > 0)
+				{
+					CulturesToStage = Params.CulturesToCook;
+				}
 
-            // Stage ICU internationalization data from Engine.
-            SC.StageFiles(StagedFileType.UFS, DirectoryReference.Combine(SC.LocalRoot, "Engine", "Content", "Internationalization", InternationalizationPreset), StageFilesSearch.AllDirectories, new StagedDirectoryReference("Engine/Content/Internationalization"));
+				// Use configuration if otherwise lacking cultures to stage.
+				if (CulturesToStage == null || CulturesToStage.Count == 0)
+				{
+					if (PlatformGameConfig != null)
+					{
+						PlatformGameConfig.GetArray("/Script/UnrealEd.ProjectPackagingSettings", "CulturesToStage", out CulturesToStage);
+					}
+				}
 
-            // Engine ufs (content)
-            SC.StageFiles(StagedFileType.UFS, DirectoryReference.Combine(SC.LocalRoot, "Engine", "Config"), StageFilesSearch.AllDirectories); // TODO: Exclude localization data generation config files.
+				// Error if no cultures have been provided.
+				if (CulturesToStage == null || CulturesToStage.Count == 0)
+				{
+					throw new AutomationException("No cultures were specified for cooking and packaging. This will lead to fatal errors when launching. Specify culture codes via commandline (-CookCultures=) or using project packaging settings (+CulturesToStage).");
+				}
 
-            StageLocalizationDataForTarget(SC, CulturesToStage, DirectoryReference.Combine(SC.LocalRoot, "Engine", "Content", "Localization", "Engine"));
-			
-            // Game ufs (content)
-			SC.StageFile(StagedFileType.UFS, SC.RawProjectPath);
+				// Stage ICU internationalization data from Engine.
+				SC.StageFiles(StagedFileType.UFS, DirectoryReference.Combine(SC.LocalRoot, "Engine", "Content", "Internationalization", InternationalizationPreset), StageFilesSearch.AllDirectories, new StagedDirectoryReference("Engine/Content/Internationalization"));
 
-			SC.StageFiles(StagedFileType.UFS, DirectoryReference.Combine(SC.ProjectRoot, "Config"), "*", StageFilesSearch.AllDirectories); // TODO: Exclude localization data generation config files.
+				// Engine ufs (content)
+				SC.StageFiles(StagedFileType.UFS, DirectoryReference.Combine(SC.LocalRoot, "Engine", "Config"), StageFilesSearch.AllDirectories); // TODO: Exclude localization data generation config files.
 
-			// Stage all project localization targets
-			{
-				DirectoryReference ProjectLocRootDirectory = DirectoryReference.Combine(SC.ProjectRoot, "Content", "Localization");
-				if (DirectoryReference.Exists(ProjectLocRootDirectory))
+				StageLocalizationDataForTarget(SC, CulturesToStage, DirectoryReference.Combine(SC.LocalRoot, "Engine", "Content", "Localization", "Engine"));
+
+				// Game ufs (content)
+				SC.StageFile(StagedFileType.UFS, SC.RawProjectPath);
+
+				SC.StageFiles(StagedFileType.UFS, DirectoryReference.Combine(SC.ProjectRoot, "Config"), "*", StageFilesSearch.AllDirectories); // TODO: Exclude localization data generation config files.
+
+				// Stage plugin config files
+				List<KeyValuePair<StagedFileReference, FileReference>> StagedPlugins = SC.FilesToStage.UFSFiles.Where(x => x.Value.HasExtension(".uplugin")).ToList();
+				foreach (KeyValuePair<StagedFileReference, FileReference> StagedPlugin in StagedPlugins)
 				{
-					foreach(DirectoryReference ProjectLocTargetDirectory in DirectoryReference.EnumerateDirectories(ProjectLocRootDirectory))
+					//PluginDescriptor Descriptor = PluginDescriptor.FromFile(StagedPlugin.Value);
+					DirectoryReference PluginConfigDirectory = DirectoryReference.Combine(StagedPlugin.Value.Directory, "Config");
+					if (DirectoryReference.Exists(PluginConfigDirectory))
 					{
-						StageLocalizationDataForTarget(SC, CulturesToStage, ProjectLocTargetDirectory);
+						SC.StageFiles(StagedFileType.UFS, PluginConfigDirectory, "*.ini", StageFilesSearch.AllDirectories);
 					}
 				}
-			}
 
-			// Stage all plugin localization targets
-			List<KeyValuePair<StagedFileReference, FileReference>> StagedPlugins = SC.FilesToStage.UFSFiles.Where(x => x.Value.HasExtension(".uplugin")).ToList();
-			foreach(KeyValuePair<StagedFileReference, FileReference> StagedPlugin in StagedPlugins)
-			{
-				PluginDescriptor Descriptor = PluginDescriptor.FromFile(StagedPlugin.Value);
-				if(Descriptor.LocalizationTargets != null)
+				// Stage all project localization targets
 				{
-					foreach (LocalizationTargetDescriptor LocalizationTarget in Descriptor.LocalizationTargets)
+					DirectoryReference ProjectLocRootDirectory = DirectoryReference.Combine(SC.ProjectRoot, "Content", "Localization");
+					if (DirectoryReference.Exists(ProjectLocRootDirectory))
 					{
-						if (LocalizationTarget.LoadingPolicy == LocalizationTargetDescriptorLoadingPolicy.Always || LocalizationTarget.LoadingPolicy == LocalizationTargetDescriptorLoadingPolicy.Game)
+						foreach (DirectoryReference ProjectLocTargetDirectory in DirectoryReference.EnumerateDirectories(ProjectLocRootDirectory))
 						{
-							DirectoryReference PluginLocTargetDirectory = DirectoryReference.Combine(StagedPlugin.Value.Directory, "Content", "Localization", LocalizationTarget.Name);
-							if (DirectoryReference.Exists(PluginLocTargetDirectory))
-							{
-								StageLocalizationDataForTarget(SC, CulturesToStage, PluginLocTargetDirectory);
-							}
+							StageLocalizationDataForTarget(SC, CulturesToStage, ProjectLocTargetDirectory);
 						}
 					}
 				}
-			}
 
-            // Stage any additional UFS and NonUFS paths specified in the project ini files; these dirs are relative to the game content directory
-            if (PlatformGameConfig != null)
-            {
-                List<string> ExtraUFSDirs;
-                if (PlatformGameConfig.GetArray("/Script/UnrealEd.ProjectPackagingSettings", "DirectoriesToAlwaysStageAsUFS", out ExtraUFSDirs))
-                {
-					List<string> InputPaths = ParseInputPaths(ExtraUFSDirs);
-					foreach(string InputPath in InputPaths)
+				// Stage all plugin localization targets
+				foreach (KeyValuePair<StagedFileReference, FileReference> StagedPlugin in StagedPlugins)
+				{
+					PluginDescriptor Descriptor = PluginDescriptor.FromFile(StagedPlugin.Value);
+					if (Descriptor.LocalizationTargets != null)
 					{
-						DirectoryReference InputDir = DirectoryReference.Combine(ProjectContentRoot, InputPath);
-						if(DirectoryReference.Exists(InputDir))
+						foreach (LocalizationTargetDescriptor LocalizationTarget in Descriptor.LocalizationTargets)
 						{
-							StagedDirectoryReference OutputDir = StagedDirectoryReference.Combine(StageContentRoot, InputPath);
-
-							List<FileReference> Files = SC.FindFilesToStage(InputDir, StageFilesSearch.AllDirectories);
-							Files.RemoveAll(x => x.HasExtension(".uasset") || x.HasExtension(".umap"));
-
-							SC.StageFiles(StagedFileType.UFS, InputDir, Files, OutputDir);
+							if (LocalizationTarget.LoadingPolicy == LocalizationTargetDescriptorLoadingPolicy.Always || LocalizationTarget.LoadingPolicy == LocalizationTargetDescriptorLoadingPolicy.Game)
+							{
+								DirectoryReference PluginLocTargetDirectory = DirectoryReference.Combine(StagedPlugin.Value.Directory, "Content", "Localization", LocalizationTarget.Name);
+								if (DirectoryReference.Exists(PluginLocTargetDirectory))
+								{
+									StageLocalizationDataForTarget(SC, CulturesToStage, PluginLocTargetDirectory);
+								}
+							}
 						}
 					}
-                }
+				}
 
-                List<string> ExtraNonUFSDirs;
-                if (PlatformGameConfig.GetArray("/Script/UnrealEd.ProjectPackagingSettings", "DirectoriesToAlwaysStageAsNonUFS", out ExtraNonUFSDirs))
-                {
-					List<string> InputPaths = ParseInputPaths(ExtraNonUFSDirs);
-					foreach(string InputPath in InputPaths)
+				// Stage any additional UFS and NonUFS paths specified in the project ini files; these dirs are relative to the game content directory
+				if (PlatformGameConfig != null)
+				{
+					StageAdditionalDirectoriesFromConfig(SC, ProjectContentRoot, StageContentRoot, PlatformGameConfig, true, "DirectoriesToAlwaysStageAsUFS");
+					// NonUFS files are never in pak files and should always be remapped
+					StageAdditionalDirectoriesFromConfig(SC, ProjectContentRoot, StageContentRoot, PlatformGameConfig, false, "DirectoriesToAlwaysStageAsNonUFS");
+
+					if (SC.DedicatedServer)
 					{
-						DirectoryReference InputDir = DirectoryReference.Combine(ProjectContentRoot, InputPath);
-						if(DirectoryReference.Exists(InputDir))
-						{
-							StagedDirectoryReference OutputDir = StagedDirectoryReference.Combine(StageContentRoot, InputPath);
-							SC.StageFiles(StagedFileType.NonUFS, InputDir, StageFilesSearch.AllDirectories, OutputDir);
-						}
+						StageAdditionalDirectoriesFromConfig(SC, ProjectContentRoot, StageContentRoot, PlatformGameConfig, true, "DirectoriesToAlwaysStageAsUFSServer");
+						// NonUFS files are never in pak files and should always be remapped
+						StageAdditionalDirectoriesFromConfig(SC, ProjectContentRoot, StageContentRoot, PlatformGameConfig, false, "DirectoriesToAlwaysStageAsNonUFSServer");
 					}
-                }
-            }
-
-            if (SC.StageTargetPlatform.StageMovies && !SC.DedicatedServer)
-            {
-                // UFS is required when using a file server
-				StagedFileType MovieFileType = Params.FileServer? StagedFileType.UFS : StagedFileType.NonUFS;
-
-				DirectoryReference EngineMoviesDir = DirectoryReference.Combine(SC.EngineRoot, "Content", "Movies");
-				if(DirectoryReference.Exists(EngineMoviesDir))
-				{
-					List<FileReference> MovieFiles = SC.FindFilesToStage(EngineMoviesDir, StageFilesSearch.AllDirectories);
-					SC.StageFiles(MovieFileType, MovieFiles.Where(x => !x.HasExtension(".uasset") && !x.HasExtension(".umap")));
 				}
 
-				DirectoryReference ProjectMoviesDir = DirectoryReference.Combine(SC.ProjectRoot, "Content", "Movies");
-				if(DirectoryReference.Exists(ProjectMoviesDir))
+				if (SC.StageTargetPlatform.StageMovies && !SC.DedicatedServer)
 				{
-					List<FileReference> MovieFiles = SC.FindFilesToStage(ProjectMoviesDir, StageFilesSearch.AllDirectories);
-					SC.StageFiles(MovieFileType, MovieFiles.Where(x => !x.HasExtension(".uasset") && !x.HasExtension(".umap")));
+					// UFS is required when using a file server
+					StagedFileType MovieFileType = Params.FileServer ? StagedFileType.UFS : StagedFileType.NonUFS;
+
+					DirectoryReference EngineMoviesDir = DirectoryReference.Combine(SC.EngineRoot, "Content", "Movies");
+					if (DirectoryReference.Exists(EngineMoviesDir))
+					{
+						List<FileReference> MovieFiles = SC.FindFilesToStage(EngineMoviesDir, StageFilesSearch.AllDirectories);
+						SC.StageFiles(MovieFileType, MovieFiles.Where(x => !x.HasExtension(".uasset") && !x.HasExtension(".umap")));
+					}
+
+					DirectoryReference ProjectMoviesDir = DirectoryReference.Combine(SC.ProjectRoot, "Content", "Movies");
+					if (DirectoryReference.Exists(ProjectMoviesDir))
+					{
+						List<FileReference> MovieFiles = SC.FindFilesToStage(ProjectMoviesDir, StageFilesSearch.AllDirectories);
+						SC.StageFiles(MovieFileType, MovieFiles.Where(x => !x.HasExtension(".uasset") && !x.HasExtension(".umap")));
+					}
 				}
-            }
 
-			// shader cache files
-			{
+			// shader cache
+				{
 				DirectoryReference ShaderCacheRoot = DirectoryReference.Combine(SC.ProjectRoot, "Content");
 				List<FileReference> ShaderCacheFiles = SC.FindFilesToStage(ShaderCacheRoot, "DrawCache-*.ushadercache", StageFilesSearch.TopDirectoryOnly);
 				SC.StageFiles(StagedFileType.UFS, ShaderCacheFiles);
-			}
+				}
 
-			// Get the final output directory for cooked data
-			DirectoryReference CookOutputDir;
-			if(!String.IsNullOrEmpty(Params.CookOutputDir))
-			{
-				CookOutputDir = DirectoryReference.Combine(new DirectoryReference(Params.CookOutputDir), SC.CookPlatform);
-			}
-			else if(Params.CookInEditor)
-			{
-				CookOutputDir = DirectoryReference.Combine(SC.ProjectRoot, "Saved", "EditorCooked", SC.CookPlatform);
-			}
-			else
-			{
-				CookOutputDir = DirectoryReference.Combine(SC.ProjectRoot, "Saved", "Cooked", SC.CookPlatform);
-			}
+				// Get the final output directory for cooked data
+				DirectoryReference CookOutputDir;
+				if (!String.IsNullOrEmpty(Params.CookOutputDir))
+				{
+					CookOutputDir = DirectoryReference.Combine(new DirectoryReference(Params.CookOutputDir), SC.CookPlatform);
+				}
+				else if (Params.CookInEditor)
+				{
+					CookOutputDir = DirectoryReference.Combine(SC.ProjectRoot, "Saved", "EditorCooked", SC.CookPlatform);
+				}
+				else
+				{
+					CookOutputDir = DirectoryReference.Combine(SC.ProjectRoot, "Saved", "Cooked", SC.CookPlatform);
+				}
 
-			// Stage all the cooked data. Currently not filtering this by restricted folders, since we shouldn't mask invalid references by filtering them out.
-			List<FileReference> CookedFiles = DirectoryReference.EnumerateFiles(CookOutputDir, "*", SearchOption.AllDirectories).ToList();
-			foreach(FileReference CookedFile in CookedFiles)
-			{
+				// Stage all the cooked data. Currently not filtering this by restricted folders, since we shouldn't mask invalid references by filtering them out.
+				List<FileReference> CookedFiles = DirectoryReference.EnumerateFiles(CookOutputDir, "*", SearchOption.AllDirectories).ToList();
+				foreach (FileReference CookedFile in CookedFiles)
+				{
 				// json files have never been staged
 				// metallib files cannot *currently* be staged as UFS as the Metal API needs to mmap them from files on disk in order to function efficiently
 				if(!CookedFile.HasExtension(".json") && !CookedFile.HasExtension(".metallib"))
-				{
-					SC.StageFile(StagedFileType.UFS, CookedFile, new StagedFileReference(CookedFile.MakeRelativeTo(CookOutputDir)));
+					{
+						SC.StageFile(StagedFileType.UFS, CookedFile, new StagedFileReference(CookedFile.MakeRelativeTo(CookOutputDir)));
+					}
 				}
-			}
 
-            // CrashReportClient is a standalone slate app that does not look in the generated pak file, so it needs the Content/Slate and Shaders/StandaloneRenderer folders Non-UFS
-            // @todo Make CrashReportClient more portable so we don't have to do this
-            if (SC.bStageCrashReporter && PlatformSupportsCrashReporter(SC.StageTargetPlatform.PlatformType) && (Params.IterateSharedBuildUsePrecompiledExe == false))
-            {
-                //If the .dat file needs to be staged as NonUFS for non-Windows/Linux hosts we need to change the casing as we do with the build properties file above.
-                SC.StageFiles(StagedFileType.NonUFS, DirectoryReference.Combine(SC.EngineRoot, "Content", "Slate"), StageFilesSearch.AllDirectories);
-                SC.StageFiles(StagedFileType.NonUFS, DirectoryReference.Combine(SC.EngineRoot, "Shaders", "StandaloneRenderer"), StageFilesSearch.AllDirectories);
+				// CrashReportClient is a standalone slate app that does not look in the generated pak file, so it needs the Content/Slate and Shaders/StandaloneRenderer folders Non-UFS
+				// @todo Make CrashReportClient more portable so we don't have to do this
+				if (SC.bStageCrashReporter && PlatformSupportsCrashReporter(SC.StageTargetPlatform.PlatformType) && (Params.IterateSharedBuildUsePrecompiledExe == false))
+				{
+					//If the .dat file needs to be staged as NonUFS for non-Windows/Linux hosts we need to change the casing as we do with the build properties file above.
+					SC.StageFiles(StagedFileType.NonUFS, DirectoryReference.Combine(SC.EngineRoot, "Content", "Slate"), StageFilesSearch.AllDirectories);
+					SC.StageFiles(StagedFileType.NonUFS, DirectoryReference.Combine(SC.EngineRoot, "Shaders", "StandaloneRenderer"), StageFilesSearch.AllDirectories);
 
-                SC.StageFiles(StagedFileType.NonUFS, DirectoryReference.Combine(SC.EngineRoot, "Content", "Internationalization", InternationalizationPreset), StageFilesSearch.AllDirectories, new StagedDirectoryReference("Engine/Content/Internationalization"));
+					SC.StageFiles(StagedFileType.NonUFS, DirectoryReference.Combine(SC.EngineRoot, "Content", "Internationalization", InternationalizationPreset), StageFilesSearch.AllDirectories, new StagedDirectoryReference("Engine/Content/Internationalization"));
 
-                // Get the architecture in use
-                string Architecture = Params.SpecifiedArchitecture;
-                if (string.IsNullOrEmpty(Architecture))
-                {
-                    Architecture = "";
-                    if (PlatformExports.IsPlatformAvailable(SC.StageTargetPlatform.PlatformType))
-                    {
-                        Architecture = PlatformExports.GetDefaultArchitecture(SC.StageTargetPlatform.PlatformType, Params.RawProjectPath);
-                    }
-                }
-
-                // Get the target receipt path for CrashReportClient
-                FileReference ReceiptFileName = TargetReceipt.GetDefaultPath(SC.EngineRoot, "CrashReportClient", SC.StageTargetPlatform.PlatformType, UnrealTargetConfiguration.Shipping, Architecture);
-                if (!FileReference.Exists(ReceiptFileName))
-                {
-                    throw new AutomationException(ExitCode.Error_MissingExecutable, "Stage Failed. Missing receipt '{0}'. Check that this target has been built.", Path.GetFileName(ReceiptFileName.FullName));
-                }
+					// Get the architecture in use
+					string Architecture = Params.SpecifiedArchitecture;
+					if (string.IsNullOrEmpty(Architecture))
+					{
+						Architecture = "";
+						if (PlatformExports.IsPlatformAvailable(SC.StageTargetPlatform.PlatformType))
+						{
+							Architecture = PlatformExports.GetDefaultArchitecture(SC.StageTargetPlatform.PlatformType, Params.RawProjectPath);
+						}
+					}
 
-                // Read the receipt for this target
-                TargetReceipt Receipt;
-                if (!TargetReceipt.TryRead(ReceiptFileName, SC.EngineRoot, null, out Receipt))
-                {
-                    throw new AutomationException("Missing or invalid target receipt ({0})", ReceiptFileName);
-                }
+					// Get the target receipt path for CrashReportClient
+					FileReference ReceiptFileName = TargetReceipt.GetDefaultPath(SC.EngineRoot, "CrashReportClient", SC.StageTargetPlatform.PlatformType, UnrealTargetConfiguration.Shipping, Architecture);
+					if (!FileReference.Exists(ReceiptFileName))
+					{
+						throw new AutomationException(ExitCode.Error_MissingExecutable, "Stage Failed. Missing receipt '{0}'. Check that this target has been built.", Path.GetFileName(ReceiptFileName.FullName));
+					}
 
-                // Stage any runtime dependencies for CrashReportClient
-                SC.StageRuntimeDependenciesFromReceipt(Receipt, true, Params.UsePak(SC.StageTargetPlatform));
+					// Read the receipt for this target
+					TargetReceipt Receipt;
+					if (!TargetReceipt.TryRead(ReceiptFileName, SC.EngineRoot, null, out Receipt))
+					{
+						throw new AutomationException("Missing or invalid target receipt ({0})", ReceiptFileName);
+					}
 
-				// Add config files.
-				SC.StageFiles(StagedFileType.NonUFS, DirectoryReference.Combine(SC.EngineRoot, "Programs", "CrashReportClient", "Config" ), StageFilesSearch.AllDirectories);
-			}
+					// Stage any runtime dependencies for CrashReportClient
+					SC.StageRuntimeDependenciesFromReceipt(Receipt, true, Params.UsePak(SC.StageTargetPlatform));
 
-			// check if the game will be verifying ssl connections - if not, we can skip staging files that won't be needed
-			bool bStageSSLCertificates = false;
-			ConfigHierarchy PlatformEngineConfig = ConfigCache.ReadHierarchy(ConfigHierarchyType.Engine, DirectoryReference.FromFile(Params.RawProjectPath), SC.StageTargetPlatform.IniPlatformType);
-			if (PlatformEngineConfig != null)
-			{
-				PlatformEngineConfig.GetBool("/Script/Engine.NetworkSettings", "n.VerifyPeer", out bStageSSLCertificates);
-			}
+					// Add config files.
+					SC.StageFiles(StagedFileType.NonUFS, DirectoryReference.Combine(SC.EngineRoot, "Programs", "CrashReportClient", "Config"), StageFilesSearch.AllDirectories);
+				}
 
-			if (bStageSSLCertificates)
-			{
-				// Game's SSL certs
-				FileReference ProjectCertFile = FileReference.Combine(SC.ProjectRoot, "Content", "Certificates", "cacert.pem");
-				if(FileReference.Exists(ProjectCertFile))
+				// check if the game will be verifying ssl connections - if not, we can skip staging files that won't be needed
+				bool bStageSSLCertificates = false;
+				ConfigHierarchy PlatformEngineConfig = ConfigCache.ReadHierarchy(ConfigHierarchyType.Engine, DirectoryReference.FromFile(Params.RawProjectPath), SC.StageTargetPlatform.IniPlatformType);
+				if (PlatformEngineConfig != null)
 				{
-					SC.StageFile(StagedFileType.UFS, ProjectCertFile);
+					PlatformEngineConfig.GetBool("/Script/Engine.NetworkSettings", "n.VerifyPeer", out bStageSSLCertificates);
 				}
-				else
+
+				if (bStageSSLCertificates)
 				{
-					// if the game had any files to be staged, then we don't need to stage the engine one - it will just added hundreds of kb of data that is never used
-					FileReference EngineCertFile = FileReference.Combine(SC.EngineRoot, "Content", "Certificates", "ThirdParty", "cacert.pem");
-					if(FileReference.Exists(EngineCertFile))
+					// Game's SSL certs
+					FileReference ProjectCertFile = FileReference.Combine(SC.ProjectRoot, "Content", "Certificates", "cacert.pem");
+					if (FileReference.Exists(ProjectCertFile))
+					{
+						SC.StageFile(StagedFileType.UFS, ProjectCertFile);
+					}
+					else
+					{
+						// if the game had any files to be staged, then we don't need to stage the engine one - it will just added hundreds of kb of data that is never used
+						FileReference EngineCertFile = FileReference.Combine(SC.EngineRoot, "Content", "Certificates", "ThirdParty", "cacert.pem");
+						if (FileReference.Exists(EngineCertFile))
+						{
+							SC.StageFile(StagedFileType.UFS, EngineCertFile);
+						}
+					}
+
+					// now stage any other game certs besides cacert
+					DirectoryReference CertificatesDir = DirectoryReference.Combine(SC.ProjectRoot, "Certificates");
+					if (DirectoryReference.Exists(CertificatesDir))
 					{
-						SC.StageFile(StagedFileType.UFS, EngineCertFile);
+						SC.StageFiles(StagedFileType.UFS, CertificatesDir, "*.pem", StageFilesSearch.AllDirectories);
 					}
 				}
 
-				// now stage any other game certs besides cacert
-				DirectoryReference CertificatesDir = DirectoryReference.Combine(SC.ProjectRoot, "Certificates");
-				if(DirectoryReference.Exists(CertificatesDir))
+				// Generate a plugin manifest if we're using a pak file and not creating a mod. Mods can be enumerated independently by users copying them into the Mods directory.
+				if (Params.UsePak(SC.StageTargetPlatform))
 				{
-					SC.StageFiles(StagedFileType.UFS, CertificatesDir, "*.pem", StageFilesSearch.AllDirectories);
+					if (Params.HasDLCName)
+					{
+						CreatePluginManifest(SC, SC.FilesToStage.NonUFSFiles, StagedFileType.NonUFS, Params.DLCFile.GetFileNameWithoutExtension());
+					}
+					else
+					{
+						CreatePluginManifest(SC, SC.FilesToStage.UFSFiles, StagedFileType.UFS, Params.ShortProjectName);
+					}
 				}
 			}
-
-			// Generate a plugin manifest if we're using a pak file and not creating a mod. Mods can be enumerated independently by users copying them into the Mods directory.
-			if (Params.UsePak(SC.StageTargetPlatform))
+			else
 			{
-				if (Params.HasDLCName)
+				if (PlatformGameConfig != null)
 				{
-					CreatePluginManifest(SC, SC.FilesToStage.NonUFSFiles, StagedFileType.NonUFS, Params.DLCFile.GetFileNameWithoutExtension());
-				}
-				else
-				{
-					CreatePluginManifest(SC, SC.FilesToStage.UFSFiles, StagedFileType.UFS, Params.ShortProjectName);
+					List<string> ExtraNonUFSDirs;
+					if (PlatformGameConfig.GetArray("/Script/UnrealEd.ProjectPackagingSettings", "DirectoriesToAlwaysStageAsNonUFS", out ExtraNonUFSDirs))
+					{
+						// Each string has the format '(Path="TheDirToStage")'
+						foreach (var PathStr in ExtraNonUFSDirs)
+						{
+							var PathParts = PathStr.Split('"');
+							if (PathParts.Length == 3)
+							{
+								var RelativePath = PathParts[1];
+								SC.StageFiles(StagedFileType.NonUFS, DirectoryReference.Combine(ProjectContentRoot, RelativePath), StageFilesSearch.AllDirectories);
+							}
+							else if (PathParts.Length == 1)
+							{
+								var RelativePath = PathParts[0];
+								SC.StageFiles(StagedFileType.UFS, DirectoryReference.Combine(ProjectContentRoot, RelativePath), StageFilesSearch.AllDirectories, StagedDirectoryReference.Combine(StageContentRoot, RelativePath));
+							}
+						}
+					}
 				}
 			}
 		}
-        else
-        {
-            if (PlatformGameConfig != null)
-            {
-                List<string> ExtraNonUFSDirs;
-                if (PlatformGameConfig.GetArray("/Script/UnrealEd.ProjectPackagingSettings", "DirectoriesToAlwaysStageAsNonUFS", out ExtraNonUFSDirs))
-                {
-                    // Each string has the format '(Path="TheDirToStage")'
-                    foreach (var PathStr in ExtraNonUFSDirs)
-                    {
-                        var PathParts = PathStr.Split('"');
-                        if (PathParts.Length == 3)
-                        {
-                            var RelativePath = PathParts[1];
-                            SC.StageFiles(StagedFileType.NonUFS, DirectoryReference.Combine(ProjectContentRoot, RelativePath), StageFilesSearch.AllDirectories);
-                        }
-                        else if (PathParts.Length == 1)
-                        {
-                            var RelativePath = PathParts[0];
-                            SC.StageFiles(StagedFileType.UFS, DirectoryReference.Combine(ProjectContentRoot, RelativePath), StageFilesSearch.AllDirectories, StagedDirectoryReference.Combine(StageContentRoot, RelativePath));
-                        }
-                    }
-                }
-            }
-        }
 
 		// Apply all the directory mappings
 		SC.FilesToStage.UFSFiles = SC.FilesToStage.UFSFiles.ToDictionary(x => ApplyDirectoryRemap(SC, x.Key), x => x.Value);
@@ -953,11 +1037,22 @@ public partial class Project : CommandUtils
             Log("Generating patch required a based on release version flag");
         }
 
-		string PostFix = "";
-		string OutputFilename = PakName + "-" + SC.FinalCookPlatform;
+		string OutputFilename = PakName + "-" + "[Platform]";
 		string OutputFilenameExtension = ".pak";
-		if (bShouldGeneratePatch)
+		string ContentPakRelativePath = Path.Combine("Content", "Paks");
+		if (Params.HasDLCName)
+		{
+			string PostFix = "";
+			ContentPakRelativePath = Path.Combine(Params.DLCFile.Directory.MakeRelativeTo(SC.ProjectRoot), "Content", "Paks", "[Platform]");
+			if (Params.DLCActLikePatch)
+			{
+				PostFix += "_D_P";
+				OutputFilename = Params.DLCFile.GetFileNameWithoutExtension() + PostFix;
+			}
+		}
+		else if (bShouldGeneratePatch)
         {
+			string PostFix = "";
 			PostFix += "_P";
 			int TargetPatchIndex = 0;
 			string ExistingPatchSearchPath = SC.StageTargetPlatform.GetReleasePakFilePath(SC, Params, null);
@@ -991,6 +1086,7 @@ public partial class Project : CommandUtils
 		}
 		else if (Params.IterateSharedCookedBuild)
 		{
+			string PostFix = "";
 			// shared cooked builds will produce a patch
 			// then be combined with the shared cooked build
 			PostFix += "_S_P";
@@ -998,14 +1094,8 @@ public partial class Project : CommandUtils
 		}
 
 		StagedFileReference OutputRelativeLocation;
-		if(Params.HasDLCName)
-		{
-			OutputRelativeLocation = StagedFileReference.Combine(SC.RelativeProjectRootForStage, Params.DLCFile.Directory.MakeRelativeTo(SC.ProjectRoot), "Content", "Paks", SC.FinalCookPlatform, Params.DLCFile.GetFileNameWithoutExtension() + ".pak");
-		}
-		else
-		{
-			OutputRelativeLocation = StagedFileReference.Combine(SC.RelativeProjectRootForStage, "Content", "Paks", OutputFilename + OutputFilenameExtension);
-		}
+		OutputRelativeLocation = StagedFileReference.Combine(SC.RelativeProjectRootForStage, ContentPakRelativePath.Replace("[Platform]", SC.FinalCookPlatform), OutputFilename.Replace("[Platform]", SC.FinalCookPlatform) + OutputFilenameExtension);
+
 		if (SC.StageTargetPlatform.DeployLowerCaseFilenames())
 		{
 			OutputRelativeLocation = OutputRelativeLocation.ToLowerInvariant();
@@ -1035,7 +1125,8 @@ public partial class Project : CommandUtils
 		{
 			// Check to see if we have an existing pak file we can use
 
-			StagedFileReference SourceOutputRelativeLocation = StagedFileReference.Combine(SC.RelativeProjectRootForStage, "Content/Paks/", PakName + "-" + SC.CookPlatform + PostFix + ".pak");
+			//StagedFileReference SourceOutputRelativeLocation = StagedFileReference.Combine(SC.RelativeProjectRootForStage, "Content/Paks/", PakName + "-" + SC.CookPlatform + PostFix + ".pak");
+			StagedFileReference SourceOutputRelativeLocation = StagedFileReference.Combine(SC.RelativeProjectRootForStage, ContentPakRelativePath.Replace("[Platform]", SC.CookPlatform), OutputFilename.Replace("[Platform]", SC.FinalCookPlatform) + OutputFilenameExtension);
 			if (SC.CookSourcePlatform.DeployLowerCaseFilenames())
 			{
 				SourceOutputRelativeLocation = SourceOutputRelativeLocation.ToLowerInvariant();
@@ -1145,7 +1236,7 @@ public partial class Project : CommandUtils
 				string VersionString = Params.ChunkInstallVersionString;
 				string ChunkInstallBasePath = CombinePaths(Params.ChunkInstallDirectory, SC.FinalCookPlatform);
 				string RawDataPath = CombinePaths(ChunkInstallBasePath, VersionString, PakName);
-				string RawDataPakPath = CombinePaths(RawDataPath, PakName + "-" + SC.FinalCookPlatform + PostFix + ".pak");
+				string RawDataPakPath = CombinePaths(RawDataPath, OutputFilename.Replace("[Platform]", SC.FinalCookPlatform) + OutputFilenameExtension); // PakName + "-" + SC.FinalCookPlatform + PostFix + ".pak");
 				//copy the pak chunk to the raw data folder
 				if (InternalUtils.SafeFileExists(RawDataPakPath, true))
 				{
diff --git a/Engine/Source/Programs/UnrealBuildTool/Configuration/TargetRules.cs b/Engine/Source/Programs/UnrealBuildTool/Configuration/TargetRules.cs
--- a/Engine/Source/Programs/UnrealBuildTool/Configuration/TargetRules.cs
+++ b/Engine/Source/Programs/UnrealBuildTool/Configuration/TargetRules.cs
@@ -213,9 +213,9 @@ namespace UnrealBuildTool
 				set { Inner.bEnableOSX109Support = value; }
 			}
 
-			#if !__MonoCS__
-			#pragma warning restore CS1591
-			#endif
+#if !__MonoCS__
+#pragma warning restore CS1591
+#endif
 			#endregion
 		}
 
@@ -738,7 +738,7 @@ namespace UnrealBuildTool
 		/// <summary>
 		/// Whether to allow the use of link time code generation (LTCG).
 		/// </summary>
-		[CommandLine("-NoLTCG", Value = "false")]
+		[CommandLine("-LTCG")]
 		[XmlConfigFile(Category = "BuildConfiguration")]
 		public bool bAllowLTCG = false;
 
diff --git a/Engine/Source/Programs/UnrealBuildTool/Platform/Mac/MacToolChain.cs b/Engine/Source/Programs/UnrealBuildTool/Platform/Mac/MacToolChain.cs
--- a/Engine/Source/Programs/UnrealBuildTool/Platform/Mac/MacToolChain.cs
+++ b/Engine/Source/Programs/UnrealBuildTool/Platform/Mac/MacToolChain.cs
@@ -209,9 +209,10 @@ namespace UnrealBuildTool
 				Result += " -Wundef" + (CompileEnvironment.bUndefinedIdentifierWarningsAsErrors ? "" : " -Wno-error=undef");
 			}
 
-			// @todo: Remove these two when the code is fixed and they're no longer needed
+			// @todo: Remove these three when the code is fixed and they're no longer needed
 			Result += " -Wno-logical-op-parentheses";
 			Result += " -Wno-unknown-pragmas";
+			Result += " -Wno-undefined-var-template";
 
 			Result += " -c";
 
diff --git a/Engine/Source/Programs/UnrealBuildTool/System/XGE.cs b/Engine/Source/Programs/UnrealBuildTool/System/XGE.cs
--- a/Engine/Source/Programs/UnrealBuildTool/System/XGE.cs
+++ b/Engine/Source/Programs/UnrealBuildTool/System/XGE.cs
@@ -33,7 +33,7 @@ namespace UnrealBuildTool
 		/// When enabled, XGE will stop compiling targets after a compile error occurs.  Recommended, as it saves computing resources for others.
 		/// </summary>
 		[XmlConfigFile(Category = "BuildConfiguration")]
-		bool bStopXGECompilationAfterErrors = true;
+		bool bStopXGECompilationAfterErrors = false;
 
 		private const string ProgressMarkupPrefix = "@action";
 
diff --git a/Engine/Source/Programs/UnrealHitchParser/App.config b/Engine/Source/Programs/UnrealHitchParser/App.config
new file mode 100644
--- /dev/null
+++ b/Engine/Source/Programs/UnrealHitchParser/App.config
@@ -0,0 +1,6 @@
+﻿<?xml version="1.0" encoding="utf-8" ?>
+<configuration>
+    <startup> 
+        <supportedRuntime version="v4.0" sku=".NETFramework,Version=v4.5.2" />
+    </startup>
+</configuration>
\ No newline at end of file
diff --git a/Engine/Source/Runtime/AIModule/Classes/BehaviorTree/BehaviorTreeTypes.h b/Engine/Source/Runtime/AIModule/Classes/BehaviorTree/BehaviorTreeTypes.h
--- a/Engine/Source/Runtime/AIModule/Classes/BehaviorTree/BehaviorTreeTypes.h
+++ b/Engine/Source/Runtime/AIModule/Classes/BehaviorTree/BehaviorTreeTypes.h
@@ -520,7 +520,11 @@ public:
 UCLASS(Abstract)
 class AIMODULE_API UBehaviorTreeTypes : public UObject
 {
-	GENERATED_UCLASS_BODY()
+	GENERATED_BODY()
+
+	static FString BTLoggingContext;
+
+public:
 
 	static FString DescribeNodeHelper(const UBTNode* Node);
 
@@ -532,4 +536,10 @@ class AIMODULE_API UBehaviorTreeTypes : public UObject
 
 	/** returns short name of object's class (BTTaskNode_Wait -> Wait) */
 	static FString GetShortTypeName(const UObject* Ob);
+	
+	static FString GetBTLoggingContext() { return BTLoggingContext; }
+	
+	// @param NewBTLoggingContext the object which name's will be added to some of the BT logging
+	// 	pass nullptr to clear
+	static void SetBTLoggingContext(const UBTNode* NewBTLoggingContext);
 };
diff --git a/Engine/Source/Runtime/AIModule/Classes/EnvironmentQuery/EQSRenderingComponent.h b/Engine/Source/Runtime/AIModule/Classes/EnvironmentQuery/EQSRenderingComponent.h
--- a/Engine/Source/Runtime/AIModule/Classes/EnvironmentQuery/EQSRenderingComponent.h
+++ b/Engine/Source/Runtime/AIModule/Classes/EnvironmentQuery/EQSRenderingComponent.h
@@ -16,10 +16,12 @@ class APlayerController;
 class IEQSQueryResultSourceInterface;
 class UCanvas;
 
-class AIMODULE_API FEQSSceneProxy : public FDebugRenderSceneProxy
+class AIMODULE_API FEQSSceneProxy final : public FDebugRenderSceneProxy
 {
 	friend class FEQSRenderingDebugDrawDelegateHelper;
 public:
+	SIZE_T GetTypeHash() const override;
+
 	DEPRECATED(4.14, "This FEQSSceneProxy constructor version is deprecated. Please use the one taking UPrimitiveComponent&")
 	FEQSSceneProxy(const UPrimitiveComponent* InComponent, const FString& ViewFlagName = TEXT("DebugAI"));
 	DEPRECATED(4.14, "This FEQSSceneProxy constructor version is deprecated. Please use the one taking UPrimitiveComponent&")
diff --git a/Engine/Source/Runtime/AIModule/Classes/EnvironmentQuery/EQSTestingPawn.h b/Engine/Source/Runtime/AIModule/Classes/EnvironmentQuery/EQSTestingPawn.h
--- a/Engine/Source/Runtime/AIModule/Classes/EnvironmentQuery/EQSTestingPawn.h
+++ b/Engine/Source/Runtime/AIModule/Classes/EnvironmentQuery/EQSTestingPawn.h
@@ -64,7 +64,9 @@ class AIMODULE_API AEQSTestingPawn : public ACharacter, public IEQSQueryResultSo
 
 	UPROPERTY(Category=EQS, EditAnywhere)
 	TEnumAsByte<EEnvQueryRunMode::Type> QueryingMode;
-	
+
+	UPROPERTY(Category = EQS, EditAnywhere)
+	FNavAgentProperties NavAgentProperties;
 
 #if WITH_EDITORONLY_DATA
 private:
@@ -97,6 +99,10 @@ public:
 	virtual float GetHighlightRangePct() const override;
 	// IEQSQueryResultSourceInterface end
 
+	// INavAgentInterface begin
+	virtual const FNavAgentProperties& GetNavAgentPropertiesRef() const override;
+	// INavAgentInterface end
+
 	void RunEQSQuery();
 
 protected:	
diff --git a/Engine/Source/Runtime/AIModule/Private/BehaviorTree/BehaviorTreeTypes.cpp b/Engine/Source/Runtime/AIModule/Private/BehaviorTree/BehaviorTreeTypes.cpp
--- a/Engine/Source/Runtime/AIModule/Private/BehaviorTree/BehaviorTreeTypes.cpp
+++ b/Engine/Source/Runtime/AIModule/Private/BehaviorTree/BehaviorTreeTypes.cpp
@@ -302,6 +302,12 @@ void FBlackboardKeySelector::ResolveSelectedKey(const UBlackboardData& Blackboar
 
 		SelectedKeyID = BlackboardAsset.GetKeyID(SelectedKeyName);
 		SelectedKeyType = BlackboardAsset.GetKeyType(SelectedKeyID);
+		UE_CLOG(IsSet() == false, LogBehaviorTree, Warning
+			, TEXT("%s> Failed to find key \'%s\' in BB asset %s. BB Key Selector will be set to \'Invalid\'")
+			, *UBehaviorTreeTypes::GetBTLoggingContext()
+			, *SelectedKeyName.ToString()
+			, *BlackboardAsset.GetFullName()
+		);
 	}
 }
 
@@ -415,9 +421,7 @@ void FBlackboardKeySelector::AddNameFilter(UObject* Owner, FName PropertyName)
 //----------------------------------------------------------------------//
 // UBehaviorTreeTypes
 //----------------------------------------------------------------------//
-UBehaviorTreeTypes::UBehaviorTreeTypes(const FObjectInitializer& ObjectInitializer) : Super(ObjectInitializer)
-{
-}
+FString UBehaviorTreeTypes::BTLoggingContext;
 
 FString UBehaviorTreeTypes::DescribeNodeResult(EBTNodeResult::Type NodeResult)
 {
@@ -471,6 +475,13 @@ FString UBehaviorTreeTypes::GetShortTypeName(const UObject* Ob)
 	return TypeDesc;
 }
 
+void UBehaviorTreeTypes::SetBTLoggingContext(const UBTNode* NewBTLoggingContext)
+{
+	BTLoggingContext = NewBTLoggingContext 
+		? FString::Printf(TEXT("%s[%d]"), *NewBTLoggingContext->GetNodeName(), NewBTLoggingContext->GetExecutionIndex())
+		: TEXT("");
+}
+
 //----------------------------------------------------------------------//
 // DEPRECATED
 //----------------------------------------------------------------------//
diff --git a/Engine/Source/Runtime/AIModule/Private/EnvironmentQuery/EQSRenderingComponent.cpp b/Engine/Source/Runtime/AIModule/Private/EnvironmentQuery/EQSRenderingComponent.cpp
--- a/Engine/Source/Runtime/AIModule/Private/EnvironmentQuery/EQSRenderingComponent.cpp
+++ b/Engine/Source/Runtime/AIModule/Private/EnvironmentQuery/EQSRenderingComponent.cpp
@@ -26,6 +26,12 @@ namespace FEQSRenderingHelper
 //----------------------------------------------------------------------//
 const FVector FEQSSceneProxy::ItemDrawRadius(30,30,30);
 
+SIZE_T FEQSSceneProxy::GetTypeHash() const
+{
+	static size_t UniquePointer;
+	return reinterpret_cast<size_t>(&UniquePointer);
+}
+
 FEQSSceneProxy::FEQSSceneProxy(const UPrimitiveComponent& InComponent, const FString& InViewFlagName, const TArray<FSphere>& InSpheres, const TArray<FText3d>& InTexts)
 	: FDebugRenderSceneProxy(&InComponent)
 	, ActorOwner(nullptr)
diff --git a/Engine/Source/Runtime/AIModule/Private/EnvironmentQuery/EQSTestingPawn.cpp b/Engine/Source/Runtime/AIModule/Private/EnvironmentQuery/EQSTestingPawn.cpp
--- a/Engine/Source/Runtime/AIModule/Private/EnvironmentQuery/EQSTestingPawn.cpp
+++ b/Engine/Source/Runtime/AIModule/Private/EnvironmentQuery/EQSTestingPawn.cpp
@@ -38,6 +38,8 @@ AEQSTestingPawn::AEQSTestingPawn(const FObjectInitializer& ObjectInitializer)
 	static FName CollisionProfileName(TEXT("NoCollision"));
 	GetCapsuleComponent()->SetCollisionProfileName(CollisionProfileName);
 
+	NavAgentProperties = FNavAgentProperties::DefaultProperties;
+
 #if WITH_EDITORONLY_DATA
 	EdRenderComp = CreateEditorOnlyDefaultSubobject<UEQSRenderingComponent>(TEXT("EQSRender"));
 	if (EdRenderComp)
@@ -105,6 +107,11 @@ float AEQSTestingPawn::GetHighlightRangePct() const
 	return (HighlightMode == EEnvQueryHightlightMode::Best25Pct) ? 0.75f : (HighlightMode == EEnvQueryHightlightMode::Best5Pct) ? 0.95f : 1.0f;
 }
 
+const FNavAgentProperties& AEQSTestingPawn::GetNavAgentPropertiesRef() const 
+{
+	return NavAgentProperties;
+}
+
 void AEQSTestingPawn::OnEditorSelectionChanged(UObject* NewSelection)
 {
 	/* This is totally busted and should not depend on gameplay debugger.
diff --git a/Engine/Source/Runtime/AIModule/Private/GameplayDebugger/GameplayDebuggerCategory_AI.cpp b/Engine/Source/Runtime/AIModule/Private/GameplayDebugger/GameplayDebuggerCategory_AI.cpp
--- a/Engine/Source/Runtime/AIModule/Private/GameplayDebugger/GameplayDebuggerCategory_AI.cpp
+++ b/Engine/Source/Runtime/AIModule/Private/GameplayDebugger/GameplayDebuggerCategory_AI.cpp
@@ -331,9 +331,15 @@ void FGameplayDebuggerCategory_AI::DrawData(APlayerController* OwnerPC, FGamepla
 
 FDebugRenderSceneProxy* FGameplayDebuggerCategory_AI::CreateDebugSceneProxy(const UPrimitiveComponent* InComponent, FDebugDrawDelegateHelper*& OutDelegateHelper)
 {
-	class FPathDebugRenderSceneProxy : public FDebugRenderSceneProxy
+	class FPathDebugRenderSceneProxy final : public FDebugRenderSceneProxy
 	{
 	public:
+		SIZE_T GetTypeHash() const override
+		{
+			static size_t UniquePointer;
+			return reinterpret_cast<size_t>(&UniquePointer);
+		}
+
 		FPathDebugRenderSceneProxy(const UPrimitiveComponent* InPrimitiveCComponent, const FString &InViewFlagName)
 			: FDebugRenderSceneProxy(InPrimitiveCComponent)
 		{
diff --git a/Engine/Source/Runtime/AIModule/Private/GameplayDebugger/GameplayDebuggerCategory_NavLocalGrid.cpp b/Engine/Source/Runtime/AIModule/Private/GameplayDebugger/GameplayDebuggerCategory_NavLocalGrid.cpp
--- a/Engine/Source/Runtime/AIModule/Private/GameplayDebugger/GameplayDebuggerCategory_NavLocalGrid.cpp
+++ b/Engine/Source/Runtime/AIModule/Private/GameplayDebugger/GameplayDebuggerCategory_NavLocalGrid.cpp
@@ -15,9 +15,15 @@
 //////////////////////////////////////////////////////////////////////////
 // Scene proxy
 
-class FNavLocalGridSceneProxy : public FDebugRenderSceneProxy
+class FNavLocalGridSceneProxy final : public FDebugRenderSceneProxy
 {
 public:
+	SIZE_T GetTypeHash() const override
+	{
+		static size_t UniquePointer;
+		return reinterpret_cast<size_t>(&UniquePointer);
+	}
+
 	FNavLocalGridSceneProxy(const UPrimitiveComponent* InComponent,
 		const FGameplayDebuggerCategory_NavLocalGrid::FRepData& RepData,
 		const FGameplayDebuggerCategory_NavLocalGrid::FRepAgentData& AgentData) 
diff --git a/Engine/Source/Runtime/AIModule/Private/Perception/AISense_Sight.cpp b/Engine/Source/Runtime/AIModule/Private/Perception/AISense_Sight.cpp
--- a/Engine/Source/Runtime/AIModule/Private/Perception/AISense_Sight.cpp
+++ b/Engine/Source/Runtime/AIModule/Private/Perception/AISense_Sight.cpp
@@ -354,11 +354,12 @@ void UAISense_Sight::RegisterSource(AActor& SourceActor)
 void UAISense_Sight::UnregisterSource(AActor& SourceActor)
 {
 	const FAISightTarget::FTargetId AsTargetId = SourceActor.GetUniqueID();
-	FAISightTarget* AsTarget = ObservedTargets.Find(AsTargetId);
+	FAISightTarget AsTarget;
 	
-	if (AsTarget != nullptr && SightQueryQueue.Num() > 0)
+	if (ObservedTargets.RemoveAndCopyValue(AsTargetId, AsTarget)
+		&& SightQueryQueue.Num() > 0)
 	{
-		AActor* TargetActor = AsTarget->Target.Get();
+		AActor* TargetActor = AsTarget.Target.Get();
 
 		if (TargetActor)
 		{
diff --git a/Engine/Source/Runtime/Analytics/Analytics/Public/AnalyticsEventAttribute.h b/Engine/Source/Runtime/Analytics/Analytics/Public/AnalyticsEventAttribute.h
--- a/Engine/Source/Runtime/Analytics/Analytics/Public/AnalyticsEventAttribute.h
+++ b/Engine/Source/Runtime/Analytics/Analytics/Public/AnalyticsEventAttribute.h
@@ -56,6 +56,8 @@ struct FAnalyticsEventAttribute
 		case AttrTypeEnum::JsonFragment:
 			return AttrValueString;
 		case AttrTypeEnum::Number:
+			if (AttrValueNumber - FMath::TruncToFloat(AttrValueNumber) == 0.0)
+				return Lex::ToSanitizedString((int64)AttrValueNumber);
 			return Lex::ToSanitizedString(AttrValueNumber);
 		case AttrTypeEnum::Boolean:
 			return Lex::ToString(AttrValueBool);
diff --git a/Engine/Source/Runtime/Analytics/AnalyticsET/Private/IAnalyticsProviderET.cpp b/Engine/Source/Runtime/Analytics/AnalyticsET/Private/IAnalyticsProviderET.cpp
--- a/Engine/Source/Runtime/Analytics/AnalyticsET/Private/IAnalyticsProviderET.cpp
+++ b/Engine/Source/Runtime/Analytics/AnalyticsET/Private/IAnalyticsProviderET.cpp
@@ -291,43 +291,6 @@ private:
 	void EventRequestComplete(FHttpRequestPtr HttpRequest, FHttpResponsePtr HttpResponse, bool bSucceeded, TSharedPtr< TArray<FAnalyticsEventEntry> > FlushedEvents);
 };
 
-class FAnalyticsProviderETNULL :
-	public IAnalyticsProviderET,
-	public TSharedFromThis<FAnalyticsProviderETNULL>
-{
-public:
-	FAnalyticsProviderETNULL(const FAnalyticsET::Config& ConfigValues) {};
-
-	// IAnalyticsProvider
-
-	virtual bool StartSession(const TArray<FAnalyticsEventAttribute>& Attributes) override { return true; }
-	virtual bool StartSession(TArray<FAnalyticsEventAttribute>&& Attributes) override { return true; }
-	virtual void EndSession() override { }
-	virtual void FlushEvents() override { }
-
-	virtual void SetAppID(FString&& AppID) override { APIKey = MoveTemp(AppID); }
-	virtual const FString& GetAppID() const override { return APIKey; }
-	virtual void SetUserID(const FString& InUserID) override { UserID = InUserID; }
-	virtual FString GetUserID() const override { return UserID; }
-
-	virtual FString GetSessionID() const override { return SessionID; }
-	virtual bool SetSessionID(const FString& InSessionID) override { SessionID = InSessionID; return true; }
-
-	virtual void RecordEvent(const FString& EventName, const TArray<FAnalyticsEventAttribute>& Attributes) override {}
-	virtual void RecordEvent(FString EventName, TArray<FAnalyticsEventAttribute>&& Attributes) override {}
-	virtual void RecordEventJson(FString EventName, TArray<FAnalyticsEventAttribute>&& AttributesJson) override {}
-	virtual void SetDefaultEventAttributes(TArray<FAnalyticsEventAttribute>&& Attributes) override {}
-	virtual void SetEventCallback(const OnEventRecorded& Callback) override {}
-
-	virtual ~FAnalyticsProviderETNULL() {};
-
-	FString GetAPIKey() const { return APIKey; }
-
-	FString APIKey;
-	FString UserID;
-	FString SessionID;
-};
-
 TSharedPtr<IAnalyticsProviderET> FAnalyticsET::CreateAnalyticsProvider(const Config& ConfigValues) const
 {
 	// If we didn't have a proper APIKey, return NULL
@@ -336,12 +299,7 @@ TSharedPtr<IAnalyticsProviderET> FAnalyticsET::CreateAnalyticsProvider(const Con
 		UE_LOG(LogAnalytics, Warning, TEXT("CreateAnalyticsProvider config not contain required parameter %s"), *Config::GetKeyNameForAPIKey());
 		return NULL;
 	}
-	//@todo sz
-#if 0
-	return TSharedPtr<IAnalyticsProviderET>(new FAnalyticsProviderETNULL(ConfigValues));
-#else
 	return TSharedPtr<IAnalyticsProviderET>(new FAnalyticsProviderET(ConfigValues));
-#endif
 }
 
 /**
@@ -582,10 +540,10 @@ void FAnalyticsProviderET::FlushEvents()
 							JsonWriter->WriteValue(Attr.AttrName, Attr.AttrValueString);
 							break;
 						case FAnalyticsEventAttribute::AttrTypeEnum::Number:
-							JsonWriter->WriteValue(Attr.AttrName, Attr.AttrValueNumber);
+							JsonWriter->WriteValue(Attr.AttrName, Attr.ToString());
 							break;
 						case FAnalyticsEventAttribute::AttrTypeEnum::Boolean:
-							JsonWriter->WriteValue(Attr.AttrName, Attr.AttrValueBool);
+							JsonWriter->WriteValue(Attr.AttrName, Attr.ToString());
 							break;
 						case FAnalyticsEventAttribute::AttrTypeEnum::JsonFragment:
 							JsonWriter->WriteRawJSONValue(Attr.AttrName, Attr.AttrValueString);
@@ -603,17 +561,17 @@ void FAnalyticsProviderET::FlushEvents()
 								JsonWriter->WriteValue(Attr.AttrName, Attr.AttrValueString);
 								break;
 							case FAnalyticsEventAttribute::AttrTypeEnum::Number:
-								JsonWriter->WriteValue(Attr.AttrName, Attr.AttrValueNumber);
+								JsonWriter->WriteValue(Attr.AttrName, Attr.ToString());
 								break;
 							case FAnalyticsEventAttribute::AttrTypeEnum::Boolean:
-								JsonWriter->WriteValue(Attr.AttrName, Attr.AttrValueBool);
+								JsonWriter->WriteValue(Attr.AttrName, Attr.ToString());
 								break;
 							case FAnalyticsEventAttribute::AttrTypeEnum::JsonFragment:
 								JsonWriter->WriteRawJSONValue(Attr.AttrName, Attr.AttrValueString);
 								break;
-							}
 						}
 					}
+					}
 					else
 					{
 						for (const FAnalyticsEventAttribute& Attr : Entry.Attributes)
diff --git a/Engine/Source/Runtime/Analytics/AnalyticsET/Public/AnalyticsET.h b/Engine/Source/Runtime/Analytics/AnalyticsET/Public/AnalyticsET.h
--- a/Engine/Source/Runtime/Analytics/AnalyticsET/Public/AnalyticsET.h
+++ b/Engine/Source/Runtime/Analytics/AnalyticsET/Public/AnalyticsET.h
@@ -9,14 +9,10 @@
 class IAnalyticsProvider;
 class IAnalyticsProviderET;
 
-
-
-class IAnalyticsProviderET;
-
 /**
  *  Public implementation of EpicGames.MCP.AnalyticsProvider
  */
-class FAnalyticsET : public IAnalyticsProviderModule
+class ANALYTICSET_API FAnalyticsET : public IAnalyticsProviderModule
 {
 	//--------------------------------------------------------------------------
 	// Module functionality
diff --git a/Engine/Source/Runtime/AnimGraphRuntime/Private/BoneControllers/AnimNode_SkeletalControlBase.cpp b/Engine/Source/Runtime/AnimGraphRuntime/Private/BoneControllers/AnimNode_SkeletalControlBase.cpp
--- a/Engine/Source/Runtime/AnimGraphRuntime/Private/BoneControllers/AnimNode_SkeletalControlBase.cpp
+++ b/Engine/Source/Runtime/AnimGraphRuntime/Private/BoneControllers/AnimNode_SkeletalControlBase.cpp
@@ -65,9 +65,14 @@ void FAnimNode_SkeletalControlBase::UpdateInternal(const FAnimationUpdateContext
 {
 }
 
-void FAnimNode_SkeletalControlBase::Update_AnyThread(const FAnimationUpdateContext& Context)
+void FAnimNode_SkeletalControlBase::UpdateComponentPose_AnyThread(const FAnimationUpdateContext& Context)
 {
 	ComponentPose.Update(Context);
+}
+
+void FAnimNode_SkeletalControlBase::Update_AnyThread(const FAnimationUpdateContext& Context)
+{
+	UpdateComponentPose_AnyThread(Context);
 
 	ActualAlpha = 0.f;
 	if (IsLODEnabled(Context.AnimInstanceProxy, LODThreshold))
@@ -96,27 +101,30 @@ bool ContainsNaN(const TArray<FBoneTransform> & BoneTransforms)
 	return false;
 }
 
+void FAnimNode_SkeletalControlBase::EvaluateComponentPose_AnyThread(FComponentSpacePoseContext& Output)
+{
+	// Evaluate the input
+	ComponentPose.EvaluateComponentSpace(Output);
+}
+
 void FAnimNode_SkeletalControlBase::EvaluateComponentSpaceInternal(FComponentSpacePoseContext& Context)
 {
 }
 
 void FAnimNode_SkeletalControlBase::EvaluateComponentSpace_AnyThread(FComponentSpacePoseContext& Output)
 {
-	// Evaluate the input
-	ComponentPose.EvaluateComponentSpace(Output);
+	EvaluateComponentPose_AnyThread(Output);
+
+#if WITH_EDITORONLY_DATA
+	// save current pose before applying skeletal control to compute the exact gizmo location in AnimGraphNode
+	ForwardedPose.CopyPose(Output.Pose);
+#endif // #if WITH_EDITORONLY_DATA
 
 	// Apply the skeletal control if it's valid
 	if (FAnimWeight::IsRelevant(ActualAlpha) && IsValidToEvaluate(Output.AnimInstanceProxy->GetSkeleton(), Output.AnimInstanceProxy->GetRequiredBones()))
 	{
 		EvaluateComponentSpaceInternal(Output);
 
-#if WITH_EDITORONLY_DATA
-		// save current pose before applying skeletal control to compute the exact gizmo location in AnimGraphNode
-		ForwardedPose.CopyPose(Output.Pose);
-#endif // #if WITH_EDITORONLY_DATA
-
-		USkeletalMeshComponent* Component = Output.AnimInstanceProxy->GetSkelMeshComponent();
-
 		BoneTransforms.Reset(BoneTransforms.Num());
 		EvaluateSkeletalControl_AnyThread(Output, BoneTransforms);
 
diff --git a/Engine/Source/Runtime/AnimGraphRuntime/Public/BoneControllers/AnimNode_SkeletalControlBase.h b/Engine/Source/Runtime/AnimGraphRuntime/Public/BoneControllers/AnimNode_SkeletalControlBase.h
--- a/Engine/Source/Runtime/AnimGraphRuntime/Public/BoneControllers/AnimNode_SkeletalControlBase.h
+++ b/Engine/Source/Runtime/AnimGraphRuntime/Public/BoneControllers/AnimNode_SkeletalControlBase.h
@@ -334,6 +334,13 @@ protected:
 	// Interface for derived skeletal controls to implement
 	// use this function to update for skeletal control base
 	virtual void UpdateInternal(const FAnimationUpdateContext& Context);
+
+	// Update incoming component pose.
+	virtual void UpdateComponentPose_AnyThread(const FAnimationUpdateContext& Context);
+
+	// Evaluate incoming component pose.
+	virtual void EvaluateComponentPose_AnyThread(FComponentSpacePoseContext& Output);
+
 	// use this function to evaluate for skeletal control base
 	virtual void EvaluateComponentSpaceInternal(FComponentSpacePoseContext& Context);
 	DEPRECATED(4.16, "Please use EvaluateSkeletalControl_AnyThread.")
diff --git a/Engine/Source/Runtime/AssetRegistry/Private/AssetRegistry.cpp b/Engine/Source/Runtime/AssetRegistry/Private/AssetRegistry.cpp
--- a/Engine/Source/Runtime/AssetRegistry/Private/AssetRegistry.cpp
+++ b/Engine/Source/Runtime/AssetRegistry/Private/AssetRegistry.cpp
@@ -107,16 +107,20 @@ UAssetRegistryImpl::UAssetRegistryImpl(const FObjectInitializer& ObjectInitializ
 			// serialize the data with the memory reader (will convert FStrings to FNames, etc)
 			Serialize(SerializedAssetData);
 		}
-		TArray<TSharedRef<IPlugin>> PakPlugins = IPluginManager::Get().GetPluginsWithPakFile();
-		for (TSharedRef<IPlugin> PakPlugin : PakPlugins)
+
+		TArray<TSharedRef<IPlugin>> ContentPlugins = IPluginManager::Get().GetEnabledPlugins();
+		for ( TSharedRef<IPlugin> ContentPlugin : ContentPlugins )
+		{
+			if ( ContentPlugin->CanContainContent() )
 		{
-			if (FFileHelper::LoadFileToArray(SerializedAssetData, *(PakPlugin->GetBaseDir() / TEXT("AssetRegistry.bin"))))
+				if (FFileHelper::LoadFileToArray(SerializedAssetData, *(ContentPlugin->GetBaseDir() / TEXT("AssetRegistry.bin"))))
 			{
 				SerializedAssetData.Seek(0);
 				Serialize(SerializedAssetData);
 			}
 		}
 	}
+	}
 
 	// Report startup time. This does not include DirectoryWatcher startup time.
 	UE_LOG(LogAssetRegistry, Log, TEXT( "FAssetRegistry took %0.4f seconds to start up" ), FPlatformTime::Seconds() - StartupStartTime );
@@ -1759,7 +1763,7 @@ void UAssetRegistryImpl::DependencyDataGathered(const double TickStartTime, TBac
 			}
 
 			PackageDependencies.Add(SoftPackageName, EAssetRegistryDependencyType::Soft);
-		}
+			}
 
 		for (const TPair<FPackageIndex, TArray<FName>>& SearchableNameList : Result.SearchableNamesMap)
 		{
@@ -2553,4 +2557,4 @@ bool UAssetRegistryImpl::SetPrimaryAssetIdForObjectPath(const FName ObjectPath,
 const FAssetData* UAssetRegistryImpl::GetCachedAssetDataForObjectPath(const FName ObjectPath) const
 {
 	return State.GetAssetByObjectPath(ObjectPath);
-}
+}
\ No newline at end of file
diff --git a/Engine/Source/Runtime/ClothingSystemRuntime/Private/ClothingSimulationNv.cpp b/Engine/Source/Runtime/ClothingSystemRuntime/Private/ClothingSimulationNv.cpp
--- a/Engine/Source/Runtime/ClothingSystemRuntime/Private/ClothingSimulationNv.cpp
+++ b/Engine/Source/Runtime/ClothingSystemRuntime/Private/ClothingSimulationNv.cpp
@@ -678,7 +678,7 @@ void FClothingSimulationNv::DestroyContext(IClothingSimulationContext* InContext
 	delete InContext;
 }
 
-void FClothingSimulationNv::GetSimulationData(TMap<int32, FClothSimulData>& OutData, USkeletalMeshComponent* InOwnerComponent, USkeletalMeshComponent* InOverrideComponent) const
+void FClothingSimulationNv::GetSimulationData(TMap<int32, FClothSimulData>& OutData, USkeletalMeshComponent* InOwnerComponent, USkinnedMeshComponent* InOverrideComponent) const
 {
 	const int32 NumActors = Actors.Num();
 
diff --git a/Engine/Source/Runtime/ClothingSystemRuntime/Public/ClothingSimulationNv.h b/Engine/Source/Runtime/ClothingSystemRuntime/Public/ClothingSimulationNv.h
--- a/Engine/Source/Runtime/ClothingSystemRuntime/Public/ClothingSimulationNv.h
+++ b/Engine/Source/Runtime/ClothingSystemRuntime/Public/ClothingSimulationNv.h
@@ -164,7 +164,7 @@ public:
 
 	virtual void DestroyActors() override;
 	virtual void DestroyContext(IClothingSimulationContext* InContext) override;
-	virtual void GetSimulationData(TMap<int32, FClothSimulData>& OutData, USkeletalMeshComponent* InOwnerComponent, USkeletalMeshComponent* InOverrideComponent) const override;
+	virtual void GetSimulationData(TMap<int32, FClothSimulData>& OutData, USkeletalMeshComponent* InOwnerComponent, USkinnedMeshComponent* InOverrideComponent) const override;
 	virtual void AddExternalCollisions(const FClothCollisionData& InData) override;
 	virtual void ClearExternalCollisions() override;
 	virtual void GetCollisions(FClothCollisionData& OutCollisions, bool bIncludeExternal = true) const override;
diff --git a/Engine/Source/Runtime/ClothingSystemRuntimeInterface/Public/ClothingSimulationInterface.h b/Engine/Source/Runtime/ClothingSystemRuntimeInterface/Public/ClothingSimulationInterface.h
--- a/Engine/Source/Runtime/ClothingSystemRuntimeInterface/Public/ClothingSimulationInterface.h
+++ b/Engine/Source/Runtime/ClothingSystemRuntimeInterface/Public/ClothingSimulationInterface.h
@@ -92,7 +92,7 @@ protected:
 	 * @param InOwnerComponent - the component that owns the simulation
 	 * @param InOverrideComponent - An override component if bound to a master pose component
 	 */
-	virtual void GetSimulationData(TMap<int32, FClothSimulData>& OutData, USkeletalMeshComponent* InOwnerComponent, USkeletalMeshComponent* InOverrideComponent) const = 0;
+	virtual void GetSimulationData(TMap<int32, FClothSimulData>& OutData, USkeletalMeshComponent* InOwnerComponent, USkinnedMeshComponent* InOverrideComponent) const = 0;
 
 	/** Get the bounds of the simulation mesh in local simulation space */
 	virtual FBoxSphereBounds GetBounds(const USkeletalMeshComponent* InOwnerComponent) const = 0;
diff --git a/Engine/Source/Runtime/Core/Private/GenericPlatform/GenericPlatformCrashContext.cpp b/Engine/Source/Runtime/Core/Private/GenericPlatform/GenericPlatformCrashContext.cpp
--- a/Engine/Source/Runtime/Core/Private/GenericPlatform/GenericPlatformCrashContext.cpp
+++ b/Engine/Source/Runtime/Core/Private/GenericPlatform/GenericPlatformCrashContext.cpp
@@ -253,8 +253,8 @@ void FGenericCrashContext::SerializeContentToBuffer()
 	AddCrashProperty( TEXT( "IsPerforceBuild" ), NCachedCrashContextProperties::bIsPerforceBuild );
 	AddCrashProperty( TEXT( "IsSourceDistribution" ), NCachedCrashContextProperties::bIsSourceDistribution );
 	AddCrashProperty( TEXT( "IsEnsure" ), bIsEnsure );
-	AddCrashProperty( TEXT( "IsAssert" ), FDebug::bHasAsserted );
-	AddCrashProperty( TEXT( "CrashType" ), GetCrashTypeString(bIsEnsure, FDebug::bHasAsserted, GIsGPUCrashed) );
+	AddCrashProperty( TEXT( "IsAssert" ), FDebug::HasAsserted() );
+	AddCrashProperty( TEXT( "CrashType" ), GetCrashTypeString(bIsEnsure, FDebug::HasAsserted(), GIsGPUCrashed) );
 
 	AddCrashProperty( TEXT( "SecondsSinceStart" ), NCachedCrashContextProperties::SecondsSinceStart );
 
diff --git a/Engine/Source/Runtime/Core/Private/IOS/IOSPlatformMisc.cpp b/Engine/Source/Runtime/Core/Private/IOS/IOSPlatformMisc.cpp
--- a/Engine/Source/Runtime/Core/Private/IOS/IOSPlatformMisc.cpp
+++ b/Engine/Source/Runtime/Core/Private/IOS/IOSPlatformMisc.cpp
@@ -1242,8 +1242,8 @@ void FIOSCrashContext::GenerateWindowsErrorReport(char const* WERPath, bool bIsE
         WriteLine(ReportFile, TEXT("</Parameter2>"));
         WriteLine(ReportFile, *FString::Printf(TEXT("\t\t<DeploymentName>%s</DeploymentName>"), FApp::GetDeploymentName()));
         WriteLine(ReportFile, *FString::Printf(TEXT("\t\t<IsEnsure>%s</IsEnsure>"), bIsEnsure ? TEXT("1") : TEXT("0")));
-        WriteLine(ReportFile, *FString::Printf(TEXT("\t\t<IsAssert>%s</IsAssert>"), FDebug::bHasAsserted ? TEXT("1") : TEXT("0")));
-        WriteLine(ReportFile, *FString::Printf(TEXT("\t\t<CrashType>%s</CrashType>"), FGenericCrashContext::GetCrashTypeString(bIsEnsure, FDebug::bHasAsserted, GIsGPUCrashed)));
+        WriteLine(ReportFile, *FString::Printf(TEXT("\t\t<IsAssert>%s</IsAssert>"), FDebug::HasAsserted() ? TEXT("1") : TEXT("0")));
+        WriteLine(ReportFile, *FString::Printf(TEXT("\t\t<CrashType>%s</CrashType>"), FGenericCrashContext::GetCrashTypeString(bIsEnsure, FDebug::HasAsserted(), GIsGPUCrashed)));
         WriteLine(ReportFile, *FString::Printf(TEXT("\t\t<BuildVersion>%s</BuildVersion>"), FApp::GetBuildVersion()));
         WriteLine(ReportFile, *FString::Printf(TEXT("\t\t<EngineModeEx>%s</EngineModeEx>"), FGenericCrashContext::EngineModeExString()));
         
diff --git a/Engine/Source/Runtime/Core/Private/Linux/LinuxPlatformCrashContext.cpp b/Engine/Source/Runtime/Core/Private/Linux/LinuxPlatformCrashContext.cpp
--- a/Engine/Source/Runtime/Core/Private/Linux/LinuxPlatformCrashContext.cpp
+++ b/Engine/Source/Runtime/Core/Private/Linux/LinuxPlatformCrashContext.cpp
@@ -302,8 +302,8 @@ void GenerateWindowsErrorReport(const FString & WERPath, bool bReportingNonCrash
 		WriteLine(ReportFile, *FString::Printf(TEXT("\t\t<DeploymentName>%s</DeploymentName>"), FApp::GetDeploymentName()));
 		WriteLine(ReportFile, *FString::Printf(TEXT("\t\t<BuildVersion>%s</BuildVersion>"), FApp::GetBuildVersion()));
 		WriteLine(ReportFile, *FString::Printf(TEXT("\t\t<IsEnsure>%s</IsEnsure>"), bReportingNonCrash ? TEXT("1") : TEXT("0")));
-		WriteLine(ReportFile, *FString::Printf(TEXT("\t\t<IsAssert>%s</IsAssert>"), FDebug::bHasAsserted ? TEXT("1") : TEXT("0")));
-		WriteLine(ReportFile, *FString::Printf(TEXT("\t\t<CrashType>%s</CrashType>"), FGenericCrashContext::GetCrashTypeString(bReportingNonCrash, FDebug::bHasAsserted, GIsGPUCrashed)));
+		WriteLine(ReportFile, *FString::Printf(TEXT("\t\t<IsAssert>%s</IsAssert>"), FDebug::HasAsserted() ? TEXT("1") : TEXT("0")));
+		WriteLine(ReportFile, *FString::Printf(TEXT("\t\t<CrashType>%s</CrashType>"), FGenericCrashContext::GetCrashTypeString(bReportingNonCrash, FDebug::HasAsserted(), GIsGPUCrashed)));
 		WriteLine(ReportFile, *FString::Printf(TEXT("\t\t<EngineModeEx>%s</EngineModeEx>"), FGenericCrashContext::EngineModeExString()));
 		WriteLine(ReportFile, TEXT("\t</DynamicSignatures>"));
 
diff --git a/Engine/Source/Runtime/Core/Private/Mac/MacPlatformMisc.cpp b/Engine/Source/Runtime/Core/Private/Mac/MacPlatformMisc.cpp
--- a/Engine/Source/Runtime/Core/Private/Mac/MacPlatformMisc.cpp
+++ b/Engine/Source/Runtime/Core/Private/Mac/MacPlatformMisc.cpp
@@ -1633,8 +1633,8 @@ void FMacCrashContext::GenerateWindowsErrorReport(char const* WERPath, bool bIsE
 		WriteLine(ReportFile, TEXT("</Parameter2>"));
 		WriteLine(ReportFile, *FString::Printf(TEXT("\t\t<DeploymentName>%s</DeploymentName>"), FApp::GetDeploymentName()));
 		WriteLine(ReportFile, *FString::Printf(TEXT("\t\t<IsEnsure>%s</IsEnsure>"), bIsEnsure ? TEXT("1") : TEXT("0")));
-		WriteLine(ReportFile, *FString::Printf(TEXT("\t\t<IsAssert>%s</IsAssert>"), FDebug::bHasAsserted ? TEXT("1") : TEXT("0")));
-		WriteLine(ReportFile, *FString::Printf(TEXT("\t\t<CrashType>%s</CrashType>"), FGenericCrashContext::GetCrashTypeString(bIsEnsure, FDebug::bHasAsserted, GIsGPUCrashed)));
+		WriteLine(ReportFile, *FString::Printf(TEXT("\t\t<IsAssert>%s</IsAssert>"), FDebug::HasAsserted() ? TEXT("1") : TEXT("0")));
+		WriteLine(ReportFile, *FString::Printf(TEXT("\t\t<CrashType>%s</CrashType>"), FGenericCrashContext::GetCrashTypeString(bIsEnsure, FDebug::HasAsserted(), GIsGPUCrashed)));
 		WriteLine(ReportFile, *FString::Printf(TEXT("\t\t<BuildVersion>%s</BuildVersion>"), FApp::GetBuildVersion()));
 		WriteLine(ReportFile, *FString::Printf(TEXT("\t\t<EngineModeEx>%s</EngineModeEx>"), FGenericCrashContext::EngineModeExString()));
 
diff --git a/Engine/Source/Runtime/Core/Private/Misc/AssertionMacros.cpp b/Engine/Source/Runtime/Core/Private/Misc/AssertionMacros.cpp
--- a/Engine/Source/Runtime/Core/Private/Misc/AssertionMacros.cpp
+++ b/Engine/Source/Runtime/Core/Private/Misc/AssertionMacros.cpp
@@ -23,10 +23,17 @@
 #include "HAL/ExceptionHandling.h"
 #include "HAL/ThreadHeartBeat.h"
 
-bool FDebug::bHasAsserted = false;
+namespace 
+{
+	// used to track state of assets/ensures
+	bool bHasAsserted = false;
+	int32 ActiveEnsureCount = 0;
 
-#define FILE_LINE_DESC TEXT(" [File:%s] [Line: %i] ")
+	/** Lock used to synchronize the fail debug calls. */
+	FCriticalSection FailDebugCriticalSection;
+}
 
+#define FILE_LINE_DESC TEXT(" [File:%s] [Line: %i] ")
 
 /*
 	Ensure behavior
@@ -63,9 +70,6 @@ bool FDebug::bHasAsserted = false;
 
 
 
-/** Lock used to synchronize the fail debug calls. */
-static FCriticalSection	FailDebugCriticalSection;
-
 /** Number of top function calls to hide when dumping the callstack as text. */
 #if PLATFORM_LINUX
 
@@ -142,6 +146,18 @@ void StaticFailDebug( const TCHAR* Error, const ANSICHAR* File, int32 Line, cons
 }
 
 
+/// track thread asserts
+bool FDebug::HasAsserted()
+{
+	return bHasAsserted;
+}
+
+// track ensures
+bool FDebug::IsEnsuring()
+{
+	return ActiveEnsureCount > 0;
+}
+
 void FDebug::LogFormattedMessageWithCallstack(const FName& LogName, const ANSICHAR* File, int32 Line, const TCHAR* Heading, const TCHAR* Message, ELogVerbosity::Type Verbosity)
 {
 	const bool bLowLevel = LogName == NAME_None;
@@ -241,7 +257,6 @@ VARARG_BODY(void, FDebug::LogAssertFailedMessage, const TCHAR*, VARARG_EXTRA(con
  */
 void FDebug::EnsureFailed(const ANSICHAR* Expr, const ANSICHAR* File, int32 Line, const TCHAR* Msg)
 {
-
 #if STATS
 	FString EnsureFailedPerfMessage = FString::Printf(TEXT("FDebug::EnsureFailed"));
 	SCOPE_LOG_TIME_IN_SECONDS(*EnsureFailedPerfMessage, nullptr)
@@ -256,6 +271,9 @@ void FDebug::EnsureFailed(const ANSICHAR* Expr, const ANSICHAR* File, int32 Line
 		return;
 	}
 
+	// Should we spin here?
+	FPlatformAtomics::InterlockedIncrement(&ActiveEnsureCount);
+
 	// Print out the blueprint callstack
 	InternalPrintScriptCallstack(false);
 
@@ -271,124 +289,128 @@ void FDebug::EnsureFailed(const ANSICHAR* Expr, const ANSICHAR* File, int32 Line
 #if !NO_LOGGING
 		UE_LOG(LogOutputDevice, Error, TEXT("%s") FILE_LINE_DESC TEXT("\n%s\n"), ErrorString, ANSI_TO_TCHAR(File), Line, Msg);
 #endif
-		return;
 	}
-
-	// If we determine that we have not sent a report for this ensure yet, send the report below.
-	bool bShouldSendNewReport = false;
-
-	// No debugger attached, so generate a call stack and submit a crash report
-	// Walk the stack and dump it to the allocated memory.
-	const SIZE_T StackTraceSize = 65535;
-	ANSICHAR* StackTrace = (ANSICHAR*) FMemory::SystemMalloc( StackTraceSize );
-	if( StackTrace != NULL )
+	else
 	{
-		// Stop checking heartbeat for this thread. Ensure can take a lot of time (when stackwalking)
-		// Thread heartbeat will be resumed the next time this thread calls FThreadHeartBeat::Get().HeartBeat();
-		// The reason why we don't call HeartBeat() at the end of this function is that maybe this thread
-		// Never had a heartbeat checked and may not be sending heartbeats at all which would later lead to a false positives when detecting hangs.
-		FThreadHeartBeat::Get().KillHeartBeat();
+		// If we determine that we have not sent a report for this ensure yet, send the report below.
+		bool bShouldSendNewReport = false;
 
+		// No debugger attached, so generate a call stack and submit a crash report
+		// Walk the stack and dump it to the allocated memory.
+		const SIZE_T StackTraceSize = 65535;
+		ANSICHAR* StackTrace = (ANSICHAR*)FMemory::SystemMalloc(StackTraceSize);
+		if (StackTrace != NULL)
 		{
+			// Stop checking heartbeat for this thread. Ensure can take a lot of time (when stackwalking)
+			// Thread heartbeat will be resumed the next time this thread calls FThreadHeartBeat::Get().HeartBeat();
+			// The reason why we don't call HeartBeat() at the end of this function is that maybe this thread
+			// Never had a heartbeat checked and may not be sending heartbeats at all which would later lead to a false positives when detecting hangs.
+			FThreadHeartBeat::Get().KillHeartBeat();
+
+			{
 #if STATS
-			FString StackWalkPerfMessage = FString::Printf(TEXT("FPlatformStackWalk::StackWalkAndDump"));
-			SCOPE_LOG_TIME_IN_SECONDS(*StackWalkPerfMessage, nullptr)
+				FString StackWalkPerfMessage = FString::Printf(TEXT("FPlatformStackWalk::StackWalkAndDump"));
+				SCOPE_LOG_TIME_IN_SECONDS(*StackWalkPerfMessage, nullptr)
 #endif
-			StackTrace[0] = 0;
-			FPlatformStackWalk::StackWalkAndDumpEx( StackTrace, StackTraceSize, CALLSTACK_IGNOREDEPTH, FGenericPlatformStackWalk::EStackWalkFlags::FlagsUsedWhenHandlingEnsure );
-		}
+					StackTrace[0] = 0;
+				FPlatformStackWalk::StackWalkAndDumpEx(StackTrace, StackTraceSize, CALLSTACK_IGNOREDEPTH, FGenericPlatformStackWalk::EStackWalkFlags::FlagsUsedWhenHandlingEnsure);
+			}
 
-		// Create a final string that we'll output to the log (and error history buffer)
-		TCHAR ErrorMsg[16384];
-		FCString::Snprintf( ErrorMsg, ARRAY_COUNT( ErrorMsg ), TEXT( "Ensure condition failed: %s [File:%s] [Line: %i]" ) LINE_TERMINATOR TEXT( "%s" ) LINE_TERMINATOR TEXT( "Stack: " ) LINE_TERMINATOR, ANSI_TO_TCHAR( Expr ), ANSI_TO_TCHAR( File ), Line, Msg );
+			// Create a final string that we'll output to the log (and error history buffer)
+			TCHAR ErrorMsg[16384];
+			FCString::Snprintf(ErrorMsg, ARRAY_COUNT(ErrorMsg), TEXT("Ensure condition failed: %s [File:%s] [Line: %i]") LINE_TERMINATOR TEXT("%s") LINE_TERMINATOR TEXT("Stack: ") LINE_TERMINATOR, ANSI_TO_TCHAR(Expr), ANSI_TO_TCHAR(File), Line, Msg);
 
-		// Also append the stack trace
-		FCString::Strncat( ErrorMsg, ANSI_TO_TCHAR(StackTrace), ARRAY_COUNT(ErrorMsg) - 1 );
-		FMemory::SystemFree( StackTrace );
+			// Also append the stack trace
+			FCString::Strncat(ErrorMsg, ANSI_TO_TCHAR(StackTrace), ARRAY_COUNT(ErrorMsg) - 1);
+			FMemory::SystemFree(StackTrace);
 
-		// Dump the error and flush the log.
+			// Dump the error and flush the log.
 #if !NO_LOGGING
-		FDebug::LogFormattedMessageWithCallstack(LogOutputDevice.GetCategoryName(), __FILE__, __LINE__, TEXT("=== Handled ensure: ==="), ErrorMsg, ELogVerbosity::Error);
+			FDebug::LogFormattedMessageWithCallstack(LogOutputDevice.GetCategoryName(), __FILE__, __LINE__, TEXT("=== Handled ensure: ==="), ErrorMsg, ELogVerbosity::Error);
 #endif
-		GLog->Flush();
+			GLog->Flush();
 
-		// Submit the error report to the server! (and display a balloon in the system tray)
-		{
-			// How many unique previous errors we should keep track of
-			const uint32 MaxPreviousErrorsToTrack = 4;
-			static uint32 StaticPreviousErrorCount = 0;
-			if( StaticPreviousErrorCount < MaxPreviousErrorsToTrack )
+			// Submit the error report to the server! (and display a balloon in the system tray)
 			{
-				// Check to see if we've already reported this error.  No point in blasting the server with
-				// the same error over and over again in a single application session.
-				bool bHasErrorAlreadyBeenReported = false;
+				// How many unique previous errors we should keep track of
+				const uint32 MaxPreviousErrorsToTrack = 4;
+				static uint32 StaticPreviousErrorCount = 0;
+				if (StaticPreviousErrorCount < MaxPreviousErrorsToTrack)
+				{
+					// Check to see if we've already reported this error.  No point in blasting the server with
+					// the same error over and over again in a single application session.
+					bool bHasErrorAlreadyBeenReported = false;
 
-				// Static: Array of previous unique error message CRCs
-				static uint32 StaticPreviousErrorCRCs[ MaxPreviousErrorsToTrack ];
+					// Static: Array of previous unique error message CRCs
+					static uint32 StaticPreviousErrorCRCs[MaxPreviousErrorsToTrack];
 
-				// Compute CRC of error string.  Note that along with the call stack, this includes the message
-				// string passed to the macro, so only truly redundant errors will go unreported.  Though it also
-				// means you shouldn't pass loop counters to ensureMsgf(), otherwise failures may spam the server!
-				const uint32 ErrorStrCRC = FCrc::StrCrc_DEPRECATED( ErrorMsg );
+					// Compute CRC of error string.  Note that along with the call stack, this includes the message
+					// string passed to the macro, so only truly redundant errors will go unreported.  Though it also
+					// means you shouldn't pass loop counters to ensureMsgf(), otherwise failures may spam the server!
+					const uint32 ErrorStrCRC = FCrc::StrCrc_DEPRECATED(ErrorMsg);
 
-				for( uint32 CurErrorIndex = 0; CurErrorIndex < StaticPreviousErrorCount; ++CurErrorIndex )
-				{
-					if( StaticPreviousErrorCRCs[ CurErrorIndex ] == ErrorStrCRC )
+					for (uint32 CurErrorIndex = 0; CurErrorIndex < StaticPreviousErrorCount; ++CurErrorIndex)
 					{
-						// Found it!  This is a redundant error message.
-						bHasErrorAlreadyBeenReported = true;
-						break;
+						if (StaticPreviousErrorCRCs[CurErrorIndex] == ErrorStrCRC)
+						{
+							// Found it!  This is a redundant error message.
+							bHasErrorAlreadyBeenReported = true;
+							break;
+						}
 					}
-				}
 
-				// Add the element to the list and bump the count
-				StaticPreviousErrorCRCs[ StaticPreviousErrorCount++ ] = ErrorStrCRC;
+					// Add the element to the list and bump the count
+					StaticPreviousErrorCRCs[StaticPreviousErrorCount++] = ErrorStrCRC;
 
-				if( !bHasErrorAlreadyBeenReported )
-				{
+					if (!bHasErrorAlreadyBeenReported)
+					{
 #if STATS
-					FString SubmitErrorReporterfMessage = FString::Printf(TEXT("SubmitErrorReport"));
-					SCOPE_LOG_TIME_IN_SECONDS(*SubmitErrorReporterfMessage, nullptr)
+						FString SubmitErrorReporterfMessage = FString::Printf(TEXT("SubmitErrorReport"));
+						SCOPE_LOG_TIME_IN_SECONDS(*SubmitErrorReporterfMessage, nullptr)
 #endif
 
-					FCoreDelegates::OnHandleSystemEnsure.Broadcast();
+							FCoreDelegates::OnHandleSystemEnsure.Broadcast();
 
-					FPlatformMisc::SubmitErrorReport( ErrorMsg, EErrorReportMode::Balloon );
+						FPlatformMisc::SubmitErrorReport(ErrorMsg, EErrorReportMode::Balloon);
 
-					bShouldSendNewReport = true;
+						bShouldSendNewReport = true;
+					}
 				}
 			}
 		}
-	}
-	else
-	{
-		// If we fail to generate the string to identify the crash we don't know if we should skip sending the report,
-		// so we will just send the report anyway.
-		bShouldSendNewReport = true;
+		else
+		{
+			// If we fail to generate the string to identify the crash we don't know if we should skip sending the report,
+			// so we will just send the report anyway.
+			bShouldSendNewReport = true;
 
-		// Add message to log even without stacktrace. It is useful for testing fail on ensure.
+			// Add message to log even without stacktrace. It is useful for testing fail on ensure.
 #if !NO_LOGGING
-		UE_LOG(LogOutputDevice, Error, TEXT("%s [File:%s] [Line: %i]"), ErrorString, ANSI_TO_TCHAR(File), Line);
+			UE_LOG(LogOutputDevice, Error, TEXT("%s [File:%s] [Line: %i]"), ErrorString, ANSI_TO_TCHAR(File), Line);
 #endif
-	}
+		}
 
-	if ( bShouldSendNewReport )
-	{
+		if (bShouldSendNewReport)
+		{
 #if STATS
-		FString SendNewReportMessage = FString::Printf(TEXT("SendNewReport"));
-		SCOPE_LOG_TIME_IN_SECONDS(*SendNewReportMessage, nullptr)
+			FString SendNewReportMessage = FString::Printf(TEXT("SendNewReport"));
+			SCOPE_LOG_TIME_IN_SECONDS(*SendNewReportMessage, nullptr)
 #endif
 
 #if PLATFORM_DESKTOP
-		FScopeLock Lock( &FailDebugCriticalSection );
+				FScopeLock Lock(&FailDebugCriticalSection);
 
-		NewReportEnsure( GErrorMessage );
+			NewReportEnsure(GErrorMessage);
 
-		GErrorHist[0] = 0;
-		GErrorMessage[0] = 0;
-		GErrorExceptionDescription[0] = 0;
+			GErrorHist[0] = 0;
+			GErrorMessage[0] = 0;
+			GErrorExceptionDescription[0] = 0;
 #endif
+		}
 	}
+
+
+	FPlatformAtomics::InterlockedDecrement(&ActiveEnsureCount);
 }
 
 #endif // DO_CHECK || DO_GUARD_SLOW
diff --git a/Engine/Source/Runtime/Core/Public/Mac/MacPlatformCompilerPreSetup.h b/Engine/Source/Runtime/Core/Public/Mac/MacPlatformCompilerPreSetup.h
--- a/Engine/Source/Runtime/Core/Public/Mac/MacPlatformCompilerPreSetup.h
+++ b/Engine/Source/Runtime/Core/Public/Mac/MacPlatformCompilerPreSetup.h
@@ -146,9 +146,8 @@
 #pragma clang diagnostic ignored "-Wconstant-logical-operand"
 #pragma clang diagnostic ignored "-Wreserved-user-defined-literal"
 
-// Apple LLVM 8.1.0 (Xcode 8.3) introduced -Wundefined-var-template
+// Apple LLVM 8.1.0 (Xcode 8.3) introduced more warnings that we want to ignore for now
 #if (__clang_major__ > 8) || (__clang_major__ == 8 && __clang_minor__ >= 1)
-#pragma clang diagnostic ignored "-Wundefined-var-template"
 #pragma clang diagnostic ignored "-Wnullability-inferred-on-nested-type"
 #pragma clang diagnostic ignored "-Wobjc-protocol-property-synthesis"
 #pragma clang diagnostic ignored "-Wnullability-completeness-on-arrays"
diff --git a/Engine/Source/Runtime/Core/Public/Misc/AssertionMacros.h b/Engine/Source/Runtime/Core/Public/Misc/AssertionMacros.h
--- a/Engine/Source/Runtime/Core/Public/Misc/AssertionMacros.h
+++ b/Engine/Source/Runtime/Core/Public/Misc/AssertionMacros.h
@@ -26,8 +26,11 @@ struct CORE_API FDebug
 	/** Logs final assert message and exits the program. */
 	static void VARARGS AssertFailed(const ANSICHAR* Expr, const ANSICHAR* File, int32 Line, const TCHAR* Format = TEXT(""), ...);
 
-	/** Records the calling of AssertFailed() */
-	static bool bHasAsserted;
+	// returns true if an assert has occurred
+	static bool HasAsserted();
+
+	// returns true if an ensure is currently in progress (e.g. the RenderThread is ensuring)
+	static bool IsEnsuring();
 
 	/** Dumps the stack trace into the log, meant to be used for debugging purposes. */
 	static void DumpStackTraceToLog();
diff --git a/Engine/Source/Runtime/Core/Public/Windows/WindowsPlatformMath.h b/Engine/Source/Runtime/Core/Public/Windows/WindowsPlatformMath.h
--- a/Engine/Source/Runtime/Core/Public/Windows/WindowsPlatformMath.h
+++ b/Engine/Source/Runtime/Core/Public/Windows/WindowsPlatformMath.h
@@ -139,9 +139,6 @@ struct FWindowsPlatformMath : public FGenericPlatformMath
 #endif
 
 #if PLATFORM_ENABLE_POPCNT_INTRINSIC
-	/**
-	 * Use the SSE instruction to count bits
-	 */
 	static FORCEINLINE int32 CountBits(uint64 Bits)
 	{
 		return _mm_popcnt_u64(Bits);
diff --git a/Engine/Source/Runtime/Engine/Classes/AI/Navigation/NavMeshRenderingComponent.h b/Engine/Source/Runtime/Engine/Classes/AI/Navigation/NavMeshRenderingComponent.h
--- a/Engine/Source/Runtime/Engine/Classes/AI/Navigation/NavMeshRenderingComponent.h
+++ b/Engine/Source/Runtime/Engine/Classes/AI/Navigation/NavMeshRenderingComponent.h
@@ -93,10 +93,12 @@ struct ENGINE_API FNavMeshSceneProxyData : public TSharedFromThis<FNavMeshSceneP
 };
 
 // exported to API for GameplayDebugger module
-class ENGINE_API FNavMeshSceneProxy : public FDebugRenderSceneProxy
+class ENGINE_API FNavMeshSceneProxy final : public FDebugRenderSceneProxy
 {
 	friend class FNavMeshDebugDrawDelegateHelper;
 public:
+	SIZE_T GetTypeHash() const override;
+
 	FNavMeshSceneProxy(const UPrimitiveComponent* InComponent, FNavMeshSceneProxyData* InProxyData, bool ForceToRender = false);
 	virtual ~FNavMeshSceneProxy();
 
diff --git a/Engine/Source/Runtime/Engine/Classes/AI/Navigation/NavTestRenderingComponent.h b/Engine/Source/Runtime/Engine/Classes/AI/Navigation/NavTestRenderingComponent.h
--- a/Engine/Source/Runtime/Engine/Classes/AI/Navigation/NavTestRenderingComponent.h
+++ b/Engine/Source/Runtime/Engine/Classes/AI/Navigation/NavTestRenderingComponent.h
@@ -16,11 +16,12 @@ class UCanvas;
 class UNavTestRenderingComponent;
 struct FNodeDebugData;
 
-class FNavTestSceneProxy : public FDebugRenderSceneProxy
+class FNavTestSceneProxy final : public FDebugRenderSceneProxy
 {
 	friend class FNavTestDebugDrawDelegateHelper;
 
 public:
+	SIZE_T GetTypeHash() const override;
 
 	struct FNodeDebugData
 	{
diff --git a/Engine/Source/Runtime/Engine/Classes/AI/Navigation/NavigationData.h b/Engine/Source/Runtime/Engine/Classes/AI/Navigation/NavigationData.h
--- a/Engine/Source/Runtime/Engine/Classes/AI/Navigation/NavigationData.h
+++ b/Engine/Source/Runtime/Engine/Classes/AI/Navigation/NavigationData.h
@@ -553,7 +553,7 @@ public:
 	virtual void CleanUpAndMarkPendingKill();
 
 	FORCEINLINE bool IsRegistered() const { return bRegistered; }
-	void OnRegistered();
+	virtual void OnRegistered();
 	void OnUnregistered();
 	
 	FORCEINLINE uint16 GetNavDataUniqueID() const { return NavDataUniqueID; }
diff --git a/Engine/Source/Runtime/Engine/Classes/AI/Navigation/RecastNavMesh.h b/Engine/Source/Runtime/Engine/Classes/AI/Navigation/RecastNavMesh.h
--- a/Engine/Source/Runtime/Engine/Classes/AI/Navigation/RecastNavMesh.h
+++ b/Engine/Source/Runtime/Engine/Classes/AI/Navigation/RecastNavMesh.h
@@ -767,6 +767,9 @@ public:
 	virtual void OnStreamingLevelRemoved(ULevel* InLevel, UWorld* InWorld) override;
 	//~ End ANavigationData Interface
 
+	virtual void AttachNavMeshDataChunk(URecastNavMeshDataChunk& NavDataChunk);
+	virtual void DetachNavMeshDataChunk(URecastNavMeshDataChunk& NavDataChunk);
+
 protected:
 	/** Serialization helper. */
 	void SerializeRecastNavMesh(FArchive& Ar, FPImplRecastNavMesh*& NavMesh, int32 NavMeshVersion);
diff --git a/Engine/Source/Runtime/Engine/Classes/Animation/AnimInstance.h b/Engine/Source/Runtime/Engine/Classes/Animation/AnimInstance.h
--- a/Engine/Source/Runtime/Engine/Classes/Animation/AnimInstance.h
+++ b/Engine/Source/Runtime/Engine/Classes/Animation/AnimInstance.h
@@ -14,11 +14,13 @@
 #include "Components/SkeletalMeshComponent.h"
 #include "Animation/AnimNotifyQueue.h"
 #include "Animation/AnimNotifies/AnimNotify.h"
+// #include "Engine/Canvas.h"
 #include "AnimInstance.generated.h"
 
 // Post Compile Validation requires WITH_EDITOR
 #define ANIMINST_PostCompileValidation WITH_EDITOR
 
+struct FDisplayDebugManager;
 class FDebugDisplayInfo;
 class IAnimClassInterface;
 class UAnimInstance;
@@ -974,9 +976,12 @@ public:
 	// Check for whether a native exit delegate is bound to the specified state
 	bool HasNativeStateExitBinding(const FName& MachineName, const FName& StateName, FName& OutBindingName);
 
-	// Debug output for this anim instance 
+	// Debug output for this anim instance. Info for SyncGroups, Graph, Montages, etc.
 	void DisplayDebug(UCanvas* Canvas, const FDebugDisplayInfo& DebugDisplay, float& YL, float& YPos);
 
+	// Display debug info about AnimInstance. Can be overridden to add custom info from child classes.
+	virtual void DisplayDebugInstance(FDisplayDebugManager& DisplayDebugManager, float& Indent);
+
 	/** Reset any dynamics running simulation-style updates (e.g. on teleport, time skip etc.) */
 	void ResetDynamics();
 
diff --git a/Engine/Source/Runtime/Engine/Classes/Animation/AnimNode_SequencePlayer.h b/Engine/Source/Runtime/Engine/Classes/Animation/AnimNode_SequencePlayer.h
--- a/Engine/Source/Runtime/Engine/Classes/Animation/AnimNode_SequencePlayer.h
+++ b/Engine/Source/Runtime/Engine/Classes/Animation/AnimNode_SequencePlayer.h
@@ -27,15 +27,22 @@ public:
 	UPROPERTY(EditAnywhere, BlueprintReadWrite, Category=Settings, meta=(PinHiddenByDefault))
 	mutable float PlayRate;
 
+	// The Basis in which the PlayRate is expressed in. This is used to rescale PlayRate inputs.
+	// For example a Basis of 100 means that the PlayRate input will be divided by 100.
+	UPROPERTY(EditAnywhere, BlueprintReadWrite, Category = Settings, meta = (PinHiddenByDefault))
+	mutable float PlayRateBasis;
+
 	// The start up position, it only applies when reinitialized
 	// if you loop, it will still start from 0.f after finishing the round
 	UPROPERTY(EditAnywhere, BlueprintReadWrite, Category=Settings, meta=(PinHiddenByDefault))
 	mutable float StartPosition;
+
 public:	
 	FAnimNode_SequencePlayer()
 		: Sequence(NULL)
 		, bLoopAnimation(true)
 		, PlayRate(1.0f)
+		, PlayRateBasis(1.0f)
 		, StartPosition(0.f)
 	{
 	}
diff --git a/Engine/Source/Runtime/Engine/Classes/Animation/AnimationSettings.h b/Engine/Source/Runtime/Engine/Classes/Animation/AnimationSettings.h
--- a/Engine/Source/Runtime/Engine/Classes/Animation/AnimationSettings.h
+++ b/Engine/Source/Runtime/Engine/Classes/Animation/AnimationSettings.h
@@ -94,6 +94,10 @@ class ENGINE_API UAnimationSettings : public UDeveloperSettings
 	UPROPERTY(config, EditAnywhere, Category = Performance)
 	bool bStripAnimationDataOnDedicatedServer;
 
+	/** If true, pre-4.19 behavior of zero-ticking animations during skeletal mesh init */
+	UPROPERTY(config, EditAnywhere, Category = Performance)
+	bool bTickAnimationOnSkeletalMeshInit;
+
 public:
 	static UAnimationSettings * Get() { return CastChecked<UAnimationSettings>(UAnimationSettings::StaticClass()->GetDefaultObject()); }
 
diff --git a/Engine/Source/Runtime/Engine/Classes/Components/DecalComponent.h b/Engine/Source/Runtime/Engine/Classes/Components/DecalComponent.h
--- a/Engine/Source/Runtime/Engine/Classes/Components/DecalComponent.h
+++ b/Engine/Source/Runtime/Engine/Classes/Components/DecalComponent.h
@@ -112,7 +112,7 @@ protected:
 	FTimerHandle TimerHandle_DestroyDecalComponent;
 
 	/** Called when the life span of the decal has been exceeded */
-	void LifeSpanCallback();
+	virtual void LifeSpanCallback();
 
 public:
 	
