Temp change to restore old behavior and unblock Orion. Needs followup

diff --git a/Engine/Source/Programs/AutomationTool/Scripts/CopyBuildToStagingDirectory.Automation.cs b/Engine/Source/Programs/AutomationTool/Scripts/CopyBuildToStagingDirectory.Automation.cs
--- a/Engine/Source/Programs/AutomationTool/Scripts/CopyBuildToStagingDirectory.Automation.cs
+++ b/Engine/Source/Programs/AutomationTool/Scripts/CopyBuildToStagingDirectory.Automation.cs
@@ -237,7 +237,19 @@ public partial class Project : CommandUtils
                 {
                     DirectoryReference InputDir = DirectoryReference.Combine(ProjectContentRoot, RelativePath);
 					// Using Path because RelativePath may contain ".." which StagedDirectoryReference will not accept
-					StagedDirectoryReference OutputDir = new StagedDirectoryReference(Path.GetFullPath(Path.Combine(StageContentRoot.ToString(), RelativePath)));
+					StagedDirectoryReference OutputDir;
+
+					// AG - Change made in Orion to CN content that needs to be examined, in the meantime fall back to old behavior
+					// of the path doesn't start with ..
+					if (RelativePath.StartsWith(".."))
+					{
+						OutputDir = new StagedDirectoryReference(Path.GetFullPath(Path.Combine(StageContentRoot.ToString(), RelativePath)));
+					}
+					else
+					{
+						OutputDir = StagedDirectoryReference.Combine(StageContentRoot, RelativePath);
+					}
+
 					if (bUFS)
                     {
                         List<FileReference> Files = SC.FindFilesToStage(InputDir, StageFilesSearch.AllDirectories);
