Update MCP.Automation in //UE4/Main to latest.
Make additional minor changes to MCP.Automation code to reflect the new namespace / assembly for FileReference and DirectoryReference.

This is to enable randomized manifest filenames by default.

diff --git a/Engine/Source/Programs/AutomationTool/AutomationUtils/MCPPublic.cs b/Engine/Source/Programs/AutomationTool/AutomationUtils/MCPPublic.cs
--- a/Engine/Source/Programs/AutomationTool/AutomationUtils/MCPPublic.cs
+++ b/Engine/Source/Programs/AutomationTool/AutomationUtils/MCPPublic.cs
@@ -1268,6 +1268,32 @@ namespace EpicGames.MCP.Automation
 		/// <returns>An array of paths to the files in the specified location and matching the prefix constraint.</returns>
 		abstract public string[] ListFiles(string Container, string Prefix, ListOptions Options);
 
+		/// <summary>
+		/// Retrieves a list of files together with basic metadata from the cloud storage provider
+		/// </summary>
+		/// <param name="Container">The name of the container from which to list folders.</param>
+		/// <param name="Prefix">A string to specify the identifer that you want to list from. Typically used to specify a relative folder within the container to list all of its folders. Specify null to return folders in the root of the container.</param>
+		/// <param name="Options">An action which acts upon an options object to configure the operation. See ListOptions for more details.</param>
+		/// <returns>An array of metadata objects (including filenames) to the files in the specified location and matching the prefix constraint.</returns>
+		public ObjectMetadata[] ListFilesWithMetadata(string Container, string Prefix, Action<ListOptions> Options)
+		{
+			ListOptions Opts = new ListOptions();
+			if (Options != null)
+			{
+				Options(Opts);
+			}
+			return ListFilesWithMetadata(Container, Prefix, Opts);
+		}
+
+		/// <summary>
+		/// Retrieves a list of files together with basic metadata from the cloud storage provider.
+		/// </summary>
+		/// <param name="Container">The name of the folder or container from which to list files.</param>
+		/// <param name="Prefix">A string with which the identifier or filename should start. Typically used to specify a relative directory within the container to list all of its files recursively. Specify null to return all files.</param>
+		/// <param name="Options">An options object to configure the operation. See ListOptions for more details.</param>
+		/// <returns>An array of metadata objects (including filenames) to the files in the specified location and matching the prefix constraint.</returns>
+		abstract public ObjectMetadata[] ListFilesWithMetadata(string Container, string Prefix, ListOptions Options);
+
 		/// <summary>
 		/// Sets one or more items of metadata on an object in cloud storage
 		/// </summary>
@@ -1304,7 +1330,9 @@ namespace EpicGames.MCP.Automation
 		/// </summary>
 		/// <param name="Container">The name of the container in which to store files.</param>
 		/// <param name="stagingInfo">Staging info used to determine where the chunks are to copy.</param>
-		abstract public void CopyChunksToCloudStorage(string Container, BuildPatchToolStagingInfo StagingInfo);
+		/// <param name="bForce">If true, will always copy the manifest and chunks to cloud storage. Otherwise, will only copy if the manifest isn't already present on cloud storage.</param>
+		/// <returns>True if the build was copied to cloud storage, false otherwise.</returns>
+		abstract public bool CopyChunksToCloudStorage(string Container, BuildPatchToolStagingInfo StagingInfo, bool bForce = false);
 
 		/// <summary>
 		/// Copies manifest and its chunks from a specific path to a given target folder in the cloud.
@@ -1312,7 +1340,9 @@ namespace EpicGames.MCP.Automation
 		/// <param name="Container">The name of the container in which to store files.</param>
 		/// <param name="RemoteCloudDir">The path within the container that the files should be stored in.</param>
 		/// <param name="ManifestFilePath">The full path of the manifest file to copy.</param>
-		abstract public void CopyChunksToCloudStorage(string Container, string RemoteCloudDir, string ManifestFilePath);
+		/// <param name="bForce">If true, will always copy the manifest and chunks to cloud storage. Otherwise, will only copy if the manifest isn't already present on cloud storage.</param>
+		/// <returns>True if the build was copied to cloud storage, false otherwise.</returns>
+		abstract public bool CopyChunksToCloudStorage(string Container, string RemoteCloudDir, string ManifestFilePath, bool bForce = false);
 
 		/// <summary>
 		/// Verifies whether a manifest for a given build is in cloud storage.
@@ -1362,6 +1392,15 @@ namespace EpicGames.MCP.Automation
 			/// </summary>
 			public bool bReturnURLs { get; set; }
 		}
+
+		public class ObjectMetadata
+		{
+			public string ETag { get; set; }
+			public string Path { get; set; }
+			public string Url { get; set; }
+			public DateTime LastModified { get; set; }
+			public long Size { get; set; }
+		}
 	}
 
 	public abstract class CatalogServiceBase
