Fix staging errors when games add additional UFS/NonUFS files through config files with relative paths (eg. Paragon).

StagedFileSystemReference now resolves relative paths in the constructor, and checks that absolute file system paths are not passed in.

diff --git a/Engine/Source/Programs/AutomationTool/AutomationUtils/StagedFileSystemReference.cs b/Engine/Source/Programs/AutomationTool/AutomationUtils/StagedFileSystemReference.cs
--- a/Engine/Source/Programs/AutomationTool/AutomationUtils/StagedFileSystemReference.cs
+++ b/Engine/Source/Programs/AutomationTool/AutomationUtils/StagedFileSystemReference.cs
@@ -24,35 +24,89 @@ public abstract class StagedFileSystemReference
 	/// Constructor
 	/// </summary>
 	/// <param name="Name">The file system reference referred to. Either type of directory separator is permitted and will be normalized. Empty path fragments and leading/trailing slashes are not permitted.</param>
-	public StagedFileSystemReference(string Name)
+	public StagedFileSystemReference(string InName)
 	{
-		Name = Name.Replace('\\', '/');
+		Name = InName.Replace('\\', '/');
+		CanonicalName = Name.ToLowerInvariant();
 
-		this.Name = Name;
-		this.CanonicalName = Name.ToLowerInvariant();
+		// Make sure it's not an absolute path
+		if (Name.Length >= 2 && (Name[1] == ':' || Name[0] == '/'))
+		{
+			throw new ArgumentException(String.Format("Attempt to construct staged filesystem reference from absolute path ({0}). Staged paths are always relative to the staging root."), InName);
+		}
+
+		// Make sure it doesn't end in a directory separator
+		if(Name.Length > 0 && Name[Name.Length - 1] == '/')
+		{
+			throw new ArgumentException(String.Format("Staged filesystem references cannot end with path separators ({0})."), InName);
+		}
 
-		if (Name.Length > 0)
+		// Remove any relative paths
+		for(int FragmentIdx = 0; FragmentIdx < Name.Length; )
 		{
-			int LastIdx = -1;
-			for (int Idx = 0; Idx <= Name.Length; Idx++)
+			// Scan to the end of the current path fragment
+			int FragmentEndIdx = FragmentIdx;
+			while(FragmentEndIdx != Name.Length && Name[FragmentEndIdx] != '/')
 			{
-				if (Idx == Name.Length || Name[Idx] == '/')
+				FragmentEndIdx++;
+			}
+
+			// If it's an empty fragment (including if it's a path separator at the start of the string), it's invalid
+			if(FragmentEndIdx == FragmentIdx + 1)
+			{
+				throw new ArgumentException(String.Format("Staged filesystem reference may not start with a path separator or contain an empty path fragment ({0})", InName));
+			}
+
+			// Remove references to the current directory (".")
+			if(FragmentEndIdx == FragmentIdx + 1 && Name[FragmentIdx] == '.')
+			{
+				if(FragmentEndIdx < Name.Length)
 				{
-					if (Idx == LastIdx + 1)
-					{
-						throw new ArgumentException(String.Format("Empty path fragment in staged filesystem reference ({0})", Name));
-					}
-					if (Idx == LastIdx + 2 && Name[Idx - 1] == '.')
-					{
-						throw new ArgumentException(String.Format("Pseudo-directories are not permitted in filesystem references ({0})", Name));
-					}
-					if (Idx == LastIdx + 3 && Name[Idx - 1] == '.' && Name[Idx - 2] == '.')
-					{
-						throw new ArgumentException(String.Format("Pseudo-directories are not permitted in filesystem references ({0})", Name));
-					}
-					LastIdx = Idx;
+					Name = Name.Remove(FragmentIdx, (FragmentEndIdx + 1) - FragmentIdx);
+					continue;
+				}
+				else if(FragmentIdx > 0)
+				{
+					Name = Name.Remove(FragmentIdx - 1, FragmentEndIdx - (FragmentIdx - 1));
+					break;
+				}
+				else
+				{
+					Name = "";
+					break;
+				}
+			}
+
+			// Remove references to the parent directory ("..")
+			if(FragmentEndIdx == FragmentIdx + 2 && Name[FragmentIdx] == '.' && Name[FragmentIdx + 1] == '.')
+			{
+				// Make sure it's not right at the start
+				if(FragmentIdx == 0)
+				{
+					throw new ArgumentException("Staged filesystem reference cannot reference outside the staging root ({0})", InName);
+				}
+
+				// Get the start of the last fragment
+				FragmentIdx = Name.LastIndexOf('/', FragmentIdx - 2) + 1;
+				if(FragmentEndIdx < Name.Length)
+				{
+					Name = Name.Remove(FragmentIdx, (FragmentEndIdx + 1) - FragmentIdx);
+					continue;
+				}
+				else if(FragmentIdx > 0)
+				{
+					Name = Name.Remove(FragmentIdx - 1, FragmentEndIdx - (FragmentIdx - 1));
+					break;
+				}
+				else
+				{
+					Name = "";
+					break;
 				}
 			}
+
+			// Move to the next fragment
+			FragmentIdx = FragmentEndIdx + 1;
 		}
 	}
 
diff --git a/Engine/Source/Programs/AutomationTool/Scripts/CopyBuildToStagingDirectory.Automation.cs b/Engine/Source/Programs/AutomationTool/Scripts/CopyBuildToStagingDirectory.Automation.cs
--- a/Engine/Source/Programs/AutomationTool/Scripts/CopyBuildToStagingDirectory.Automation.cs
+++ b/Engine/Source/Programs/AutomationTool/Scripts/CopyBuildToStagingDirectory.Automation.cs
@@ -236,20 +236,7 @@ public partial class Project : CommandUtils
 				if (RelativePath != null)
 				{
 					DirectoryReference InputDir = DirectoryReference.Combine(ProjectContentRoot, RelativePath);
-					// Using Path because RelativePath may contain ".." which StagedDirectoryReference will not accept
-					StagedDirectoryReference OutputDir;
-
-					// AG - Change made in Orion to CN content that needs to be examined, in the meantime fall back to old behavior
-					// of the path doesn't start with ..
-					if (RelativePath.StartsWith(".."))
-					{
-						OutputDir = new StagedDirectoryReference(Path.GetFullPath(Path.Combine(StageContentRoot.ToString(), RelativePath)));
-					}
-					else
-					{
-						OutputDir = StagedDirectoryReference.Combine(StageContentRoot, RelativePath);
-					}
-
+					StagedDirectoryReference OutputDir = StagedDirectoryReference.Combine(StageContentRoot, RelativePath);
 					if (bUFS)
 					{
 						List<FileReference> Files = SC.FindFilesToStage(InputDir, StageFilesSearch.AllDirectories);
