Fix one CPU core spinning during compilation using ParallelExecutor, due to CompletedEvent never being reset. (PR #4333)

diff --git a/Engine/Source/Programs/AutomationTool/AutomationUtils/ParallelExecutor.cs b/Engine/Source/Programs/AutomationTool/AutomationUtils/ParallelExecutor.cs
--- a/Engine/Source/Programs/AutomationTool/AutomationUtils/ParallelExecutor.cs
+++ b/Engine/Source/Programs/AutomationTool/AutomationUtils/ParallelExecutor.cs
@@ -41,10 +41,10 @@ namespace AutomationTool
 			public List<string> LogLines = new List<string>();
 			public int ExitCode = -1;
 
-			private ManualResetEvent CompletedEvent;
+			private AutoResetEvent CompletedEvent;
 			private List<BuildActionExecutor> CompletedActions;
 
-			public BuildActionExecutor(BuildAction InAction, ManualResetEvent InCompletedEvent, List<BuildActionExecutor> InCompletedActions)
+			public BuildActionExecutor(BuildAction InAction, AutoResetEvent InCompletedEvent, List<BuildActionExecutor> InCompletedActions)
 			{
 				Action = InAction;
 				CompletedEvent = InCompletedEvent;
@@ -102,7 +102,7 @@ namespace AutomationTool
 				Dictionary<BuildActionExecutor, Thread> ExecutingActions = new Dictionary<BuildActionExecutor,Thread>();
 				List<BuildActionExecutor> CompletedActions = new List<BuildActionExecutor>();
 
-				using(ManualResetEvent CompletedEvent = new ManualResetEvent(false))
+				using(AutoResetEvent CompletedEvent = new AutoResetEvent(false))
 				{
 					while(QueuedActions.Count > 0 || ExecutingActions.Count > 0)
 					{
