Bringing over missing buildgraph change for enginecompatiblechange and the changes to the buildscript that somehow didn't get merged over with the enterprise changes from yesterday

diff --git a/Engine/Source/Programs/AutomationTool/BuildGraph/BuildGraph.cs b/Engine/Source/Programs/AutomationTool/BuildGraph/BuildGraph.cs
--- a/Engine/Source/Programs/AutomationTool/BuildGraph/BuildGraph.cs
+++ b/Engine/Source/Programs/AutomationTool/BuildGraph/BuildGraph.cs
@@ -133,6 +133,7 @@ namespace AutomationTool
 				DefaultProperties["EngineMajorVersion"] = Version.MajorVersion.ToString();
 				DefaultProperties["EngineMinorVersion"] = Version.MinorVersion.ToString();
 				DefaultProperties["EnginePatchVersion"] = Version.PatchVersion.ToString();
+				DefaultProperties["EngineCompatibleChange"] = Version.CompatibleChangelist.ToString();
 			}
 
 			// Add any additional custom arguments from the command line (of the form -Set:X=Y)
