Workaround for hang running automation tests, reverting a change integrated from Dev-Core.

Editor processes spawn ADB.EXE (Android Debug Bridge), which forks itself to start a server. The engine creates the child process is created with bInheritHandles = TRUE, so the forked server inherits the editor's stdout and stderr pipes AS WELL AS the stdout and stderr pipes created for ADB. The C# Process.WaitForExit() method waits until all pipes are closed before returning, so leaking these handles prevents the editor process from reporting as closed.

ProcessUtils contains code that polls Process.HasExited, then waits 120s for the pipes to be closed before abandoning them, which was historically always occuring when running automation tests. Dev-Core merge removed the need to capture that output, which resulted in the 120s delay being skipped.

This change restores the output capture and 120s delay. Will implement a proper fix in Dev-Core (probably only allowing ADB to inherit specific handles that are relevant to it from the engine), which should remove the delay and allow the process to shut down correctly.

diff --git a/Engine/Source/Programs/AutomationTool/AutomationUtils/ProcessUtils.cs b/Engine/Source/Programs/AutomationTool/AutomationUtils/ProcessUtils.cs
--- a/Engine/Source/Programs/AutomationTool/AutomationUtils/ProcessUtils.cs
+++ b/Engine/Source/Programs/AutomationTool/AutomationUtils/ProcessUtils.cs
@@ -1029,7 +1029,7 @@ namespace AutomationTool
 				EnvironmentVars.Add(AutomationTool.EnvVarNames.LocalRoot, ""); // if we don't clear this out, it will think it is a build machine; it will rederive everything
 			}
 
-			IProcessResult Result = Run(CmdEnv.UATExe, CommandLine, null, ERunOptions.Default | ERunOptions.NoStdOutCapture, EnvironmentVars, Identifier: Identifier);
+			IProcessResult Result = Run(CmdEnv.UATExe, CommandLine, null, ERunOptions.Default, EnvironmentVars, Identifier: Identifier);
 			if (Result.ExitCode != 0)
 			{
 				throw new CommandFailedException(String.Format("Recursive UAT command failed (exit code {0})", Result.ExitCode));
