Copying //UE4/Dev-Build to Dev-Main (//UE4/Dev-Main)

diff --git a/Engine/Source/Programs/AutomationTool/AutomationUtils/P4Utils.cs b/Engine/Source/Programs/AutomationTool/AutomationUtils/P4Utils.cs
--- a/Engine/Source/Programs/AutomationTool/AutomationUtils/P4Utils.cs
+++ b/Engine/Source/Programs/AutomationTool/AutomationUtils/P4Utils.cs
@@ -1161,13 +1161,15 @@ namespace AutomationTool
 		/// <param name="Client">Workspace (can be null, in which case the environment variable default will be used)</param>
 		/// <param name="ServerAndPort">Server:Port (can be null, in which case the environment variable default will be used)</param>
 		/// <param name="P4LogPath">Log filename (can be null, in which case CmdEnv.LogFolder/p4.log will be used)</param>
-		public P4Connection(string User, string Client, string ServerAndPort = null, string P4LogPath = null)
+		/// <param name="AdditionalOpts">Additional global options to include on every p4 command line</param>
+		public P4Connection(string User, string Client, string ServerAndPort = null, string P4LogPath = null, string AdditionalOpts = null)
 		{
 			var UserOpts = String.IsNullOrEmpty(User) ? "" : ("-u" + User + " ");
 			var ClientOpts = String.IsNullOrEmpty(Client) ? "" : ("-c" + Client + " ");
-			var ServerOpts = String.IsNullOrEmpty(ServerAndPort) ? "" : ("-p" + ServerAndPort + " ");			
-			GlobalOptions = UserOpts + ClientOpts + ServerOpts;
-            GlobalOptionsWithoutClient = UserOpts + ServerOpts;
+			var ServerOpts = String.IsNullOrEmpty(ServerAndPort) ? "" : ("-p" + ServerAndPort + " ");
+			AdditionalOpts = String.IsNullOrEmpty(AdditionalOpts) ? "" : AdditionalOpts + " ";
+			GlobalOptions = UserOpts + ClientOpts + ServerOpts + AdditionalOpts;
+            GlobalOptionsWithoutClient = UserOpts + ServerOpts + AdditionalOpts;
 
 			if (P4LogPath == null)
 			{
@@ -1687,10 +1689,10 @@ namespace AutomationTool
 		/// <param name="DescribeRecord">Describe record for the given changelist.</param>
 		/// <param name="AllowSpew"></param>
 		/// <returns>True if everything went okay</returns>
-        public bool DescribeChangelist(int Changelist, out DescribeRecord DescribeRecord, bool AllowSpew = true)
+        public bool DescribeChangelist(int Changelist, out DescribeRecord DescribeRecord, bool AllowSpew = true, bool bShelvedFiles = false)
         {
 			List<DescribeRecord> DescribeRecords;
-			if(!DescribeChangelists(new List<int>{ Changelist }, out DescribeRecords, AllowSpew))
+			if(!DescribeChangelists(new List<int>{ Changelist }, out DescribeRecords, AllowSpew, bShelvedFiles))
 			{
 				DescribeRecord = null;
 				return false;
@@ -1713,8 +1715,9 @@ namespace AutomationTool
 		/// <param name="Changelists">List of changelist numbers to query full descriptions for</param>
 		/// <param name="DescribeRecords">List of records we found.  One for each changelist number.  These will be sorted from oldest to newest.</param>
 		/// <param name="AllowSpew"></param>
+		/// <param name="bShelvedFiles">Whether to display shelved files</param>
 		/// <returns>True if everything went okay</returns>
-        public bool DescribeChangelists(List<int> Changelists, out List<DescribeRecord> DescribeRecords, bool AllowSpew = true)
+        public bool DescribeChangelists(List<int> Changelists, out List<DescribeRecord> DescribeRecords, bool AllowSpew = true, bool bShelvedFiles = false)
         {
 			DescribeRecords = new List<DescribeRecord>();
             try
@@ -1731,7 +1734,11 @@ namespace AutomationTool
 
 				string Output;
 				string CommandLine = "-s";		// Don't automatically diff the files
-				
+				if(bShelvedFiles)
+				{
+					CommandLine += " -S";
+				}
+
 				// Add changelists to the command-line
 				foreach( var Changelist in Changelists )
 				{
@@ -1826,7 +1833,7 @@ namespace AutomationTool
 
 						Line = Lines[ LineIndex ];
 
-						string MatchAffectedFiles = "Affected files";
+						string MatchAffectedFiles = bShelvedFiles? "Shelved files" : "Affected files";
 						int AffectedFilesAt = Line.IndexOf(MatchAffectedFiles);
 						if( AffectedFilesAt == 0 )
 						{
@@ -2929,7 +2936,7 @@ namespace AutomationTool
 					for(; LineIdx < Lines.Count && Lines[LineIdx].Length > 0; LineIdx++)
 					{
 						const string DepotFilePrefix = "... depotFile ";
-						if(Lines[LineIdx].StartsWith(DepotFilePrefix) && Lines[LineIdx].Substring(DepotFilePrefix.Length) != DepotFile)
+						if(Lines[LineIdx].StartsWith(DepotFilePrefix) && !Lines[LineIdx].Substring(DepotFilePrefix.Length).Equals(DepotFile, StringComparison.InvariantCultureIgnoreCase))
 						{
 							throw new AutomationException("Expected file record for '{0}'; received output '{1}'", DepotFile, Lines[LineIdx]);
 						}
diff --git a/Engine/Source/Programs/AutomationTool/Scripts/SubmitUtililzationReportToEC.Automation.cs b/Engine/Source/Programs/AutomationTool/Scripts/SubmitUtililzationReportToEC.Automation.cs
--- a/Engine/Source/Programs/AutomationTool/Scripts/SubmitUtililzationReportToEC.Automation.cs
+++ b/Engine/Source/Programs/AutomationTool/Scripts/SubmitUtililzationReportToEC.Automation.cs
@@ -53,7 +53,7 @@ class SubmitUtilizationReportToEC : BuildCommand
 		CommandUtils.ERunOptions Opts = CommandUtils.ERunOptions.Default | CommandUtils.ERunOptions.SpewIsVerbose;
 		XmlDocument Reader = new XmlDocument();
 		// get all of the existing properties
-		Reader.LoadXml(CommandUtils.RunAndLog("ectool", "--timeout 900 getProperties --path \"/projects/GUBP_V5/Generated/Utilization\"", Options: Opts));
+		Reader.LoadXml(CommandUtils.RunAndLog("ectool", "--timeout 900 getProperties --path \"/projects/GUBP_V5/Generated/Utilization2\"", Options: Opts));
 		// grab just the prop names
 		XmlNodeList ExistingProps = Reader.GetElementsByTagName("propertyName");
 		foreach(XmlNode prop in ExistingProps)
@@ -69,12 +69,12 @@ class SubmitUtilizationReportToEC : BuildCommand
 			if(ExistingDate < TwoWeeksAgo)
 			{
 				Log("Deleting out-of-date property: " + prop.InnerText);
-				CommandUtils.RunAndLog("ectool", string.Format("--timeout 900 deleteProperty --propertyName \"/projects/GUBP_V5/Generated/Utilization/{0}\"", prop.InnerText), Options: Opts);
+				CommandUtils.RunAndLog("ectool", string.Format("--timeout 900 deleteProperty --propertyName \"/projects/GUBP_V5/Generated/Utilization2/{0}\"", prop.InnerText), Options: Opts);
 			}
 
 		}
 		// add today's
 		Log("Adding report: " + Day.Subtract(new TimeSpan(1,0,0,0)).ToString("yyyy-MM-dd") + "->" + Day.ToString("yyyy-MM-dd"));
-		CommandUtils.RunAndLog("ectool", string.Format("--timeout 900 setProperty \"/projects/GUBP_V5/Generated/Utilization/Report_{0}\" --valueFile {1}", Day.Subtract(new TimeSpan(1,0,0,0)).ToString("yyyy-MM-dd"), FileName), Options: Opts);
+		CommandUtils.RunAndLog("ectool", string.Format("--timeout 900 setProperty \"/projects/GUBP_V5/Generated/Utilization2/Report_{0}\" --valueFile {1}", Day.Subtract(new TimeSpan(1,0,0,0)).ToString("yyyy-MM-dd"), FileName), Options: Opts);
 	}
 }
