Fix warnings caused by unknown config files being staged for various games. Can now set files which should/shouldn't be staged via WhitelistConfigFiles and BlacklistConfigFiles in the [Staging] section of *Game.ini.

diff --git a/Engine/Source/Programs/AutomationTool/AutomationUtils/DeploymentContext.cs b/Engine/Source/Programs/AutomationTool/AutomationUtils/DeploymentContext.cs
--- a/Engine/Source/Programs/AutomationTool/AutomationUtils/DeploymentContext.cs
+++ b/Engine/Source/Programs/AutomationTool/AutomationUtils/DeploymentContext.cs
@@ -231,6 +231,18 @@ public class DeploymentContext //: ProjectParams
 	/// </summary>
 	public List<Tuple<StagedDirectoryReference, StagedDirectoryReference>> RemapDirectories = new List<Tuple<StagedDirectoryReference, StagedDirectoryReference>>();
 
+	/// <summary>
+	/// Set of config files which are whitelisted to be staged. By default, we warn for config files which are not well known to prevent internal data (eg. editor/server settings)
+	/// leaking in packaged builds. This list is read from the +WhitelistConfigFiles=... array in the [Staging] section of *Game.ini files.
+	/// </summary>
+	public HashSet<StagedFileReference> WhitelistConfigFiles = new HashSet<StagedFileReference>();
+
+	/// <summary>
+	/// Set of config files which are blacklisted from staging. By default, we warn for config files which are not well known to prevent internal data (eg. editor/server settings)
+	/// leaking in packaged builds. This list is read from the +BlacklistConfigFiles=... array in the [Staging] section of *Game.ini files.
+	/// </summary>
+	public HashSet<StagedFileReference> BlacklistConfigFiles = new HashSet<StagedFileReference>();
+
 	/// <summary>
 	///  Directory to archive all of the files in: d:\archivedir\WindowsNoEditor
 	/// </summary>
@@ -431,6 +443,10 @@ public class DeploymentContext //: ProjectParams
 			}
 		}
 
+		// Read the list of files which are whitelisted to be staged
+		ReadConfigFileList(GameConfig, "Staging", "WhitelistConfigFiles", WhitelistConfigFiles);
+		ReadConfigFileList(GameConfig, "Staging", "BlacklistConfigFiles", BlacklistConfigFiles);
+
         // If we were configured to use manifests across the whole project, then this platform should use manifests.
         // Otherwise, read whether we are generating chunks from the ProjectPackagingSettings ini.
         if (InForceChunkManifests)
@@ -449,6 +465,25 @@ public class DeploymentContext //: ProjectParams
         }
     }
 
+	/// <summary>
+	/// Read a list of whitelisted or blacklisted config files names from a config file
+	/// </summary>
+	/// <param name="Config">The config hierarchy to read from</param>
+	/// <param name="SectionName">The section name</param>
+	/// <param name="KeyName">The key name to read from</param>
+	/// <param name="ConfigFiles">Receives a list of config file paths</param>
+	private static void ReadConfigFileList(ConfigHierarchy Config, string SectionName, string KeyName, HashSet<StagedFileReference> ConfigFiles)
+	{
+		List<string> ConfigFileNames;
+		if(Config.GetArray(SectionName, KeyName, out ConfigFileNames))
+		{
+			foreach(string ConfigFileName in ConfigFileNames)
+			{
+				ConfigFiles.Add(new StagedFileReference(ConfigFileName));
+			}
+		}
+	}
+
 	/// <summary>
 	/// Finds files to stage under a given base directory.
 	/// </summary>
@@ -502,11 +537,11 @@ public class DeploymentContext //: ProjectParams
 	}
 
 	/// <summary>
-	/// Stage a single file to its default location
+	/// Gets the default location to stage an input file
 	/// </summary>
-	/// <param name="FileType">The type of file being staged</param>
-	/// <param name="InputFile">Path to the file</param>
-	public void StageFile(StagedFileType FileType, FileReference InputFile)
+	/// <param name="InputFile">Location of the file in the file system</param>
+	/// <returns>Staged file location</returns>
+	public StagedFileReference GetStagedFileLocation(FileReference InputFile)
 	{
 		StagedFileReference OutputFile;
 		if(InputFile.IsUnderDirectory(ProjectRoot))
@@ -535,6 +570,17 @@ public class DeploymentContext //: ProjectParams
         {
 			throw new AutomationException("Can't deploy {0} because it doesn't start with {1} or {2}", InputFile, ProjectRoot, LocalRoot);
 		}
+		return OutputFile;
+	}
+
+	/// <summary>
+	/// Stage a single file to its default location
+	/// </summary>
+	/// <param name="FileType">The type of file being staged</param>
+	/// <param name="InputFile">Path to the file</param>
+	public void StageFile(StagedFileType FileType, FileReference InputFile)
+	{
+		StagedFileReference OutputFile = GetStagedFileLocation(InputFile);
 		StageFile(FileType, InputFile, OutputFile);
 	}
 
diff --git a/Engine/Source/Programs/AutomationTool/Scripts/CopyBuildToStagingDirectory.Automation.cs b/Engine/Source/Programs/AutomationTool/Scripts/CopyBuildToStagingDirectory.Automation.cs
--- a/Engine/Source/Programs/AutomationTool/Scripts/CopyBuildToStagingDirectory.Automation.cs
+++ b/Engine/Source/Programs/AutomationTool/Scripts/CopyBuildToStagingDirectory.Automation.cs
@@ -869,9 +869,14 @@ public partial class Project : CommandUtils
 		List<FileReference> ConfigFiles = SC.FindFilesToStage(ConfigDir, "*.ini", StageFilesSearch.AllDirectories);
 		foreach(FileReference ConfigFile in ConfigFiles)
 		{
-			if(ShouldStageConfigFile(SC, ConfigDir, ConfigFile))
+			Nullable<bool> ShouldStage = ShouldStageConfigFile(SC, ConfigDir, ConfigFile);
+			if(ShouldStage == null)
+			{
+				CommandUtils.LogWarning("The config file '{0}' will be staged, but is not whitelisted or blacklisted. Add +WhitelistConfigFiles={0} or +BlacklistConfigFiles={0} to the [Staging] section of DefaultGame.ini", SC.GetStagedFileLocation(ConfigFile));
+			}
+
+			if(ShouldStage ?? true)
 			{
-				CommandUtils.Log("Including config file {0}", ConfigFile);
 				SC.StageFile(StagedFileType.UFS, ConfigFile);
 			}
 			else
@@ -888,8 +893,18 @@ public partial class Project : CommandUtils
 	/// <param name="ConfigDir">Directory containing the config files</param>
 	/// <param name="ConfigFile">The config file to check</param>
 	/// <returns>True if the file should be staged, false otherwise</returns>
-	static bool ShouldStageConfigFile(DeploymentContext SC, DirectoryReference ConfigDir, FileReference ConfigFile)
+	static Nullable<bool> ShouldStageConfigFile(DeploymentContext SC, DirectoryReference ConfigDir, FileReference ConfigFile)
 	{
+		StagedFileReference StagedConfigFile = SC.GetStagedFileLocation(ConfigFile);
+		if(SC.WhitelistConfigFiles.Contains(StagedConfigFile))
+		{
+			return true;
+		}
+		if(SC.BlacklistConfigFiles.Contains(StagedConfigFile))
+		{
+			return false;
+		}
+
 		string NormalizedPath = ConfigFile.MakeRelativeTo(ConfigDir).ToLowerInvariant().Replace('\\', '/');
 
 		int DirectoryIdx = NormalizedPath.IndexOf('/');
@@ -910,7 +925,7 @@ public partial class Project : CommandUtils
 			const string DedicatedServerPrefix = "dedicatedserver";
 			if(NormalizedPath.StartsWith(DedicatedServerPrefix))
 			{
-				return SC.DedicatedServer && ShouldStageConfigSuffix(SC, ConfigFile, NormalizedPath.Substring(DedicatedServerPrefix.Length));
+				return SC.DedicatedServer? ShouldStageConfigSuffix(SC, ConfigFile, NormalizedPath.Substring(DedicatedServerPrefix.Length)) : false;
 			}
 
 			if(NormalizedPath == "consolevariables.ini")
@@ -947,9 +962,7 @@ public partial class Project : CommandUtils
 				return false;
 			}
 		}
-
-		CommandUtils.LogWarning("The config file '{0}' is not whitelisted nor blacklisted for staging. Ignoring.", ConfigFile);
-		return true;
+		return null;
 	}
 
 	/// <summary>
@@ -958,8 +971,8 @@ public partial class Project : CommandUtils
 	/// <param name="SC">The staging context</param>
 	/// <param name="ConfigFile">Full path to the config file</param>
 	/// <param name="InvariantSuffix">Suffix for the config file, as a lowercase invariant string</param>
-	/// <returns>True if the suffix should be staged, false otherwise</returns>
-	static bool ShouldStageConfigSuffix(DeploymentContext SC, FileReference ConfigFile, string InvariantSuffix)
+	/// <returns>True if the suffix should be staged, false if not, null if unknown</returns>
+	static Nullable<bool> ShouldStageConfigSuffix(DeploymentContext SC, FileReference ConfigFile, string InvariantSuffix)
 	{
 		switch(InvariantSuffix)
 		{
@@ -986,8 +999,7 @@ public partial class Project : CommandUtils
 			case "lightmass.ini":
 				return false;
 			default:
-				CommandUtils.LogWarning("The config file '{0}' is not whitelisted nor blacklisted for staging. Including anyway.", ConfigFile);
-				return true;
+				return null;
 		}
 	}
 
