Copying Quail@3980758 to Dev-Main (//UE4/Dev-Main)

diff --git a/Engine/Config/Linux/LinuxEngine.ini b/Engine/Config/Linux/LinuxEngine.ini
--- a/Engine/Config/Linux/LinuxEngine.ini
+++ b/Engine/Config/Linux/LinuxEngine.ini
@@ -1,15 +1,4 @@
 [Audio]
-
-; Uncomment one of the below lines
-
-; 1) battle-tested, but limited and buggy (UE-49729, UE-49807) audio device
-;AudioDeviceModuleName=ALAudio
-; 2) new, multi-platform mixer module
-AudioDeviceModuleName=AudioMixerSDL
-
-; Defining below allows switching to audio mixer using -audiomixer commandline
-AudioMixerModuleName=AudioMixerSDL
-
 ; Defines a platform-specific volume headroom (in dB) for audio to provide better platform consistency with respect to volume levels.
 PlatformHeadroomDB=0
 
@@ -31,6 +20,3 @@ DeviceProfileSelectionModule="LinuxDeviceProfileSelector"
 ;r.TriangleOrderOptimization=2
 ; larger timeout since drivers may take longer time
 g.TimeoutForBlockOnRenderFence=60000
-
-[PlatformCrypto]
-PlatformRequiresDataCrypto=True
\ No newline at end of file
diff --git a/Engine/Config/Unix/UnixEngine.ini b/Engine/Config/Unix/UnixEngine.ini
new file mode 100644
--- /dev/null
+++ b/Engine/Config/Unix/UnixEngine.ini
@@ -0,0 +1,6 @@
+[Audio]
+AudioDeviceModuleName=AudioMixerSDL
+AudioMixerModuleName=AudioMixerSDL
+
+[PlatformCrypto]
+PlatformRequiresDataCrypto=True
diff --git a/Engine/Plugins/Online/OnlineSubsystem/Source/Private/OnlineSubsystemModule.cpp b/Engine/Plugins/Online/OnlineSubsystem/Source/Private/OnlineSubsystemModule.cpp
--- a/Engine/Plugins/Online/OnlineSubsystem/Source/Private/OnlineSubsystemModule.cpp
+++ b/Engine/Plugins/Online/OnlineSubsystem/Source/Private/OnlineSubsystemModule.cpp
@@ -113,17 +113,7 @@ void FOnlineSubsystemModule::LoadDefaultSubsystem()
 		bHasLoadedModule = TryLoadSubsystemAndSetDefault(FName(*InterfaceString));
 	}
 
-	// if the above failed, then attempt the platform's default
-	if (!bHasLoadedModule)
-	{
-		FName PlatformOSSName = FPlatformMisc::GetDefaultOnlineSubsystemName();
-		if (PlatformOSSName != NAME_None)
-		{
-			bHasLoadedModule = TryLoadSubsystemAndSetDefault(PlatformOSSName);
-		}
-	}
-
-	// if all else fails, attempt to load Null
+	// if default fails, attempt to load Null
 	if (!bHasLoadedModule)
 	{
 		bHasLoadedModule = TryLoadSubsystemAndSetDefault(NULL_SUBSYSTEM);
diff --git a/Engine/Plugins/Online/OnlineSubsystem/Source/Public/OnlineSubsystem.h b/Engine/Plugins/Online/OnlineSubsystem/Source/Public/OnlineSubsystem.h
--- a/Engine/Plugins/Online/OnlineSubsystem/Source/Public/OnlineSubsystem.h
+++ b/Engine/Plugins/Online/OnlineSubsystem/Source/Public/OnlineSubsystem.h
@@ -165,17 +165,34 @@ public:
 	 */
 	static IOnlineSubsystem* GetByPlatform(bool bAutoLoad=true)
 	{
-		FName PlatformDefaultSubsystemName = FPlatformMisc::GetDefaultOnlineSubsystemName();
-		if (PlatformDefaultSubsystemName == NAME_None)
+		if (PLATFORM_PS4)
 		{
-			return nullptr;
+			if (bAutoLoad || IOnlineSubsystem::IsLoaded(PS4_SUBSYSTEM))
+			{
+				return IOnlineSubsystem::Get(PS4_SUBSYSTEM);
+			}
 		}
-
-		if (bAutoLoad || IOnlineSubsystem::IsLoaded(PlatformDefaultSubsystemName))
+		else if (PLATFORM_XBOXONE)
 		{
-			return IOnlineSubsystem::Get(PlatformDefaultSubsystemName);
+			if (bAutoLoad || IOnlineSubsystem::IsLoaded(LIVE_SUBSYSTEM))
+			{
+				return IOnlineSubsystem::Get(LIVE_SUBSYSTEM);
+			}
+		}
+		else if (PLATFORM_ANDROID)
+		{
+			if (bAutoLoad || IOnlineSubsystem::IsLoaded(GOOGLEPLAY_SUBSYSTEM))
+			{
+				return IOnlineSubsystem::Get(GOOGLEPLAY_SUBSYSTEM);
+			}
+		}
+		else if (PLATFORM_IOS)
+		{
+			if (bAutoLoad || IOnlineSubsystem::IsLoaded(IOS_SUBSYSTEM))
+			{
+				return IOnlineSubsystem::Get(IOS_SUBSYSTEM);
+			}
 		}
-
 		return nullptr;
 	}
 
diff --git a/Engine/Plugins/Performance/PerformanceMonitor/Source/PerformanceMonitor/Private/PerformanceMonitor.cpp b/Engine/Plugins/Performance/PerformanceMonitor/Source/PerformanceMonitor/Private/PerformanceMonitor.cpp
--- a/Engine/Plugins/Performance/PerformanceMonitor/Source/PerformanceMonitor/Private/PerformanceMonitor.cpp
+++ b/Engine/Plugins/Performance/PerformanceMonitor/Source/PerformanceMonitor/Private/PerformanceMonitor.cpp
@@ -26,6 +26,9 @@ FPerformanceMonitorModule::FPerformanceMonitorModule()
 	}
 	TimeBetweenRecords = 0.01f;
 	TestTimeOut = 0.f;
+	NumOfCapturedFrames = 0;
+	NumOfFramesToCapture = 0;
+	bCvsToolsMode = false;
 }
 
 FPerformanceMonitorModule::~FPerformanceMonitorModule()
@@ -145,13 +148,22 @@ bool FPerformanceMonitorModule::Exec(UWorld* InWorld, const TCHAR* Cmd, FOutputD
 				}
 			}
 		}
+		else if (FParse::Command(&Cmd, TEXT("cvstoolsmode")))
+		{
+			FString StringCommand = Cmd;
+			if (!StringCommand.IsEmpty())
+			{
+				bCvsToolsMode = FCString::ToBool(*StringCommand);
+			}
+		}
 		else
 		{
 			Ar.Logf(TEXT("Incorrect PerformanceMonitor command syntax! Supported commands are: "));
 			Ar.Logf(TEXT("\tPerformanceMonitor start <filename as string>"));
 			Ar.Logf(TEXT("\tPerformanceMonitor stop"));
 			Ar.Logf(TEXT("\tPerformanceMonitor setinterval <seconds as float>"));
-			Ar.Logf(TEXT("\tAutomation addtimer <timername as string>"));
+			Ar.Logf(TEXT("\tPerformanceMonitor addtimer <timername as string>"));
+			Ar.Logf(TEXT("\tPerformanceMonitor cvstoolsmode <true|false>"));
 		}
 
 
@@ -234,6 +246,15 @@ void FPerformanceMonitorModule::StartRecordingPerfTimers(FString FileNameToUse,
 		{
 			TestTimeOut = floatValueReceived;
 		}
+		int32 intValueReceived = 0;
+		if (GConfig->GetInt(*ConfigCategory,
+			TEXT("PerformanceMonitorNumOfFramesToCapture"),
+			intValueReceived,
+			GGameIni
+		))
+		{
+			NumOfFramesToCapture = intValueReceived;
+		}
 		TArray<FString> TimersOfInterest;
 		if (GConfig->GetArray(*ConfigCategory,
 			TEXT("PerformanceMonitorTimers"), TimersOfInterest, GGameIni))
@@ -357,6 +378,15 @@ void FPerformanceMonitorModule::RecordFrame()
 		StopRecordingPerformanceTimers();
 	}
 
+	if (NumOfFramesToCapture > 0)
+	{
+		++NumOfCapturedFrames;
+
+		if (NumOfCapturedFrames >= NumOfFramesToCapture)
+		{
+			StopRecordingPerformanceTimers();
+		}
+	}
 }
 
 void FPerformanceMonitorModule::GetStatsBreakdown()
@@ -424,7 +454,15 @@ void FPerformanceMonitorModule::StopRecordingPerformanceTimers()
 		return;
 	}
 
-	RecordData();
+	if (bCvsToolsMode)
+	{
+		RecordDataInCvsToolsMode();
+	}
+	else
+	{
+		RecordData();
+	}
+
 	FThreadStats::MasterEnableSubtract(1);
 	FStatsThreadState& Stats = FStatsThreadState::GetLocalState();
 	Stats.NewFrameDelegate.RemoveAll(this);
@@ -444,7 +482,7 @@ void FPerformanceMonitorModule::StopRecordingPerformanceTimers()
 
 	if (bExitOnCompletion)
 	{
-		FPlatformMisc::RequestExit(true);
+		FPlatformMisc::RequestExit(false);
 	}
 #endif
 }
@@ -506,6 +544,57 @@ float FPerformanceMonitorModule::GetAverageOfArray(TArray<float>ArrayToAvg, FStr
 	return RetVal;
 }
 
+void FPerformanceMonitorModule::RecordDataInCvsToolsMode()
+{
+#if STATS
+
+	int MaxStats = 0;
+
+	// Print header
+	FString StringToPrint;
+	for (TMap<FString, TArray<float>>::TConstIterator ThreadedItr(GeneratedStats); ThreadedItr; ++ThreadedItr)
+	{
+		if (!StringToPrint.IsEmpty())
+		{
+			StringToPrint += TEXT(",");
+		}
+
+		StringToPrint += ThreadedItr.Key();
+
+		MaxStats = FMath::Max(MaxStats, ThreadedItr.Value().Num());
+	}
+
+	StringToPrint += TEXT("\n");
+	FileToLogTo->Serialize(TCHAR_TO_ANSI(*StringToPrint), StringToPrint.Len());
+
+	for (int i = 0; i < MaxStats; ++i)
+	{
+		StringToPrint.Empty();
+
+		for (TMap<FString, TArray<float>>::TConstIterator ThreadedItr(GeneratedStats); ThreadedItr; ++ThreadedItr)
+		{
+			float StatValue = -1;
+			const TArray<float>& StatValues = ThreadedItr.Value();
+			if (i < StatValues.Num())
+			{
+				StatValue = StatValues[i];
+			}
+
+			if (!StringToPrint.IsEmpty())
+			{
+				StringToPrint += TEXT(",");
+			}
+
+			StringToPrint += FString::Printf(TEXT("%0.4f"), StatValue);
+		}
+
+		StringToPrint += TEXT("\n");
+		FileToLogTo->Serialize(TCHAR_TO_ANSI(*StringToPrint), StringToPrint.Len());
+	}
+
+#endif
+}
+
 void FPerformanceMonitorModule::RecordData()
 {
 #if STATS
diff --git a/Engine/Plugins/Performance/PerformanceMonitor/Source/PerformanceMonitor/Private/PerformanceMonitor.h b/Engine/Plugins/Performance/PerformanceMonitor/Source/PerformanceMonitor/Private/PerformanceMonitor.h
--- a/Engine/Plugins/Performance/PerformanceMonitor/Source/PerformanceMonitor/Private/PerformanceMonitor.h
+++ b/Engine/Plugins/Performance/PerformanceMonitor/Source/PerformanceMonitor/Private/PerformanceMonitor.h
@@ -38,10 +38,13 @@ private:
 	float TestTimeOut;
 	float TimeOfLastRecord;
 	float TimeBetweenRecords;
+	int32 NumOfCapturedFrames;
+	int32 NumOfFramesToCapture;
 	FString MapToTest;
 	bool bHasWarnedAboutTime;
 	bool bExitOnCompletion;
 	bool bRequiresCutsceneStart;
+	bool bCvsToolsMode;
 protected:
 public:
 
@@ -93,6 +96,7 @@ public:
 
 	void RecordFrame();
 	void GetStatsBreakdown();
+	void RecordDataInCvsToolsMode();
 	void RecordData();	
 
 	float GetAverageOfArray(TArray<float>ArrayToAvg, FString StatName);
diff --git a/Engine/Source/Developer/DesktopPlatform/Private/PlatformInfo.cpp b/Engine/Source/Developer/DesktopPlatform/Private/PlatformInfo.cpp
--- a/Engine/Source/Developer/DesktopPlatform/Private/PlatformInfo.cpp
+++ b/Engine/Source/Developer/DesktopPlatform/Private/PlatformInfo.cpp
@@ -134,6 +134,7 @@ static const FPlatformInfo AllPlatformInfoArray[] = {
 
 	BuildPlatformInfo(TEXT("TVOS"),						TEXT("TVOS"),				LOCTEXT("TVOSTargetPlatDisplay", "tvOS"),						EPlatformType::Game,	EPlatformFlags::None,			FPlatformIconPaths(TEXT("Launcher/TVOS/Platform_TVOS_24x"), TEXT("Launcher/TVOS/Platform_TVOS_128x")),								TEXT(""),											TEXT(""),			EPlatformSDKStatus::Unknown,	TEXT(""),																								IsAvailableOnWindows || IsAvailableOnMac,						TEXT("TVOS"),		TEXT("TVOS"),		false,					true,					false,			TEXT("TVOS"),		TEXT("Mobile")),
 	BuildPlatformInfo(TEXT("Switch"),					TEXT("Switch"),				LOCTEXT("Switch", "Switch"),									EPlatformType::Game,	EPlatformFlags::None,			FPlatformIconPaths(TEXT("Launcher/Switch/Platform_Switch_24x"), TEXT("Launcher/Switch/Platform_Switch_128x")),						TEXT(""),											TEXT(""),			EPlatformSDKStatus::Unknown,	TEXT(""),																								IsAvailableOnWindows,											TEXT("Switch"),		TEXT("Switch"),		false,					true,					true,			TEXT("Switch"),		TEXT("Mobile")),
+	BuildPlatformInfo(TEXT("Quail"),					TEXT("Quail"),				LOCTEXT("Quail", "Quail"),										EPlatformType::Game,	EPlatformFlags::None,			FPlatformIconPaths(TEXT("Launcher/Quail/Platform_Quail_24x"), TEXT("Launcher/Quail/Platform_Quail_128x")),							TEXT(""),											TEXT(""),			EPlatformSDKStatus::Unknown,	TEXT(""),																								IsAvailableOnWindows,											TEXT("Quail"),		TEXT("Quail"),		false,					true,					true,			TEXT("Quail"),		TEXT("")),
 
 	// Note: For "AllDesktop" bEnabledForUse value, see SProjectTargetPlatformSettings::Construct !!!! IsAvailableOnWindows || IsAvailableOnMac || IsAvailableOnLinux
 };
diff --git a/Engine/Source/Programs/AutomationTool/AutomationUtils/Platform.cs b/Engine/Source/Programs/AutomationTool/AutomationUtils/Platform.cs
--- a/Engine/Source/Programs/AutomationTool/AutomationUtils/Platform.cs
+++ b/Engine/Source/Programs/AutomationTool/AutomationUtils/Platform.cs
@@ -210,6 +210,17 @@ namespace AutomationTool
 			return ClientProcess;
 		}
 
+		/// <summary>
+		/// Downloads file from target device to local pc
+		/// </summary>
+		/// <param name="RemoteFilePath"></param>
+		/// <param name="LocalFile"></param>
+		/// <param name="Params"></param>
+		public virtual void GetTargetFile(string RemoteFilePath, string LocalFile, ProjectParams Params)
+		{
+			throw new NotImplementedException();
+		}
+
 		/// <summary>
 		/// Allow platform specific clean-up or detection after client has run
 		/// </summary>
diff --git a/Engine/Source/Programs/AutomationTool/AutomationUtils/ProjectParams.cs b/Engine/Source/Programs/AutomationTool/AutomationUtils/ProjectParams.cs
--- a/Engine/Source/Programs/AutomationTool/AutomationUtils/ProjectParams.cs
+++ b/Engine/Source/Programs/AutomationTool/AutomationUtils/ProjectParams.cs
@@ -306,6 +306,7 @@ namespace AutomationTool
 			this.ForcePackageData = InParams.ForcePackageData;
 			this.Deploy = InParams.Deploy;
 			this.DeployFolder = InParams.DeployFolder;
+			this.GetFile = InParams.GetFile;
 			this.IterativeDeploy = InParams.IterativeDeploy;
 			this.IgnoreCookErrors = InParams.IgnoreCookErrors;
 			this.FastCook = InParams.FastCook;
@@ -426,6 +427,7 @@ namespace AutomationTool
 			bool? Client = null,
 			bool? Deploy = null,
 			string DeployFolder = null,
+			string GetFile = null,
 			bool? FileServer = null,
 			bool? Foreign = null,
 			bool? ForeignCode = null,
@@ -724,6 +726,8 @@ namespace AutomationTool
 				this.Deploy = true;
 			}
 
+			this.GetFile = ParseParamValueIfNotSpecified(Command, GetFile, "getfile", null);
+
 			this.IterativeDeploy = GetParamValueIfNotSpecified(Command, IterativeDeploy, this.IterativeDeploy, new string[] {"iterativedeploy", "iterate" } );
 			this.FastCook = GetParamValueIfNotSpecified(Command, FastCook, this.FastCook, "FastCook");
 			this.IgnoreCookErrors = GetParamValueIfNotSpecified(Command, IgnoreCookErrors, this.IgnoreCookErrors, "IgnoreCookErrors");
@@ -1889,6 +1893,13 @@ namespace AutomationTool
 
 		#endregion
 
+		#region GetFile
+
+		[Help("getfile", "download file from target after successful run")]
+		public string GetFile { get; set; }
+
+		#endregion
+
 		#region Misc
 
 		[Help("MapsToRebuildLightMaps", "List of maps that need light maps rebuilding")]
diff --git a/Engine/Source/Programs/AutomationTool/Scripts/AutomationScripts.Automation.csproj b/Engine/Source/Programs/AutomationTool/Scripts/AutomationScripts.Automation.csproj
--- a/Engine/Source/Programs/AutomationTool/Scripts/AutomationScripts.Automation.csproj
+++ b/Engine/Source/Programs/AutomationTool/Scripts/AutomationScripts.Automation.csproj
@@ -84,6 +84,7 @@
     <Compile Include="CopyUAT.Automation.cs" />
     <Compile Include="CopyUsingDistillFileSet.Automation.cs" />
     <Compile Include="DeployCommand.Automation.cs" />
+    <Compile Include="GetFileCommand.Automation.cs" />
     <Compile Include="RebasePublicIncludePaths.cs" />
     <Compile Include="FixupRedirects.Automation.cs" />
     <Compile Include="GenerateDSYM.Automation.cs" />
@@ -91,6 +92,7 @@
     <Compile Include="ListThirdPartySoftware.Automation.cs" />
     <Compile Include="RebuildHLODCommand.Automation.cs" />
     <Compile Include="RebuildLightMapsCommand.Automation.cs" />
+    <Compile Include="RecordPerformance.Automation.cs" />
     <Compile Include="ReplaceAssetsUsingManifest.cs" />
     <Compile Include="ResavePackagesCommand.Automation.cs" />
     <Compile Include="IPhonePackager.Automation.cs" />
diff --git a/Engine/Source/Programs/AutomationTool/Scripts/BuildCookRun.Automation.cs b/Engine/Source/Programs/AutomationTool/Scripts/BuildCookRun.Automation.cs
--- a/Engine/Source/Programs/AutomationTool/Scripts/BuildCookRun.Automation.cs
+++ b/Engine/Source/Programs/AutomationTool/Scripts/BuildCookRun.Automation.cs
@@ -222,6 +222,7 @@ public class BuildCookRun : BuildCommand
 		Project.Deploy(Params);
 		PrintRunTime();
 		Project.Run(Params);
+		Project.GetFile(Params);
 
 		// Check everything in!
 		if (WorkingCL != -1)
diff --git a/Engine/Source/Programs/AutomationTool/Scripts/GetFileCommand.Automation.cs b/Engine/Source/Programs/AutomationTool/Scripts/GetFileCommand.Automation.cs
new file mode 100644
--- /dev/null
+++ b/Engine/Source/Programs/AutomationTool/Scripts/GetFileCommand.Automation.cs
@@ -0,0 +1,41 @@
+﻿// Copyright 1998-2018 Epic Games, Inc. All Rights Reserved.
+
+using System;
+using System.Collections.Generic;
+using System.IO;
+using System.Threading;
+using System.Reflection;
+using AutomationTool;
+using UnrealBuildTool;
+using Tools.DotNETCommon;
+
+public partial class Project : CommandUtils
+{
+	#region GetFile Command
+
+	public static void GetFile(ProjectParams Params)
+	{
+		Params.ValidateAndLog();
+		if (string.IsNullOrEmpty(Params.GetFile))
+		{
+			return;
+		}
+
+		Log("********** GETFILE COMMAND STARTED **********");
+
+		var FileName = Path.GetFileName(Params.GetFile);
+		var LocalFile = CombinePaths(CmdEnv.EngineSavedFolder, FileName);
+
+		var SC = CreateDeploymentContext(Params, false);
+		if (SC.Count == 0)
+		{
+			throw new AutomationException("Failed to create deployment context");
+		}
+
+		SC[0].StageTargetPlatform.GetTargetFile(Params.GetFile, LocalFile, Params);
+
+		Log("********** GETFILE COMMAND COMPLETED **********");
+	}
+
+	#endregion
+}
diff --git a/Engine/Source/Programs/AutomationTool/Scripts/RecordPerformance.Automation.cs b/Engine/Source/Programs/AutomationTool/Scripts/RecordPerformance.Automation.cs
new file mode 100644
--- /dev/null
+++ b/Engine/Source/Programs/AutomationTool/Scripts/RecordPerformance.Automation.cs
@@ -0,0 +1,218 @@
+// Copyright 1998-2018 Epic Games, Inc. All Rights Reserved.
+using System;
+using System.Collections.Generic;
+using System.Text;
+using System.IO;
+using AutomationTool;
+using UnrealBuildTool;
+using Tools.DotNETCommon;
+using System.Linq;
+
+[Help("UAT command to run performance test demo using different RHIs and compare results")]
+class RecordPerformance : BuildCommand
+{
+	private string[] ProjectsToRun =
+	{
+		"SubwaySequencer",		// 0
+		"InfiltratorDemo",		// 1
+		"ElementalDemo",		// 2
+		"ShowdownDemo"			// 3
+	};
+
+	private string[] ProjectPathsToRun =
+	{
+		"Samples/Showcases/SubwaySequencer",
+		"Samples/Showcases/InfiltratorDemo",
+		"Samples/Showcases/ElementalDemo",
+		"Samples/Showcases/ShowdownDemo"
+	};
+
+	private string ProjectToRun;
+	private DirectoryReference RootDir;
+	private DirectoryReference ProjectPath;
+	private const string PerformanceMonitorConfig = "Default";
+	private DirectoryReference OutputDir;
+
+	public override void ExecuteBuild()
+	{
+		SetupPaths();
+
+		// Project settings should be appropriate (PerformanceMonitor plugin included and configured, per-platform settings set).
+		int Index = int.Parse(ParseParamValue("DemoIndex", "1"));
+		Index = Math.Max(0, Index);
+		Index = Math.Min(ProjectsToRun.Length - 1, Index);
+		ProjectToRun = ProjectsToRun[Index];
+		ProjectPath = DirectoryReference.Combine(RootDir, ProjectPathsToRun[Index]);
+
+		// Which platforms to test?
+		var PlatformsToTest = ParseMultipleParams("Platforms");
+		if (PlatformsToTest.Count() == 0)
+		{
+			throw new AutomationException("Missing -Platforms=Win64+Linux+... parameter");
+		}
+
+		// How many runs of each test to accumulate data?
+		var NumOfRuns = int.Parse(ParseParamValue("NumOfRuns", "3"));
+
+		// Which stats to present?
+		var StatsToPresent = ParseMultipleParams("Stats");
+		if (StatsToPresent.Count() == 0)
+		{
+			StatsToPresent = new string[] { "FrameTime", "GPUFrameTime", "RenderThreadTime", "GameThreadTime" };
+		}
+
+		// First cook and deploy all the platforms
+		foreach (var PlatformName in PlatformsToTest)
+		{
+			var Platform = GetPlatformByName(PlatformName);
+			if (Platform == null)
+			{
+				Log("Cannot find platform '{0}'. Skipping.", PlatformName);
+				continue;
+			}
+
+			BuildAndCookPlatform(Platform.PlatformType);
+		}
+
+		// Now run tests
+		foreach (var PlatformName in PlatformsToTest)
+		{
+			var Platform = GetPlatformByName(PlatformName);
+			if (Platform == null)
+			{
+				continue;
+			}
+
+			var RhisToTest = GetRhisForPlatform(PlatformName);
+
+			foreach (var Rhi in RhisToTest)
+			{
+				// First run to make sure all caches are populated
+				RunGame(Platform, Rhi, 0, false);
+
+				// Actual runs to gather perf data
+				for (var i = 1; i <= NumOfRuns; ++i)
+				{
+					RunGame(Platform, Rhi, i, true);
+				}
+
+				CombineStats(Platform, Rhi);
+			}
+		}
+
+		foreach (var Stat in StatsToPresent)
+		{
+			CreateChart(Stat);
+		}
+
+		Log("Performance charts stored in {0}", OutputDir);
+	}
+
+	private IEnumerable<string> ParseMultipleParams(string ParamName)
+	{
+		List<string> Params = new List<string>();
+
+		var ParamsString = ParseParamValue(ParamName);
+		if (!string.IsNullOrEmpty(ParamsString))
+		{
+			var ParamNames = new ParamList<string>(ParamsString.Split('+'));
+			foreach (var Param in ParamNames)
+			{
+				Params.Add(Param);
+			}
+		}
+
+		return Params;
+	}
+
+	private IEnumerable<string> GetRhisForPlatform(string PlatformName)
+	{
+		if (PlatformName == "Win64" || PlatformName == "Win32")
+		{
+			return new string[] { "d3d11", /*"d3d12",*/ "vulkan" };
+		}
+		else if (PlatformName == "Linux")
+		{
+			return new string[] { "opengl4", "vulkan" };
+		}
+		else
+		{
+			return new string[] { "default" };
+		}
+	}
+
+	private Platform GetPlatformByName(string PlatformName)
+	{
+		UnrealTargetPlatform PlatformType;
+		if (!Enum.TryParse(PlatformName, true, out PlatformType))
+		{
+			return null;
+		}
+
+		return Platform.GetPlatform(PlatformType);
+	}
+
+	private void RunCsvTool(string ToolName, string Arguments)
+	{
+		var ToolPath = FileReference.Combine(DirectoryReference.Combine(RootDir, "Engine", "Programs", "NotForLicensees", "CSVTools", "Binaries"), ToolName + ".exe");
+		Run(ToolPath.FullName, Arguments);
+	}
+
+	private void CreateChart(string StatName)
+	{
+		var StatsDir = DirectoryReference.Combine(OutputDir, "Combined");
+		var ChartFileName = FileReference.Combine(OutputDir, StatName + ".svg");
+		RunCsvTool("CSVToSVG", string.Format("-csvDir {0} -o {1} -stats {2} -maxY 100", StatsDir.FullName, ChartFileName.FullName, StatName));
+	}
+
+	private void CombineStats(Platform Platform, string Rhi)
+	{
+		var StatsDir = FileReference.Combine(OutputDir, string.Format("{0}-{1}", Platform.PlatformType, Rhi));
+		var OutputFile = FileReference.Combine(OutputDir, "Combined", string.Format("{0}-{1}.csv", Platform.PlatformType, Rhi));
+
+		CreateDirectory(OutputFile.Directory.FullName);
+
+		RunCsvTool("CSVCollate", string.Format("-csvDir {0} -o {1} -avg", StatsDir.FullName, OutputFile.FullName));
+	}
+
+	private void SetupPaths()
+	{
+		var TestName = DateTime.Now.ToString("yyyyMMddHHmm");
+
+		RootDir = new DirectoryReference(CmdEnv.LocalRoot);
+		OutputDir = DirectoryReference.Combine(RootDir, "Engine", "Saved", "RecordPerformance", TestName);
+		CreateDirectory(OutputDir.FullName);
+	}
+
+	private void RunGame(Platform Platform, string Rhi, int RunNumber, bool UseStats)
+	{
+		var AdditionalCommandLineParameters = "-savevulkanpsocacheonexit";
+		var StatFileName = Path.Combine(ProjectToRun, "Saved", "FXPerformance", string.Format("{0}.csv", PerformanceMonitorConfig));
+
+		Log("Running {0} on {1}...", Rhi, Platform.PlatformType);
+
+		var RhiParam = (Rhi != "default") ? "-" + Rhi : "";
+
+		RunBCR(string.Format("-project={0} -platform={1} -skipbuild -skipcook -skipstage -skipdeploy -run -getfile=\"{5}\" -addcmdline=\"{2} {4} -benchmark -execcmds=\\\"PerformanceMonitor addtimer Time,PerformanceMonitor cvstoolsmode true,PerformanceMonitor start {3}, ce start\\\"\"",
+			ProjectToRun, Platform.PlatformType, RhiParam, PerformanceMonitorConfig, AdditionalCommandLineParameters, StatFileName));
+
+		if (UseStats)
+		{
+			var OutputFile = FileReference.Combine(OutputDir, string.Format("{0}-{1}", Platform.PlatformType, Rhi), string.Format("{0}.csv", RunNumber));
+			CreateDirectory(OutputFile.Directory.FullName);
+			CopyFile(CombinePaths(CmdEnv.EngineSavedFolder, Path.GetFileName(StatFileName)), OutputFile.FullName);
+		}
+	}
+
+	private void BuildAndCookPlatform(UnrealTargetPlatform Platform)
+	{
+		Log("Building {0}...", Platform.ToString());
+
+		RunBCR(string.Format("-project={0} -platform={1} -build -cook -stage -pak -deploy", ProjectToRun, Platform));
+	}
+
+	private void RunBCR(string Arguments)
+	{
+		RunUAT(CmdEnv, "BuildCookRun " + Arguments, "BCR");
+	}
+}
diff --git a/Engine/Source/Programs/AutomationTool/Win/WinPlatform.Automation.cs b/Engine/Source/Programs/AutomationTool/Win/WinPlatform.Automation.cs
--- a/Engine/Source/Programs/AutomationTool/Win/WinPlatform.Automation.cs
+++ b/Engine/Source/Programs/AutomationTool/Win/WinPlatform.Automation.cs
@@ -99,6 +99,12 @@ public abstract class BaseWinPlatform : Platform
     {
     }
 
+	public override void GetTargetFile(string RemoteFilePath, string LocalFile, ProjectParams Params)
+	{
+		var SourceFile = FileReference.Combine(new DirectoryReference(Params.BaseStageDirectory), GetCookPlatform(Params.HasServerCookedTargets, Params.HasClientTargetDetected), RemoteFilePath);
+		CommandUtils.CopyFile(SourceFile.FullName, LocalFile);
+	}
+
 	void StageBootstrapExecutable(DeploymentContext SC, string ExeName, FileReference TargetFile, StagedFileReference StagedRelativeTargetPath, string StagedArguments)
 	{
 		FileReference InputFile = FileReference.Combine(SC.LocalRoot, "Engine", "Binaries", SC.PlatformDir, String.Format("BootstrapPackagedGame-{0}-Shipping.exe", SC.PlatformDir));
diff --git a/Engine/Source/Programs/UnrealBuildTool/Configuration/UEBuildPlatform.cs b/Engine/Source/Programs/UnrealBuildTool/Configuration/UEBuildPlatform.cs
--- a/Engine/Source/Programs/UnrealBuildTool/Configuration/UEBuildPlatform.cs
+++ b/Engine/Source/Programs/UnrealBuildTool/Configuration/UEBuildPlatform.cs
@@ -416,6 +416,7 @@ namespace UnrealBuildTool
 				case CppPlatform.Linux:			return UnrealTargetPlatform.Linux;
 				case CppPlatform.TVOS:			return UnrealTargetPlatform.TVOS;
 				case CppPlatform.Switch: 		return UnrealTargetPlatform.Switch;
+				case CppPlatform.Quail:			return UnrealTargetPlatform.Quail;
 			}
 			throw new BuildException("CPPTargetPlatformToUnrealTargetPlatform: Unknown CPPTargetPlatform {0}", InCPPPlatform.ToString());
 		}
diff --git a/Engine/Source/Programs/UnrealBuildTool/Configuration/UEBuildTarget.cs b/Engine/Source/Programs/UnrealBuildTool/Configuration/UEBuildTarget.cs
--- a/Engine/Source/Programs/UnrealBuildTool/Configuration/UEBuildTarget.cs
+++ b/Engine/Source/Programs/UnrealBuildTool/Configuration/UEBuildTarget.cs
@@ -83,6 +83,11 @@ namespace UnrealBuildTool
 		/// Nintendo Switch
 		/// </summary>
 		Switch,
+
+		/// <summary>
+		/// NDA'd platform Quail
+		/// </summary>
+		Quail,
 	}
 
 	/// <summary>
diff --git a/Engine/Source/Programs/UnrealBuildTool/System/CppCompileEnvironment.cs b/Engine/Source/Programs/UnrealBuildTool/System/CppCompileEnvironment.cs
--- a/Engine/Source/Programs/UnrealBuildTool/System/CppCompileEnvironment.cs
+++ b/Engine/Source/Programs/UnrealBuildTool/System/CppCompileEnvironment.cs
@@ -26,6 +26,7 @@ namespace UnrealBuildTool
 		Linux,
 		TVOS,
 		Switch,
+		Quail,
 	}
 
 	/// <summary>
diff --git a/Engine/Source/Programs/UnrealBuildTool/UnrealBuildTool.csproj b/Engine/Source/Programs/UnrealBuildTool/UnrealBuildTool.csproj
--- a/Engine/Source/Programs/UnrealBuildTool/UnrealBuildTool.csproj
+++ b/Engine/Source/Programs/UnrealBuildTool/UnrealBuildTool.csproj
@@ -124,6 +124,15 @@
       </ItemGroup>
     </When>
   </Choose>
+  <Choose>
+    <When Condition="Exists('Platform\Quail\UEBuildQuail.cs')">
+      <ItemGroup>
+        <Compile Include="Platform\Quail\QuailProjectGenerator.cs" />
+        <Compile Include="Platform\Quail\QuailToolChain.cs" />
+        <Compile Include="Platform\Quail\UEBuildQuail.cs" />
+      </ItemGroup>
+    </When>
+  </Choose>
   <ItemGroup>
     <Compile Include="Configuration\ModuleRules.cs" />
     <Compile Include="Configuration\PlatformTargetRules.cs" />
diff --git a/Engine/Source/Runtime/ApplicationCore/Public/GenericPlatform/GenericApplication.h b/Engine/Source/Runtime/ApplicationCore/Public/GenericPlatform/GenericApplication.h
--- a/Engine/Source/Runtime/ApplicationCore/Public/GenericPlatform/GenericApplication.h
+++ b/Engine/Source/Runtime/ApplicationCore/Public/GenericPlatform/GenericApplication.h
@@ -529,6 +529,8 @@ public:
 
 	virtual bool ApplicationLicenseValid(FPlatformUserId PlatformUser = PLATFORMUSERID_NONE) { return true; }
 
+	virtual bool IsAllowedToRender() const { return true; }
+
 public:
 
 	const TSharedPtr< ICursor > Cursor;
diff --git a/Engine/Source/Runtime/ApplicationCore/Public/HAL/PlatformApplicationMisc.h b/Engine/Source/Runtime/ApplicationCore/Public/HAL/PlatformApplicationMisc.h
--- a/Engine/Source/Runtime/ApplicationCore/Public/HAL/PlatformApplicationMisc.h
+++ b/Engine/Source/Runtime/ApplicationCore/Public/HAL/PlatformApplicationMisc.h
@@ -18,6 +18,8 @@
 #include "Android/AndroidApplicationMisc.h"
 #elif PLATFORM_HTML5
 #include "HTML5/HTML5PlatformApplicationMisc.h"
+#elif PLATFORM_QUAIL
+#include "Quail/QuailPlatformApplicationMisc.h"
 #elif PLATFORM_LINUX
 #include "Linux/LinuxPlatformApplicationMisc.h"
 #elif PLATFORM_SWITCH
diff --git a/Engine/Source/Runtime/ApplicationCore/Public/HAL/PlatformSplash.h b/Engine/Source/Runtime/ApplicationCore/Public/HAL/PlatformSplash.h
--- a/Engine/Source/Runtime/ApplicationCore/Public/HAL/PlatformSplash.h
+++ b/Engine/Source/Runtime/ApplicationCore/Public/HAL/PlatformSplash.h
@@ -19,6 +19,8 @@
 #include "HTML5/HTML5PlatformSplash.h"
 #elif PLATFORM_LINUX
 #include "Linux/LinuxPlatformSplash.h"
+#elif PLATFORM_QUAIL
+#include "Quail/QuailPlatformSplash.h"
 #elif PLATFORM_SWITCH
 #include "Switch/SwitchPlatformSplash.h"
 #else
diff --git a/Engine/Source/Runtime/Core/Private/Android/AndroidMisc.cpp b/Engine/Source/Runtime/Core/Private/Android/AndroidMisc.cpp
--- a/Engine/Source/Runtime/Core/Private/Android/AndroidMisc.cpp
+++ b/Engine/Source/Runtime/Core/Private/Android/AndroidMisc.cpp
@@ -1603,12 +1603,6 @@ FString FAndroidMisc::GetOSVersion()
 	return AndroidVersion;
 }
 
-const FName& FAndroidMisc::GetDefaultOnlineSubsystemName()
-{
-	static FName OSSName = GOOGLEPLAY_SUBSYSTEM;
-	return OSSName;
-}
-
 bool FAndroidMisc::GetDiskTotalAndFreeSpace(const FString& InPath, uint64& TotalNumberOfBytes, uint64& NumberOfFreeBytes)
 {
 	extern FString GExternalFilePath;
diff --git a/Engine/Source/Runtime/Core/Private/GenericPlatform/GenericPlatformMisc.cpp b/Engine/Source/Runtime/Core/Private/GenericPlatform/GenericPlatformMisc.cpp
--- a/Engine/Source/Runtime/Core/Private/GenericPlatform/GenericPlatformMisc.cpp
+++ b/Engine/Source/Runtime/Core/Private/GenericPlatform/GenericPlatformMisc.cpp
@@ -1081,9 +1081,3 @@ void FGenericPlatformMisc::UnregisterForRemoteNotifications()
 {
 	// not implemented by default
 }
-
-const FName& FGenericPlatformMisc::GetDefaultOnlineSubsystemName()
-{
-	static FName OSSName = NAME_None;
-	return OSSName;
-}
diff --git a/Engine/Source/Runtime/Core/Private/IOS/IOSPlatformMisc.cpp b/Engine/Source/Runtime/Core/Private/IOS/IOSPlatformMisc.cpp
--- a/Engine/Source/Runtime/Core/Private/IOS/IOSPlatformMisc.cpp
+++ b/Engine/Source/Runtime/Core/Private/IOS/IOSPlatformMisc.cpp
@@ -718,12 +718,6 @@ FString FIOSPlatformMisc::GetCPUBrand()
 	return GetIOSDeviceIDString();
 }
 
-const FName& FIOSPlatformMisc::GetDefaultOnlineSubsystemName()
-{
-	static FName OSSName = IOS_SUBSYSTEM;
-	return OSSName;
-}
-
 void FIOSPlatformMisc::GetOSVersions(FString& out_OSVersionLabel, FString& out_OSSubVersionLabel)
 {
 #if PLATFORM_TVOS
diff --git a/Engine/Source/Runtime/Core/Public/Android/AndroidMisc.h b/Engine/Source/Runtime/Core/Public/Android/AndroidMisc.h
--- a/Engine/Source/Runtime/Core/Public/Android/AndroidMisc.h
+++ b/Engine/Source/Runtime/Core/Public/Android/AndroidMisc.h
@@ -20,8 +20,6 @@ class TFunction;
 #define UE_DEBUG_BREAK() (FAndroidMisc::DebugBreakInternal())
 #endif
 
-#define GOOGLEPLAY_SUBSYSTEM FName(TEXT("GOOGLEPLAY"))
-
 /**
  * Android implementation of the misc OS functions
  */
@@ -197,8 +195,6 @@ public:
 		__sync_synchronize();
 	}
 
-	static const class FName& GetDefaultOnlineSubsystemName();
-
 
 #if STATS || ENABLE_STATNAMEDEVENTS
 	static void BeginNamedEventFrame();
diff --git a/Engine/Source/Runtime/Core/Public/CoreGlobals.h b/Engine/Source/Runtime/Core/Public/CoreGlobals.h
--- a/Engine/Source/Runtime/Core/Public/CoreGlobals.h
+++ b/Engine/Source/Runtime/Core/Public/CoreGlobals.h
@@ -25,6 +25,7 @@ CORE_API DECLARE_LOG_CATEGORY_EXTERN(LogPS4, Log, All);
 CORE_API DECLARE_LOG_CATEGORY_EXTERN(LogXboxOne, Log, All);
 CORE_API DECLARE_LOG_CATEGORY_EXTERN(LogWindows, Log, All);
 CORE_API DECLARE_LOG_CATEGORY_EXTERN(LogSwitch, Log, All);
+CORE_API DECLARE_LOG_CATEGORY_EXTERN(LogQuail, Log, All);
 CORE_API DECLARE_LOG_CATEGORY_EXTERN(LogSerialization, Log, All);
 CORE_API DECLARE_LOG_CATEGORY_EXTERN(LogUnrealMath, Log, All);
 CORE_API DECLARE_LOG_CATEGORY_EXTERN(LogUnrealMatrix, Log, All);
diff --git a/Engine/Source/Runtime/Core/Public/GenericPlatform/GenericPlatformMisc.h b/Engine/Source/Runtime/Core/Public/GenericPlatform/GenericPlatformMisc.h
--- a/Engine/Source/Runtime/Core/Public/GenericPlatform/GenericPlatformMisc.h
+++ b/Engine/Source/Runtime/Core/Public/GenericPlatform/GenericPlatformMisc.h
@@ -717,11 +717,6 @@ public:
 	/** Function to store the current working directory for use with LaunchDir() */
 	static void CacheLaunchDir();
 
-	static bool IsAllowedToRender()
-	{
-		return true;
-	}
-
 	/**
 	 *	Return the project directory
 	 */
@@ -1074,13 +1069,6 @@ public:
 		return true;
 	}
 
-	/**
-	 * Returns the default OSS that this platform uses. This is pulled out of OSS so that we don't have to put confidential
-	 * platform information into shared code.
-	 */
-	static const class FName& GetDefaultOnlineSubsystemName();
-
-
 #if !UE_BUILD_SHIPPING
 	/** 
 	 * Returns any platform specific warning messages we want printed on screen
diff --git a/Engine/Source/Runtime/Core/Public/HAL/Platform.h b/Engine/Source/Runtime/Core/Public/HAL/Platform.h
--- a/Engine/Source/Runtime/Core/Public/HAL/Platform.h
+++ b/Engine/Source/Runtime/Core/Public/HAL/Platform.h
@@ -40,6 +40,9 @@
 #if !defined(PLATFORM_ANDROID_VULKAN)
 	#define PLATFORM_ANDROID_VULKAN 0
 #endif
+#if !defined(PLATFORM_QUAIL)
+	#define PLATFORM_QUAIL 0
+#endif
 #if !defined(PLATFORM_ANDROIDESDEFERRED)
 	#define PLATFORM_ANDROIDESDEFERRED 0
 #endif
@@ -79,6 +82,8 @@
 	#include "HTML5/HTML5PlatformCompilerPreSetup.h"
 #elif PLATFORM_LINUX
 	#include "Linux/LinuxPlatformCompilerPreSetup.h"
+#elif PLATFORM_QUAIL
+	#include "Quail/QuailPlatformCompilerPreSetup.h"
 #elif PLATFORM_SWITCH
 	#include "Switch/SwitchPlatformCompilerPreSetup.h"
 #else
@@ -115,6 +120,8 @@
 	#include "HTML5/HTML5Platform.h"
 #elif PLATFORM_LINUX
 	#include "Linux/LinuxPlatform.h"
+#elif PLATFORM_QUAIL
+	#include "Quail/QuailPlatform.h"
 #elif PLATFORM_SWITCH
 	#include "Switch/SwitchPlatform.h"
 #else
@@ -845,6 +852,8 @@ namespace TypeTests
 	#include "HTML5/HTML5PlatformCompilerSetup.h"
 #elif PLATFORM_LINUX
 	#include "Linux/LinuxPlatformCompilerSetup.h"
+#elif PLATFORM_QUAIL
+	#include "Quail/QuailPlatformCompilerSetup.h"
 #elif PLATFORM_SWITCH
 	#include "Switch/SwitchPlatformCompilerSetup.h"
 #else
diff --git a/Engine/Source/Runtime/Core/Public/HAL/PlatformMisc.h b/Engine/Source/Runtime/Core/Public/HAL/PlatformMisc.h
--- a/Engine/Source/Runtime/Core/Public/HAL/PlatformMisc.h
+++ b/Engine/Source/Runtime/Core/Public/HAL/PlatformMisc.h
@@ -18,6 +18,8 @@
 #include "Android/AndroidMisc.h"
 #elif PLATFORM_HTML5
 #include "HTML5/HTML5PlatformMisc.h"
+#elif PLATFORM_QUAIL
+#include "Quail/QuailPlatformMisc.h"
 #elif PLATFORM_LINUX
 #include "Linux/LinuxPlatformMisc.h"
 #elif PLATFORM_SWITCH
diff --git a/Engine/Source/Runtime/Core/Public/HAL/PlatformProcess.h b/Engine/Source/Runtime/Core/Public/HAL/PlatformProcess.h
--- a/Engine/Source/Runtime/Core/Public/HAL/PlatformProcess.h
+++ b/Engine/Source/Runtime/Core/Public/HAL/PlatformProcess.h
@@ -16,6 +16,8 @@
 #include "IOS/IOSPlatformProcess.h"
 #elif PLATFORM_ANDROID
 #include "Android/AndroidProcess.h"
+#elif PLATFORM_QUAIL
+#include "Quail/QuailPlatformProcess.h"
 #elif PLATFORM_HTML5
 #include "HTML5/HTML5PlatformProcess.h"
 #elif PLATFORM_LINUX
diff --git a/Engine/Source/Runtime/Core/Public/HAL/PlatformProperties.h b/Engine/Source/Runtime/Core/Public/HAL/PlatformProperties.h
--- a/Engine/Source/Runtime/Core/Public/HAL/PlatformProperties.h
+++ b/Engine/Source/Runtime/Core/Public/HAL/PlatformProperties.h
@@ -25,6 +25,9 @@ typedef FAndroidPlatformProperties FPlatformProperties;
 #elif PLATFORM_HTML5
 #include "HTML5/HTML5PlatformProperties.h"
 typedef FHTML5PlatformProperties FPlatformProperties;
+#elif PLATFORM_QUAIL
+#include "Quail/QuailPlatformProperties.h"
+typedef FQuailPlatformProperties FPlatformProperties;
 #elif PLATFORM_LINUX
 #include "Linux/LinuxPlatformProperties.h"
 typedef FLinuxPlatformProperties<WITH_EDITORONLY_DATA, UE_SERVER, !WITH_SERVER_CODE> FPlatformProperties;
diff --git a/Engine/Source/Runtime/Core/Public/IOS/IOSPlatformMisc.h b/Engine/Source/Runtime/Core/Public/IOS/IOSPlatformMisc.h
--- a/Engine/Source/Runtime/Core/Public/IOS/IOSPlatformMisc.h
+++ b/Engine/Source/Runtime/Core/Public/IOS/IOSPlatformMisc.h
@@ -9,8 +9,6 @@
 #include "IOS/IOSSystemIncludes.h"
 #include "Apple/ApplePlatformMisc.h"
 
-#define IOS_SUBSYSTEM FName(TEXT("IOS"))
-
 /**
 * iOS implementation of the misc OS functions
 **/
@@ -47,9 +45,6 @@ struct CORE_API FIOSPlatformMisc : public FApplePlatformMisc
 
 	static class IPlatformChunkInstall* GetPlatformChunkInstall();
 	
-	static const class FName& GetDefaultOnlineSubsystemName();
-
-
 	//////// Platform specific
 	static int GetDefaultStackSize();
 	static void HandleLowMemoryWarning();
diff --git a/Engine/Source/Runtime/Engine/Private/GameViewportClient.cpp b/Engine/Source/Runtime/Engine/Private/GameViewportClient.cpp
--- a/Engine/Source/Runtime/Engine/Private/GameViewportClient.cpp
+++ b/Engine/Source/Runtime/Engine/Private/GameViewportClient.cpp
@@ -1367,7 +1367,7 @@ void UGameViewportClient::Draw(FViewport* InViewport, FCanvas* SceneCanvas)
 	}
 
 	// Draw the player views.
-	if (!bDisableWorldRendering && !bUIDisableWorldRendering && PlayerViewMap.Num() > 0 && FPlatformMisc::IsAllowedToRender()) //-V560
+	if (!bDisableWorldRendering && !bUIDisableWorldRendering && PlayerViewMap.Num() > 0 && FSlateApplication::Get().GetPlatformApplication()->IsAllowedToRender()) //-V560
 	{
 		GetRendererModule().BeginRenderingViewFamily(SceneCanvas,&ViewFamily);
 	}
diff --git a/Engine/Source/Runtime/InputCore/Public/HAL/PlatformInput.h b/Engine/Source/Runtime/InputCore/Public/HAL/PlatformInput.h
--- a/Engine/Source/Runtime/InputCore/Public/HAL/PlatformInput.h
+++ b/Engine/Source/Runtime/InputCore/Public/HAL/PlatformInput.h
@@ -18,6 +18,8 @@
 #include "Android/AndroidPlatformInput.h"
 #elif PLATFORM_HTML5
 #include "HTML5/HTML5PlatformInput.h"
+#elif PLATFORM_QUAIL
+#include "Quail/QuailPlatformInput.h"
 #elif PLATFORM_LINUX
 #include "Linux/LinuxPlatformInput.h"
 #elif PLATFORM_SWITCH
diff --git a/Engine/Source/Runtime/VulkanRHI/Private/VulkanRHI.cpp b/Engine/Source/Runtime/VulkanRHI/Private/VulkanRHI.cpp
--- a/Engine/Source/Runtime/VulkanRHI/Private/VulkanRHI.cpp
+++ b/Engine/Source/Runtime/VulkanRHI/Private/VulkanRHI.cpp
@@ -395,6 +395,11 @@ void FVulkanDynamicRHI::Init()
 
 void FVulkanDynamicRHI::Shutdown()
 {
+	if (FParse::Param(FCommandLine::Get(), TEXT("savevulkanpsocacheonexit")))
+	{
+		SavePipelineCache();
+	}
+
 	check(IsInGameThread() && IsInRenderingThread());
 	check(Device);
 
