Fix Mac CIS (unused variable in C# code).

diff --git a/Engine/Source/Programs/AutomationTool/Scripts/RecordPerformance.Automation.cs b/Engine/Source/Programs/AutomationTool/Scripts/RecordPerformance.Automation.cs
--- a/Engine/Source/Programs/AutomationTool/Scripts/RecordPerformance.Automation.cs
+++ b/Engine/Source/Programs/AutomationTool/Scripts/RecordPerformance.Automation.cs
@@ -29,7 +29,6 @@ class RecordPerformance : BuildCommand
 
 	private string ProjectToRun;
 	private DirectoryReference RootDir;
-	private DirectoryReference ProjectPath;
 	private const string PerformanceMonitorConfig = "Default";
 	private DirectoryReference OutputDir;
 
@@ -42,7 +41,6 @@ class RecordPerformance : BuildCommand
 		Index = Math.Max(0, Index);
 		Index = Math.Min(ProjectsToRun.Length - 1, Index);
 		ProjectToRun = ProjectsToRun[Index];
-		ProjectPath = DirectoryReference.Combine(RootDir, ProjectPathsToRun[Index]);
 
 		// Which platforms to test?
 		var PlatformsToTest = ParseMultipleParams("Platforms");
