Remove unused variable
#rnx

diff --git a/Engine/Source/Programs/AutomationTool/Scripts/RecordPerformance.Automation.cs b/Engine/Source/Programs/AutomationTool/Scripts/RecordPerformance.Automation.cs
--- a/Engine/Source/Programs/AutomationTool/Scripts/RecordPerformance.Automation.cs
+++ b/Engine/Source/Programs/AutomationTool/Scripts/RecordPerformance.Automation.cs
@@ -19,14 +19,6 @@ class RecordPerformance : BuildCommand
 		"ShowdownDemo"			// 3
 	};
 
-	private string[] ProjectPathsToRun =
-	{
-		"Samples/Showcases/SubwaySequencer",
-		"Samples/Showcases/InfiltratorDemo",
-		"Samples/Showcases/ElementalDemo",
-		"Samples/Showcases/ShowdownDemo"
-	};
-
 	private string ProjectToRun;
 	private DirectoryReference RootDir;
 	private const string PerformanceMonitorConfig = "Default";
