#rnx

diff --git a/Engine/Source/Programs/AutomationTool/Scripts/CopyBuildToStagingDirectory.Automation.cs b/Engine/Source/Programs/AutomationTool/Scripts/CopyBuildToStagingDirectory.Automation.cs
--- a/Engine/Source/Programs/AutomationTool/Scripts/CopyBuildToStagingDirectory.Automation.cs
+++ b/Engine/Source/Programs/AutomationTool/Scripts/CopyBuildToStagingDirectory.Automation.cs
@@ -775,6 +775,22 @@ public partial class Project : CommandUtils
 						}
 					}
 				}
+
+				// UE-58423
+				DirectoryReference CookOutputDir;
+				if (!String.IsNullOrEmpty(Params.CookOutputDir))
+				{
+					CookOutputDir = DirectoryReference.Combine(new DirectoryReference(Params.CookOutputDir), SC.CookPlatform);
+				}
+				else if (Params.CookInEditor)
+				{
+					CookOutputDir = DirectoryReference.Combine(SC.ProjectRoot, "Saved", "EditorCooked", SC.CookPlatform);
+				}
+				else
+				{
+					CookOutputDir = DirectoryReference.Combine(SC.ProjectRoot, "Saved", "Cooked", SC.CookPlatform);
+				}
+				SC.MetadataDir = DirectoryReference.Combine(CookOutputDir, SC.ShortProjectName, "Metadata");
 			}
 		}
 
