UAT: don't produce .pak file while cooking on the fly (merged from Dev-Console)

diff --git a/Engine/Source/Programs/AutomationTool/Scripts/CopyBuildToStagingDirectory.Automation.cs b/Engine/Source/Programs/AutomationTool/Scripts/CopyBuildToStagingDirectory.Automation.cs
--- a/Engine/Source/Programs/AutomationTool/Scripts/CopyBuildToStagingDirectory.Automation.cs
+++ b/Engine/Source/Programs/AutomationTool/Scripts/CopyBuildToStagingDirectory.Automation.cs
@@ -2012,6 +2012,11 @@ public partial class Project : CommandUtils
 
 	private static bool ShouldCreatePak(ProjectParams Params, DeploymentContext SC)
 	{
+		if (Params.CookOnTheFly)
+		{
+			return false;
+		}
+
 		Platform.PakType Pak = SC.StageTargetPlatform.RequiresPak(Params);
 
 		// we may care but we don't want. 
