Fix editor targets not being rebuilt when packaging from an installed engine build.

diff --git a/Engine/Source/Programs/AutomationTool/Scripts/BuildProjectCommand.Automation.cs b/Engine/Source/Programs/AutomationTool/Scripts/BuildProjectCommand.Automation.cs
--- a/Engine/Source/Programs/AutomationTool/Scripts/BuildProjectCommand.Automation.cs
+++ b/Engine/Source/Programs/AutomationTool/Scripts/BuildProjectCommand.Automation.cs
@@ -1,4 +1,4 @@
-﻿// Copyright 1998-2018 Epic Games, Inc. All Rights Reserved.
+// Copyright 1998-2018 Epic Games, Inc. All Rights Reserved.
 using System;
 using System.Collections.Generic;
 using System.IO;
@@ -69,29 +69,33 @@ public partial class Project : CommandUtils
 		var CrashReportPlatforms = new HashSet<UnrealTargetPlatform>();
 
 		// Setup editor targets
-		if (Params.HasEditorTargets && (!Params.SkipBuildEditor) && !Automation.IsEngineInstalled() && (TargetMask & ProjectBuildTargets.Editor) == ProjectBuildTargets.Editor)
+		if (Params.HasEditorTargets && (!Params.SkipBuildEditor) && (TargetMask & ProjectBuildTargets.Editor) == ProjectBuildTargets.Editor)
 		{
 			// @todo Mac: proper platform detection
 			UnrealTargetPlatform EditorPlatform = HostPlatform.Current.HostEditorPlatform;
 			const UnrealTargetConfiguration EditorConfiguration = UnrealTargetConfiguration.Development;
 
-			CrashReportPlatforms.Add(EditorPlatform);
             Agenda.AddTargets(Params.EditorTargets.ToArray(), EditorPlatform, EditorConfiguration, Params.CodeBasedUprojectPath);
-			if (Params.EditorTargets.Contains("UnrealHeaderTool") == false)
-			{
-				Agenda.AddTargets(new string[] { "UnrealHeaderTool" }, EditorPlatform, EditorConfiguration);
-			}
-			if (Params.EditorTargets.Contains("ShaderCompileWorker") == false)
-			{
-				Agenda.AddTargets(new string[] { "ShaderCompileWorker" }, EditorPlatform, EditorConfiguration);
-			}
-			if (Params.Pak && Params.EditorTargets.Contains("UnrealPak") == false)
-			{
-				Agenda.AddTargets(new string[] { "UnrealPak" }, EditorPlatform, EditorConfiguration);
-			}
-			if (Params.FileServer && Params.EditorTargets.Contains("UnrealFileServer") == false)
+
+			if(!Automation.IsEngineInstalled())
 			{
-				Agenda.AddTargets(new string[] { "UnrealFileServer" }, EditorPlatform, EditorConfiguration);
+				CrashReportPlatforms.Add(EditorPlatform);
+				if (Params.EditorTargets.Contains("UnrealHeaderTool") == false)
+				{
+					Agenda.AddTargets(new string[] { "UnrealHeaderTool" }, EditorPlatform, EditorConfiguration);
+				}
+				if (Params.EditorTargets.Contains("ShaderCompileWorker") == false)
+				{
+					Agenda.AddTargets(new string[] { "ShaderCompileWorker" }, EditorPlatform, EditorConfiguration);
+				}
+				if (Params.Pak && Params.EditorTargets.Contains("UnrealPak") == false)
+				{
+					Agenda.AddTargets(new string[] { "UnrealPak" }, EditorPlatform, EditorConfiguration);
+				}
+				if (Params.FileServer && Params.EditorTargets.Contains("UnrealFileServer") == false)
+				{
+					Agenda.AddTargets(new string[] { "UnrealFileServer" }, EditorPlatform, EditorConfiguration);
+				}
 			}
 		}
 
