UE-55180 Cook on the Fly server quits unexpectedly using Windows UFE to iOS

#4.20

Check if the CookOutputDir exists to prevent exception on EnumerateFiles().ToList()

diff --git a/Engine/Source/Programs/AutomationTool/IOS/IOSPlatform.Automation.cs b/Engine/Source/Programs/AutomationTool/IOS/IOSPlatform.Automation.cs
--- a/Engine/Source/Programs/AutomationTool/IOS/IOSPlatform.Automation.cs
+++ b/Engine/Source/Programs/AutomationTool/IOS/IOSPlatform.Automation.cs
@@ -954,10 +954,13 @@ public class IOSPlatform : Platform
 			{
 				CookOutputDir = DirectoryReference.Combine(SC.ProjectRoot, "Saved", "Cooked", SC.CookPlatform);
 			}
-			List<FileReference> CookedFiles = DirectoryReference.EnumerateFiles(CookOutputDir, "*.metallib", SearchOption.AllDirectories).ToList();
-			foreach(FileReference CookedFile in CookedFiles)
+			if (DirectoryReference.Exists(CookOutputDir))
 			{
-				SC.StageFile(StagedFileType.NonUFS, CookedFile, new StagedFileReference(CookedFile.MakeRelativeTo(CookOutputDir)));
+				List<FileReference> CookedFiles = DirectoryReference.EnumerateFiles(CookOutputDir, "*.metallib", SearchOption.AllDirectories).ToList();
+				foreach (FileReference CookedFile in CookedFiles)
+				{
+					SC.StageFile(StagedFileType.NonUFS, CookedFile, new StagedFileReference(CookedFile.MakeRelativeTo(CookOutputDir)));
+				}
 			}
 		}
     }
