restricting latter agents with names multiply defined to either have no agent type or match the former exactly

diff --git a/Engine/Source/Programs/AutomationTool/BuildGraph/Script.cs b/Engine/Source/Programs/AutomationTool/BuildGraph/Script.cs
--- a/Engine/Source/Programs/AutomationTool/BuildGraph/Script.cs
+++ b/Engine/Source/Programs/AutomationTool/BuildGraph/Script.cs
@@ -794,12 +794,10 @@ namespace AutomationTool
 				{
 					if (Types.Length > 0 && Agent.PossibleTypes.Length > 0)
 					{
-						string[] NewTypes = Agent.PossibleTypes.Intersect(Types, StringComparer.InvariantCultureIgnoreCase).ToArray();
-						if (NewTypes.Length == 0)
+						if (Types.Length != Agent.PossibleTypes.Length || !Types.SequenceEqual(Agent.PossibleTypes, StringComparer.InvariantCultureIgnoreCase))
 						{
-							LogError(Element, "No common agent types with previous agent definition");
+							LogError(Element, "Agent types ({0}) were different than previous agent definition with types ({1}). Must either be empty or match exactly.", string.Join(",", Types), string.Join(",", Agent.PossibleTypes));
 						}
-						Agent.PossibleTypes = NewTypes;
 					}
 				}
 				else
