Fix install/uninstall batch filenames for Android BP projects
#ue4
#android

diff --git a/Engine/Source/Programs/AutomationTool/Android/AndroidPlatform.Automation.cs b/Engine/Source/Programs/AutomationTool/Android/AndroidPlatform.Automation.cs
--- a/Engine/Source/Programs/AutomationTool/Android/AndroidPlatform.Automation.cs
+++ b/Engine/Source/Programs/AutomationTool/Android/AndroidPlatform.Automation.cs
@@ -223,13 +223,16 @@ public class AndroidPlatform : Platform
 				break;
 		}
 
+		// Get the name of the APK to use for batch file
+		string ExecutableName = Path.GetFileNameWithoutExtension(ApkName);
+		
 		switch(BatchType)
 		{
 			case EBatchType.Install:
 			case EBatchType.Uninstall:
-				return Path.Combine(Path.GetDirectoryName(ApkName), (BatchType == EBatchType.Uninstall ? "Uninstall_" : "Install_") + SC.StageExecutables[0] + (!bNoOBBInstall ? "" : "_NoOBBInstall") + Architecture + GPUArchitecture + Extension);
+				return Path.Combine(Path.GetDirectoryName(ApkName), (BatchType == EBatchType.Uninstall ? "Uninstall_" : "Install_") + ExecutableName + (!bNoOBBInstall ? "" : "_NoOBBInstall") + Extension);
 			case EBatchType.Symbolize:
-				return Path.Combine(Path.GetDirectoryName(ApkName), "SymbolizeCrashDump_" + SC.StageExecutables[0] + Architecture + GPUArchitecture + Extension);
+				return Path.Combine(Path.GetDirectoryName(ApkName), "SymbolizeCrashDump_" + ExecutableName + Extension);
 		}
 		return "";
 	}
