Need to build symbols for phyx/apex when they rebuild for Linux

diff --git a/Engine/Source/Programs/AutomationTool/Scripts/BuildPhysX.Automation.cs b/Engine/Source/Programs/AutomationTool/Scripts/BuildPhysX.Automation.cs
--- a/Engine/Source/Programs/AutomationTool/Scripts/BuildPhysX.Automation.cs
+++ b/Engine/Source/Programs/AutomationTool/Scripts/BuildPhysX.Automation.cs
@@ -76,6 +76,9 @@ class BuildPhysX : BuildCommand
 	private static DirectoryReference RootOutputLibDirectory = DirectoryReference.Combine(PhysXSourceRootDirectory, "Lib");
 	private static DirectoryReference ThirdPartySourceDirectory = DirectoryReference.Combine(CommandUtils.RootDirectory, "Engine", "Source", "ThirdParty");
 
+	private static DirectoryReference DumpSymsPath = DirectoryReference.Combine(CommandUtils.RootDirectory, "Engine", "Binaries", "Linux", "dump_syms");
+	private static DirectoryReference BreakpadSymbolEncoderPath = DirectoryReference.Combine(CommandUtils.RootDirectory, "Engine", "Binaries", "Linux", "BreakpadSymbolEncoder");
+
 	//private static DirectoryReference PhysX34SourceLibRootDirectory = DirectoryReference.Combine(PhysX34SourceRootDirectory, "Lib");
 	//private static DirectoryReference APEX14SourceLibRootDirectory = DirectoryReference.Combine(APEX14SourceRootDirectory, "Lib");
 	//private static DirectoryReference SharedSourceLibRootDirectory = DirectoryReference.Combine(SharedSourceRootDirectory, "Lib");
@@ -1315,7 +1318,7 @@ class BuildPhysX : BuildCommand
 
     private static bool FileGeneratedByNvCloth(string FileNameUpper)
     {
-        if(FileNameUpper.StartsWith("NVCLOTH"))
+		if (FileNameUpper.Contains("NVCLOTH"))
         {
             return true;
         }
@@ -1325,7 +1328,7 @@ class BuildPhysX : BuildCommand
     
 	private static bool FileGeneratedByAPEX(string FileNameUpper)
 	{
-		if (FileNameUpper.StartsWith("APEX"))
+		if (FileNameUpper.Contains("APEX"))
 		{
 			return true;
 		}
@@ -1333,7 +1336,7 @@ class BuildPhysX : BuildCommand
 		{
 			foreach (string SpecialApexLib in APEXSpecialLibs)
 			{
-				if (FileNameUpper.StartsWith(SpecialApexLib.ToUpper()))	//There are some APEX libs that don't use the APEX prefix so make sure to test against it
+				if (FileNameUpper.Contains(SpecialApexLib.ToUpper()))	//There are some APEX libs that don't use the APEX prefix so make sure to test against it
 				{
 					return true;
 				}
@@ -1345,7 +1348,7 @@ class BuildPhysX : BuildCommand
 
 	private static void FindOutputFilesHelper(HashSet<FileReference> OutputFiles, DirectoryReference BaseDir, string SearchPrefix, PhysXTargetLib TargetLib)
 	{
-		if(!DirectoryReference.Exists(BaseDir))
+		if (!DirectoryReference.Exists(BaseDir))
 		{
 			return;
 		}
@@ -1361,6 +1364,56 @@ class BuildPhysX : BuildCommand
 		}
 	}
 
+
+	private static void GenerateDebugFiles(HashSet<FileReference> OutFiles, PhysXTargetLib TargetLib, TargetPlatformData TargetData, string TargetConfiguration, WindowsCompiler TargetWindowsCompiler = WindowsCompiler.VisualStudio2015)
+	{
+		if (TargetData.Platform == UnrealTargetPlatform.Linux)
+		{
+			HashSet<FileReference> SoFiles = new HashSet<FileReference>();
+
+			string SearchSuffix = GetConfigurationSuffix(TargetConfiguration, TargetData).ToUpper();
+			string SearchPrefix = "*" + SearchSuffix + ".";
+
+			DirectoryReference BinaryDir = GetPlatformBinaryDirectory(TargetData, TargetWindowsCompiler);
+			FindOutputFilesHelper(SoFiles, BinaryDir, SearchPrefix + GetPlatformBinaryExtension(TargetData), TargetLib);
+
+			foreach (FileReference SOFile in SoFiles)
+			{
+				string ExeSuffix = "";
+				if (BuildHostPlatform.Current.Platform == UnrealTargetPlatform.Win32 || BuildHostPlatform.Current.Platform == UnrealTargetPlatform.Win64)
+				{
+					ExeSuffix += ".exe";
+				}
+
+				FileReference PSymbolFile = FileReference.Combine(SOFile.Directory, SOFile.GetFileNameWithoutExtension() + ".psym");
+				FileReference SymbolFile = FileReference.Combine(SOFile.Directory, SOFile.GetFileNameWithoutExtension() + ".sym");
+
+				// dump_syms
+				ProcessStartInfo StartInfo = new ProcessStartInfo();
+				StartInfo.FileName = DumpSymsPath.ToString() + ExeSuffix;
+				StartInfo.Arguments = SOFile.FullName + " " + PSymbolFile.ToString();
+				StartInfo.RedirectStandardError = true;
+
+				Log("Running: '{0} {1}'", StartInfo.FileName, StartInfo.Arguments);
+
+				RunLocalProcessAndLogOutput(StartInfo);
+
+				// BreakpadSymbolEncoder
+				StartInfo.FileName = BreakpadSymbolEncoderPath.ToString() + ExeSuffix;
+				StartInfo.Arguments = PSymbolFile.ToString() + " " + SymbolFile.ToString();
+
+				Log("Running: '{0} {1}'", StartInfo.FileName, StartInfo.Arguments);
+
+				RunLocalProcessAndLogOutput(StartInfo);
+
+				// Clean up the Temp *.psym file, as they are no longer needed
+				InternalUtils.SafeDeleteFile(PSymbolFile.ToString());
+
+				OutFiles.Add(SymbolFile);
+			}
+		}
+	}
+
 	private static string GetConfigurationSuffix(string TargetConfiguration, TargetPlatformData TargetData)
 	{
 		// default
@@ -1557,7 +1610,7 @@ class BuildPhysX : BuildCommand
 				// build target lib for all platforms
 				foreach (TargetPlatformData TargetData in TargetPlatforms)
 				{
-					if(!PlatformSupportsTargetLib(TargetLib, TargetData))
+					if (!PlatformSupportsTargetLib(TargetLib, TargetData))
 					{
 						continue;
 					}
@@ -1580,14 +1633,18 @@ class BuildPhysX : BuildCommand
 								break;
 						}
 					}
-					
-					foreach(FileReference FileToDelete in FilesToDelete)
+					foreach (FileReference FileToDelete in FilesToDelete)
 					{
 						FilesToReconcile.Add(FileToDelete);
 						InternalUtils.SafeDeleteFile(FileToDelete.ToString());
 					}
 
 					BuildTargetLibForPlatform(TargetLib, TargetData, TargetConfigurations, TargetWindowsCompilers);
+
+					foreach (string TargetConfiguration in TargetConfigurations)
+					{
+						GenerateDebugFiles(FilesToReconcile, TargetLib, TargetData, TargetConfiguration);
+					}
 				}
 			}
 		}
@@ -1640,20 +1697,20 @@ class BuildPhysX : BuildCommand
 						{
 							case UnrealTargetPlatform.Win32:
 							case UnrealTargetPlatform.Win64:
-										foreach (WindowsCompiler TargetCompiler in TargetWindowsCompilers)
+								foreach (WindowsCompiler TargetCompiler in TargetWindowsCompilers)
 								{
-											FindOutputFiles(FilesToReconcile, TargetLib, TargetData, TargetConfiguration, TargetCompiler);
+									FindOutputFiles(FilesToReconcile, TargetLib, TargetData, TargetConfiguration, TargetCompiler);
 								}
 								break;
 							default:
-										FindOutputFiles(FilesToReconcile, TargetLib, TargetData, TargetConfiguration);
+								FindOutputFiles(FilesToReconcile, TargetLib, TargetData, TargetConfiguration);
 								break;
 						}
 					}
 				}
 			}
-			
-			foreach(FileReference FileToReconcile in FilesToReconcile)
+
+			foreach (FileReference FileToReconcile in FilesToReconcile)
 			{
 				P4.Reconcile(P4ChangeList, FileToReconcile.ToString());
 			}
