Fix a deadlock caused when redirecting stderr but not reading from it after starting a process

diff --git a/Engine/Source/Programs/AutomationTool/Scripts/BuildPhysX.Automation.cs b/Engine/Source/Programs/AutomationTool/Scripts/BuildPhysX.Automation.cs
--- a/Engine/Source/Programs/AutomationTool/Scripts/BuildPhysX.Automation.cs
+++ b/Engine/Source/Programs/AutomationTool/Scripts/BuildPhysX.Automation.cs
@@ -560,6 +560,12 @@ class BuildPhysX : BuildCommand
 				// Start the process up and then wait for it to finish
 				LocalProcess.Start();
 				LocalProcess.BeginOutputReadLine();
+
+				if (LocalProcess.StartInfo.RedirectStandardError)
+				{
+					LocalProcess.BeginErrorReadLine();
+				}
+
 				LocalProcess.WaitForExit();
 				ExitCode = LocalProcess.ExitCode;
 			}
@@ -1244,6 +1250,7 @@ class BuildPhysX : BuildCommand
 			case UnrealTargetPlatform.Win32:
 			case UnrealTargetPlatform.Win64:
 //			case UnrealTargetPlatform.Mac:
+			case UnrealTargetPlatform.Linux:
 			case UnrealTargetPlatform.XboxOne:
 				return true;
 		}
@@ -1259,6 +1266,8 @@ class BuildPhysX : BuildCommand
 				return "pdb";
 			case UnrealTargetPlatform.Mac:
 				return "dSYM";
+			case UnrealTargetPlatform.Linux:
+				return "sym";
 		}
 		throw new AutomationException(String.Format("No debug database extension for platform '{0}'", TargetData.Platform.ToString()));
 	}
