XLOC now gets the latest build for the specified file, rather than the latest build in general

The old functionality could return a build that didn't include the required file, which could lead to translations being missed

diff --git a/Engine/Source/Programs/AutomationTool/XLocLocalization/XLocLocalizationProvider.cs b/Engine/Source/Programs/AutomationTool/XLocLocalization/XLocLocalizationProvider.cs
--- a/Engine/Source/Programs/AutomationTool/XLocLocalization/XLocLocalizationProvider.cs
+++ b/Engine/Source/Programs/AutomationTool/XLocLocalization/XLocLocalizationProvider.cs
@@ -94,13 +94,18 @@ public abstract class XLocLocalizationProvider : LocalizationProvider
 		try
 		{
 			var EpicCultureToXLocLanguageId = GetEpicCultureToXLocLanguageId();
-			LatestBuildXml = RequestLatestBuild(XLocApiClient, AuthToken, EpicCultureToXLocLanguageId[Culture]);
+			LatestBuildXml = RequestLatestBuild(XLocApiClient, AuthToken, EpicCultureToXLocLanguageId[Culture], XLocFilename);
 		}
 		catch (Exception Ex)
 		{
 			BuildCommand.LogWarning("RequestLatestBuild failed for {0}. {1}", Culture, Ex);
 			return;
 		}
+		if (String.IsNullOrEmpty(LatestBuildXml))
+		{
+			Console.WriteLine("[IGNORED] '{0}' has no build data ({1})", XLocFilename, Culture);
+			return;
+		}
 
 		var POFileUri = "";
 		var BuildsXmlDoc = new XmlDocument();
@@ -354,9 +359,9 @@ public abstract class XLocLocalizationProvider : LocalizationProvider
 		}
 	}
 
-	private string RequestLatestBuild(XLocApiClient XLocApiClient, string AuthToken, string LanguageId)
+	private string RequestLatestBuild(XLocApiClient XLocApiClient, string AuthToken, string LanguageId, string RemoteFilename)
 	{
-		return XLocApiClient.GetLatestBuild(Config.APIKey, AuthToken, XLocUtils.MD5HashString(Config.APIKey + Config.APISecret + Config.LocalizationId + LanguageId), Config.LocalizationId, LanguageId);
+		return XLocApiClient.GetLatestBuildByFile(Config.APIKey, AuthToken, XLocUtils.MD5HashString(Config.APIKey + Config.APISecret + Config.LocalizationId + LanguageId + RemoteFilename), Config.LocalizationId, LanguageId, RemoteFilename);
 	}
 
 	private string GetXLocFilename(string BaseFilename)
diff --git a/Engine/Source/Programs/AutomationTool/XLocLocalization/XLocService/XLocApiService.cs b/Engine/Source/Programs/AutomationTool/XLocLocalization/XLocService/XLocApiService.cs
--- a/Engine/Source/Programs/AutomationTool/XLocLocalization/XLocService/XLocApiService.cs
+++ b/Engine/Source/Programs/AutomationTool/XLocLocalization/XLocService/XLocApiService.cs
@@ -96,8 +96,6 @@ namespace XLoc.Xml.Types
         
         private string LanguageNameField;
         
-        private System.Guid[] LocalizationIdsField;
-        
         public System.Runtime.Serialization.ExtensionDataObject ExtensionData
         {
             get
@@ -135,19 +133,6 @@ namespace XLoc.Xml.Types
                 this.LanguageNameField = value;
             }
         }
-        
-        [System.Runtime.Serialization.DataMemberAttribute()]
-        public System.Guid[] LocalizationIds
-        {
-            get
-            {
-                return this.LocalizationIdsField;
-            }
-            set
-            {
-                this.LocalizationIdsField = value;
-            }
-        }
     }
     
     [System.Diagnostics.DebuggerStepThroughAttribute()]
@@ -1752,7 +1737,6 @@ public partial class ContextListData : object, System.Runtime.Serialization.IExt
 [System.Runtime.Serialization.KnownTypeAttribute(typeof(XLoc.Xml.Types.ContextType))]
 [System.Runtime.Serialization.KnownTypeAttribute(typeof(XLoc.Xml.Types.LocalizedString[]))]
 [System.Runtime.Serialization.KnownTypeAttribute(typeof(XLoc.Xml.Types.SearchCriteria))]
-[System.Runtime.Serialization.KnownTypeAttribute(typeof(System.Guid[]))]
 [System.Runtime.Serialization.KnownTypeAttribute(typeof(ContextListData[]))]
 [System.Runtime.Serialization.KnownTypeAttribute(typeof(ContextListData))]
 [System.Runtime.Serialization.KnownTypeAttribute(typeof(ContextListLine[]))]
@@ -1983,6 +1967,12 @@ public interface IXLocApi
     [System.ServiceModel.OperationContractAttribute(Action="urn:XLoc.Services/IXLocApi/GetLatestBuild", ReplyAction="urn:XLoc.Services/IXLocApi/GetLatestBuildResponse")]
     System.Threading.Tasks.Task<string> GetLatestBuildAsync(string api_key, string auth_token, string api_sig, string localization_id, string language_id);
     
+    [System.ServiceModel.OperationContractAttribute(Action="urn:XLoc.Services/IXLocApi/GetLatestBuildByFile", ReplyAction="urn:XLoc.Services/IXLocApi/GetLatestBuildByFileResponse")]
+    string GetLatestBuildByFile(string api_key, string auth_token, string api_sig, string localization_id, string language_id, string filename);
+    
+    [System.ServiceModel.OperationContractAttribute(Action="urn:XLoc.Services/IXLocApi/GetLatestBuildByFile", ReplyAction="urn:XLoc.Services/IXLocApi/GetLatestBuildByFileResponse")]
+    System.Threading.Tasks.Task<string> GetLatestBuildByFileAsync(string api_key, string auth_token, string api_sig, string localization_id, string language_id, string filename);
+    
     [System.ServiceModel.OperationContractAttribute(Action="urn:XLoc.Services/IXLocApi/Build", ReplyAction="urn:XLoc.Services/IXLocApi/BuildResponse")]
     string Build(string api_key, string auth_token, string api_sig, string localization_id, string language_id, string platform_id);
     
@@ -2200,6 +2190,16 @@ public partial class XLocApiClient : System.ServiceModel.ClientBase<IXLocApi>, I
         return base.Channel.GetLatestBuildAsync(api_key, auth_token, api_sig, localization_id, language_id);
     }
     
+    public string GetLatestBuildByFile(string api_key, string auth_token, string api_sig, string localization_id, string language_id, string filename)
+    {
+        return base.Channel.GetLatestBuildByFile(api_key, auth_token, api_sig, localization_id, language_id, filename);
+    }
+    
+    public System.Threading.Tasks.Task<string> GetLatestBuildByFileAsync(string api_key, string auth_token, string api_sig, string localization_id, string language_id, string filename)
+    {
+        return base.Channel.GetLatestBuildByFileAsync(api_key, auth_token, api_sig, localization_id, language_id, filename);
+    }
+    
     public string Build(string api_key, string auth_token, string api_sig, string localization_id, string language_id, string platform_id)
     {
         return base.Channel.Build(api_key, auth_token, api_sig, localization_id, language_id, platform_id);
