Change scripts for generating documentation to output HTML rather than UDN files. Files are written to Engine/Documentation/Builds/HTML.

diff --git a/Engine/Source/Programs/AutomationTool/BuildGraph/BuildGraph.cs b/Engine/Source/Programs/AutomationTool/BuildGraph/BuildGraph.cs
--- a/Engine/Source/Programs/AutomationTool/BuildGraph/BuildGraph.cs
+++ b/Engine/Source/Programs/AutomationTool/BuildGraph/BuildGraph.cs
@@ -849,27 +849,28 @@ namespace AutomationTool
 			}
 
 			// Create the output directory
-			DirectoryReference.CreateDirectory(OutputFile.Directory);
-			FileReference.MakeWriteable(OutputFile);
-			Log("Writing {0}...", OutputFile);
-
-			// Parse the engine version
-			BuildVersion Version;
-			if(!BuildVersion.TryRead(BuildVersion.GetDefaultFileName(), out Version))
+			if(FileReference.Exists(OutputFile))
 			{
-				throw new AutomationException("Couldn't read Build.version");
+				FileReference.MakeWriteable(OutputFile);
+			}
+			else
+			{
+				DirectoryReference.CreateDirectory(OutputFile.Directory);
 			}
 
 			// Write the output file
+			Log("Writing {0}...", OutputFile);
 			using (StreamWriter Writer = new StreamWriter(OutputFile.FullName))
 			{
-				Writer.WriteLine("Availability: NoPublish");
-				Writer.WriteLine("Title: BuildGraph Predefined Tasks");
-				Writer.WriteLine("Crumbs: %ROOT%, Programming, Programming/Development, Programming/Development/BuildGraph, Programming/Development/BuildGraph/BuildGraphScriptTasks");
-				Writer.WriteLine("Description: This is a procedurally generated markdown page.");
-				Writer.WriteLine("version: {0}.{1}", Version.MajorVersion, Version.MinorVersion);
-				Writer.WriteLine("parent:Programming/Development/BuildGraph/BuildGraphScriptTasks");
-				Writer.WriteLine();
+				Writer.WriteLine("<html>");
+				Writer.WriteLine("  <head>");
+				Writer.WriteLine("    <style>");
+				Writer.WriteLine("      table { border-collapse: collapse }");
+				Writer.WriteLine("      table, th, td { border: 1px solid black; }");
+				Writer.WriteLine("    </style>");
+				Writer.WriteLine("  </head>");
+				Writer.WriteLine("  <body>");
+				Writer.WriteLine("    <h1>BuildGraph Tasks</h1>");
 				foreach(string TaskName in NameToTask.Keys.OrderBy(x => x))
 				{
 					// Get the task object
@@ -880,13 +881,19 @@ namespace AutomationTool
 					if(MemberNameToElement.TryGetValue("T:" + Task.TaskClass.FullName, out TaskElement))
 					{
 						// Write the task heading
-						Writer.WriteLine("### {0}", TaskName);
-						Writer.WriteLine();
-						Writer.WriteLine(ConvertToMarkdown(TaskElement.SelectSingleNode("summary")));
-						Writer.WriteLine();
+						Writer.WriteLine("    <h2>{0}</h2>", TaskName);
+						Writer.WriteLine("    <p>{0}</p>", TaskElement.SelectSingleNode("summary").InnerXml.Trim());
+
+						// Start the parameter table
+						Writer.WriteLine("    <table>");
+						Writer.WriteLine("      <tr>");
+						Writer.WriteLine("        <th>Attribute</th>");
+						Writer.WriteLine("        <th>Type</th>");
+						Writer.WriteLine("        <th>Usage</th>");
+						Writer.WriteLine("        <th>Description</th>");
+						Writer.WriteLine("      </tr>");
 
 						// Document the parameters
-						List<string[]> Rows = new List<string[]>();
 						foreach(string ParameterName in Task.NameToParameter.Keys)
 						{
 							// Get the parameter data
@@ -910,43 +917,29 @@ namespace AutomationTool
 									TypeName = NewTypeName.ToString();
 								}
 
-								string[] Columns = new string[4];
-								Columns[0] = ParameterName;
-								Columns[1] = TypeName;
-								Columns[2] = Parameter.bOptional? "Optional" : "Required";
-								Columns[3] = ConvertToMarkdown(ParameterElement.SelectSingleNode("summary"));
-								Rows.Add(Columns);
+								Writer.WriteLine("      <tr>");
+								Writer.WriteLine("         <td>{0}</td>", ParameterName);
+								Writer.WriteLine("         <td>{0}</td>", TypeName);
+								Writer.WriteLine("         <td>{0}</td>", Parameter.bOptional? "Optional" : "Required");
+								Writer.WriteLine("         <td>{0}</td>", ParameterElement.SelectSingleNode("summary").InnerXml.Trim());
+								Writer.WriteLine("      </tr>");
 							}
 						}
 
 						// Always include the "If" attribute
-						string[] IfColumns = new string[4];
-						IfColumns[0] = "If";
-						IfColumns[1] = "Condition";
-						IfColumns[2] = "Optional";
-						IfColumns[3] = "Whether to execute this task. It is ignored if this condition evaluates to false.";
-						Rows.Add(IfColumns);
-
-						// Get the width of each column
-						int[] Widths = new int[4];
-						for(int Idx = 0; Idx < 4; Idx++)
-						{
-							Widths[Idx] = Rows.Max(x => x[Idx].Length);
-						}
-
-						// Format the markdown table
-						string Format = String.Format("| {{0,-{0}}} | {{1,-{1}}} | {{2,-{2}}} | {{3,-{3}}} |", Widths[0], Widths[1], Widths[2], Widths[3]);
-						Writer.WriteLine(Format, "", "", "", "");
-						Writer.WriteLine(Format, new string('-', Widths[0]), new string('-', Widths[1]), new string('-', Widths[2]), new string('-', Widths[3]));
-						for(int Idx = 0; Idx < Rows.Count; Idx++)
-						{
-							Writer.WriteLine(Format, Rows[Idx][0], Rows[Idx][1], Rows[Idx][2], Rows[Idx][3]);
-						}
-
-						// Blank line before next task
-						Writer.WriteLine();
+						Writer.WriteLine("     <tr>");
+						Writer.WriteLine("       <td>If</td>");
+						Writer.WriteLine("       <td>Condition</td>");
+						Writer.WriteLine("       <td>Optional</td>");
+						Writer.WriteLine("       <td>Whether to execute this task. It is ignored if this condition evaluates to false.</td>");
+						Writer.WriteLine("     </tr>");
+
+						// Close the table
+						Writer.WriteLine("    <table>");
 					}
 				}
+				Writer.WriteLine("  </body>");
+				Writer.WriteLine("</html>");
 			}
 		}
 
diff --git a/Engine/Source/Programs/UnrealBuildTool/System/RulesDocumentation.cs b/Engine/Source/Programs/UnrealBuildTool/System/RulesDocumentation.cs
--- a/Engine/Source/Programs/UnrealBuildTool/System/RulesDocumentation.cs
+++ b/Engine/Source/Programs/UnrealBuildTool/System/RulesDocumentation.cs
@@ -1,4 +1,4 @@
-﻿// Copyright 1998-2018 Epic Games, Inc. All Rights Reserved.
+// Copyright 1998-2018 Epic Games, Inc. All Rights Reserved.
 
 using System;
 using System.Collections.Generic;
@@ -23,13 +23,6 @@ namespace UnrealBuildTool
 				throw new BuildException("Generated assembly documentation not found at {0}.", InputDocumentationFile);
 			}
 
-			// Get the current engine version for versioning the page
-			BuildVersion Version;
-			if(!BuildVersion.TryRead(BuildVersion.GetDefaultFileName(), out Version))
-			{
-				throw new BuildException("Unable to read the current build version");
-			}
-
 			// Read the documentation
 			XmlDocument InputDocumentation = new XmlDocument();
 			InputDocumentation.Load(InputDocumentationFile.FullName);
@@ -53,36 +46,42 @@ namespace UnrealBuildTool
 			}
 
 			// Make sure the output file is writable
-			FileReference.MakeWriteable(OutputFile);
+			if(FileReference.Exists(OutputFile))
+			{
+				FileReference.MakeWriteable(OutputFile);
+			}
+			else
+			{
+				DirectoryReference.CreateDirectory(OutputFile.Directory);
+			}
 
-			// Generate the UDN documentation file
+			// Generate the documentation file
 			using (StreamWriter Writer = new StreamWriter(OutputFile.FullName))
 			{
-				Writer.WriteLine("Availability: NoPublish");
-				Writer.WriteLine("Title: Build Configuration Properties Page");
-				Writer.WriteLine("Crumbs:");
-				Writer.WriteLine("Description: This is a procedurally generated markdown page.");
-				Writer.WriteLine("Version: {0}.{1}", Version.MajorVersion, Version.MinorVersion);
-				Writer.WriteLine("");
+				Writer.WriteLine("<html>");
+				Writer.WriteLine("  <body>");
 				if(ReadOnlyFields.Count > 0)
 				{
-					Writer.WriteLine("### Read-Only Properties");
-					Writer.WriteLine();
+					Writer.WriteLine("    <h2>Read-Only Properties</h2>");
+					Writer.WriteLine("    <dl>");
 					foreach(FieldInfo Field in ReadOnlyFields)
 					{
 						OutputField(InputDocumentation, Field, Writer);
 					}
-					Writer.WriteLine();
+					Writer.WriteLine("    </dl>");
 				}
 				if(ReadWriteFields.Count > 0)
 				{
-					Writer.WriteLine("### Read/Write Properties");
+					Writer.WriteLine("    <h2>Read/Write Properties</h2>");
+					Writer.WriteLine("    <dl>");
 					foreach(FieldInfo Field in ReadWriteFields)
 					{
 						OutputField(InputDocumentation, Field, Writer);
 					}
-					Writer.WriteLine("");
+					Writer.WriteLine("    </dl>");
 				}
+				Writer.WriteLine("  </body>");
+				Writer.WriteLine("</html>");
 			}
 
 			// Success!
@@ -115,22 +114,36 @@ namespace UnrealBuildTool
 					}
 				}
 */
-				// Write the result to the .udn file
+				// Write the result to the HTML file
 				if(Lines.Count > 0)
 				{
-					Writer.WriteLine("$ {0} ({1}): {2}", Field.Name, GetPrettyTypeName(Field.FieldType), Lines[0]);
-					for(int Idx = 1; Idx < Lines.Count; Idx++)
+					Writer.WriteLine("      <dt>{0} ({1})</dt>", Field.Name, GetPrettyTypeName(Field.FieldType));
+
+					if(Lines.Count == 1)
 					{
-						if(Lines[Idx].StartsWith("*") || Lines[Idx].StartsWith("-"))
-						{
-							Writer.WriteLine("        * {0}", Lines[Idx].Substring(1).TrimStart());
-						}
-						else
+						Writer.WriteLine("      <dd>{0}</dd>", Lines[0]);
+					}
+					else
+					{
+						Writer.WriteLine("      <dd>");
+						for(int Idx = 0; Idx < Lines.Count; Idx++)
 						{
-							Writer.WriteLine("    * {0}", Lines[Idx]);
+							if(Lines[Idx].StartsWith("*") || Lines[Idx].StartsWith("-"))
+							{
+								Writer.WriteLine("        <ul>");
+								for(; Idx < Lines.Count && (Lines[Idx].StartsWith("*") || Lines[Idx].StartsWith("-")); Idx++)
+								{
+									Writer.WriteLine("          <li>{0}</li>", Lines[Idx].Substring(1).TrimStart());
+								}
+								Writer.WriteLine("        </ul>");
+							}
+							else
+							{
+								Writer.WriteLine("        {0}", Lines[Idx]);
+							}
 						}
+						Writer.WriteLine("      </dd>");
 					}
-					Writer.WriteLine();
 				}
 			}
 		}
diff --git a/Engine/Source/Programs/UnrealBuildTool/System/XmlConfig.cs b/Engine/Source/Programs/UnrealBuildTool/System/XmlConfig.cs
--- a/Engine/Source/Programs/UnrealBuildTool/System/XmlConfig.cs
+++ b/Engine/Source/Programs/UnrealBuildTool/System/XmlConfig.cs
@@ -652,35 +652,31 @@ namespace UnrealBuildTool
 				throw new BuildException("Generated assembly documentation not found at {0}.", InputDocumentationFile);
 			}
 
-			// Get the current engine version for versioning the page
-			BuildVersion Version;
-			if(!BuildVersion.TryRead(BuildVersion.GetDefaultFileName(), out Version))
-			{
-				throw new BuildException("Unable to read the current build version");
-			}
-
 			// Read the documentation
 			XmlDocument InputDocumentation = new XmlDocument();
 			InputDocumentation.Load(InputDocumentationFile.FullName);
 
 			// Make sure we can write to the output file
-			FileReference.MakeWriteable(OutputFile);
+			if(FileReference.Exists(OutputFile))
+			{
+				FileReference.MakeWriteable(OutputFile);
+			}
+			else
+			{
+				DirectoryReference.CreateDirectory(OutputFile.Directory);
+			}
 
-			// Generate the UDN documentation file
+			// Generate the documentation file
 			using (StreamWriter Writer = new StreamWriter(OutputFile.FullName))
 			{
-				Writer.WriteLine("Availability: NoPublish");
-				Writer.WriteLine("Title: Build Configuration Properties Page");
-				Writer.WriteLine("Crumbs:");
-				Writer.WriteLine("Description: This is a procedurally generated markdown page.");
-				Writer.WriteLine("Version: {0}.{1}", Version.MajorVersion, Version.MinorVersion);
-				Writer.WriteLine("");
-
+				Writer.WriteLine("<html>");
+				Writer.WriteLine("  <body>");
+				Writer.WriteLine("  <h2>BuildConfiguration Properties</h2>");
 				foreach(KeyValuePair<string, Dictionary<string, FieldInfo>> CategoryPair in CategoryToFields)
 				{
 					string CategoryName = CategoryPair.Key;
-					Writer.WriteLine("### {0}", CategoryName);
-					Writer.WriteLine();
+					Writer.WriteLine("    <h3>{0}</h3>", CategoryName);
+					Writer.WriteLine("    <dl>");
 
 					Dictionary<string, FieldInfo> Fields = CategoryPair.Value;
 					foreach(KeyValuePair<string, FieldInfo> FieldPair in Fields)
@@ -705,23 +701,41 @@ namespace UnrealBuildTool
 							// Write the result to the .udn file
 							if(Lines.Count > 0)
 							{
-								Writer.WriteLine("$ {0} : {1}", FieldName, Lines[0]);
-								for(int Idx = 1; Idx < Lines.Count; Idx++)
+								Writer.WriteLine("      <dt>{0}</dt>", FieldName);
+
+								if(Lines.Count == 1)
 								{
-									if(Lines[Idx].StartsWith("*") || Lines[Idx].StartsWith("-"))
-									{
-										Writer.WriteLine("        * {0}", Lines[Idx].Substring(1).TrimStart());
-									}
-									else
+									Writer.WriteLine("      <dd>{0}</dd>", Lines[0]);
+								}
+								else
+								{
+									Writer.WriteLine("      <dd>");
+									for(int Idx = 0; Idx < Lines.Count; Idx++)
 									{
-										Writer.WriteLine("    * {0}", Lines[Idx]);
+										if(Lines[Idx].StartsWith("*") || Lines[Idx].StartsWith("-"))
+										{
+											Writer.WriteLine("        <ul>");
+											for(; Idx < Lines.Count && (Lines[Idx].StartsWith("*") || Lines[Idx].StartsWith("-")); Idx++)
+											{
+												Writer.WriteLine("          <li>{0}</li>", Lines[Idx].Substring(1).TrimStart());
+											}
+											Writer.WriteLine("        </ul>");
+										}
+										else
+										{
+											Writer.WriteLine("        {0}", Lines[Idx]);
+										}
 									}
+									Writer.WriteLine("      </dd>");
 								}
-								Writer.WriteLine();
 							}
 						}
 					}
+
+					Writer.WriteLine("    </dl>");
 				}
+				Writer.WriteLine("  </body>");
+				Writer.WriteLine("</html>");
 			}
 
 			// Success!
