fix for IPA name of BP projects on PC

diff --git a/Engine/Source/Programs/AutomationTool/IOS/IOSPlatform.Automation.cs b/Engine/Source/Programs/AutomationTool/IOS/IOSPlatform.Automation.cs
--- a/Engine/Source/Programs/AutomationTool/IOS/IOSPlatform.Automation.cs
+++ b/Engine/Source/Programs/AutomationTool/IOS/IOSPlatform.Automation.cs
@@ -265,7 +265,7 @@ public class IOSPlatform : Platform
 	protected string MakeIPAFileName( UnrealTargetConfiguration TargetConfiguration, ProjectParams Params, DeploymentContext SC, bool bAllowDistroPrefix )
 	{
 		return Path.Combine(Path.GetDirectoryName(Params.RawProjectPath.FullName), "Binaries", PlatformName, 
-			((bAllowDistroPrefix && Params.Distribution) ? "Distro_" : "") + SC.StageExecutables[0] + ".ipa");
+			((bAllowDistroPrefix && Params.Distribution) ? "Distro_" : "") + Params.RawProjectPath.GetFileNameWithoutExtension() + ".ipa");
 	}
 
 	// Determine if we should code sign
@@ -427,6 +427,11 @@ public class IOSPlatform : Platform
 		{
 			var ProjectIPA = MakeIPAFileName(TargetConfiguration, Params, SC, Params.Distribution);
 			var ProjectStub = Path.GetFullPath(ProjectGameExeFilename);
+            var IPPProjectIPA = "";
+            if (ProjectStub.Contains("UE4Game.stub"))
+            {
+                IPPProjectIPA = Path.Combine(Path.GetDirectoryName(ProjectIPA), "UE4Game.ipa");
+            }
 
 			// package a .ipa from the now staged directory
 			var IPPExe = CombinePaths(CmdEnv.LocalRoot, "Engine/Binaries/DotNET/IOS/IPhonePackager.exe");
@@ -443,6 +448,10 @@ public class IOSPlatform : Platform
 			{
 				// delete the .ipa to make sure it was made
 				DeleteFile(ProjectIPA);
+                if (IPPProjectIPA.Length > 0)
+                {
+                    DeleteFile(IPPProjectIPA);
+                }
 
 				bCreatedIPA = true;
 
@@ -500,8 +509,14 @@ public class IOSPlatform : Platform
 						IPPArguments += " -tvos";
 					}
 					RunAndLog(CmdEnv, IPPExe, IPPArguments);
-				}
-				else
+
+                    if (IPPProjectIPA.Length > 0)
+                    {
+                        CopyFile(IPPProjectIPA, ProjectIPA);
+                        DeleteFile(IPPProjectIPA);
+                    }
+                }
+                else
 				{
 					List<string> IPPArguments = new List<string>();
 					IPPArguments.Add("RepackageFromStage");
@@ -565,8 +580,14 @@ public class IOSPlatform : Platform
 					{
 						throw new AutomationException("IPP Failed");
 					}
-				}
-			}
+
+                    if (IPPProjectIPA.Length > 0)
+                    {
+                        CopyFile(IPPProjectIPA, ProjectIPA);
+                        DeleteFile(IPPProjectIPA);
+                    }
+                }
+            }
 
 			// verify the .ipa exists
 			if (!FileExists(ProjectIPA))
