Move hooks to allow platforms to build additional utility projects into the project packaging script, rather than running whenever we compile in UAT. UE4Build is used by all BuildGraph compile tasks, and attempting to compile things that we weren't asked to causes access violations if the files are locked.

diff --git a/Engine/Source/Programs/AutomationTool/AutomationUtils/UE4Build.cs b/Engine/Source/Programs/AutomationTool/AutomationUtils/UE4Build.cs
--- a/Engine/Source/Programs/AutomationTool/AutomationUtils/UE4Build.cs
+++ b/Engine/Source/Programs/AutomationTool/AutomationUtils/UE4Build.cs
@@ -1255,13 +1255,6 @@ namespace AutomationTool
 				}
 			}
 
-			// allow all involved platforms to hook into the agenda
-			foreach (var TargetPlatform in UniquePlatforms)
-			{
-				Platform.GetPlatform(TargetPlatform).PreBuildAgenda(this, Agenda);
-			}
-
-
 			foreach (var File in Agenda.ExtraDotNetFiles)
 			{
 				PrepareBuildProductsForCSharpProj(Path.Combine(CommandUtils.CmdEnv.LocalRoot, File));
diff --git a/Engine/Source/Programs/AutomationTool/Scripts/BuildProjectCommand.Automation.cs b/Engine/Source/Programs/AutomationTool/Scripts/BuildProjectCommand.Automation.cs
--- a/Engine/Source/Programs/AutomationTool/Scripts/BuildProjectCommand.Automation.cs
+++ b/Engine/Source/Programs/AutomationTool/Scripts/BuildProjectCommand.Automation.cs
@@ -99,6 +99,15 @@ public partial class Project : CommandUtils
 			}
 		}
 
+		// allow all involved platforms to hook into the agenda
+		HashSet<UnrealTargetPlatform> UniquePlatforms = new HashSet<UnrealTargetPlatform>();
+		UniquePlatforms.UnionWith(Params.ClientTargetPlatforms.Select(x => x.Type));
+		UniquePlatforms.UnionWith(Params.ServerTargetPlatforms.Select(x => x.Type));
+		foreach (UnrealTargetPlatform TargetPlatform in UniquePlatforms)
+		{
+			Platform.GetPlatform(TargetPlatform).PreBuildAgenda(UE4Build, Agenda);
+		}
+
 		// Additional compile arguments
 		string AdditionalArgs = "";
 
