UAT: Use the same heuristic as UGS for determining code changes.

diff --git a/Engine/Source/Programs/AutomationTool/AutomationUtils/P4Environment.cs b/Engine/Source/Programs/AutomationTool/AutomationUtils/P4Environment.cs
--- a/Engine/Source/Programs/AutomationTool/AutomationUtils/P4Environment.cs
+++ b/Engine/Source/Programs/AutomationTool/AutomationUtils/P4Environment.cs
@@ -1,4 +1,4 @@
-﻿// Copyright 1998-2018 Epic Games, Inc. All Rights Reserved.
+// Copyright 1998-2018 Epic Games, Inc. All Rights Reserved.
 
 using System;
 using System.Collections.Generic;
@@ -440,9 +440,16 @@ namespace AutomationTool
 		{
 			CommandUtils.LogVerbose("uebp_CodeCL not set, detecting last code CL...");
 
-			// Retrieve the current changelist 
-			string P4Cmd = String.Format("changes -m 1 \"{0}/....cpp#have\" \"{0}/....h#have\" \"{0}/....inl#have\" \"{0}/....cs#have\" \"{0}/....usf#have\" \"{0}/....ush#have\"", CommandUtils.CombinePaths(PathSeparator.Depot, ClientRootPath));
-			IProcessResult P4Result = Connection.P4(P4Cmd, AllowSpew: false);
+			// Retrieve the current changelist
+			StringBuilder P4Cmd = new StringBuilder("changes -m 1");
+
+			string[] CodeExtensions = { ".cs", ".h", ".cpp", ".inl", ".usf", ".ush", ".uproject", ".uplugin" };
+			foreach(string CodeExtension in CodeExtensions)
+			{
+				P4Cmd.AppendFormat(" \"{0}/...{1}#have\"", CommandUtils.CombinePaths(PathSeparator.Depot, ClientRootPath), CodeExtension);
+			}
+
+			IProcessResult P4Result = Connection.P4(P4Cmd.ToString(), AllowSpew: false);
 
 			// Loop through all the lines of the output. Even though we requested one result, we'll get one for each search pattern.
 			int CL = 0;
