UAT: Do not try to build content-only projects in installed builds. We don't know the settings that they were originally compiled with (eg. -nodebuginfo), so we don't want to make assumptions about what is output.

diff --git a/Engine/Source/Programs/AutomationTool/Scripts/BuildProjectCommand.Automation.cs b/Engine/Source/Programs/AutomationTool/Scripts/BuildProjectCommand.Automation.cs
--- a/Engine/Source/Programs/AutomationTool/Scripts/BuildProjectCommand.Automation.cs
+++ b/Engine/Source/Programs/AutomationTool/Scripts/BuildProjectCommand.Automation.cs
@@ -61,6 +61,10 @@ public partial class Project : CommandUtils
 		{
 			return;
 		}
+		if (Automation.IsEngineInstalled() && !Params.IsCodeBasedProject)
+		{
+			return;
+		}
 
 		Log("********** BUILD COMMAND STARTED **********");
 
