Fix being able to package plugins for the marketplace.

diff --git a/Engine/Source/Programs/AutomationTool/Scripts/BuildPluginCommand.Automation.cs b/Engine/Source/Programs/AutomationTool/Scripts/BuildPluginCommand.Automation.cs
--- a/Engine/Source/Programs/AutomationTool/Scripts/BuildPluginCommand.Automation.cs
+++ b/Engine/Source/Programs/AutomationTool/Scripts/BuildPluginCommand.Automation.cs
@@ -1,4 +1,4 @@
-﻿// Copyright 1998-2018 Epic Games, Inc. All Rights Reserved.
+// Copyright 1998-2018 Epic Games, Inc. All Rights Reserved.
 using System;
 using System.Collections.Generic;
 using System.IO;
@@ -142,8 +142,11 @@ class BuildPlugin : BuildCommand
 			CommandUtils.Log("Building plugin for target platforms: {0}", String.Join(", ", TargetPlatforms));
 			foreach (UnrealTargetPlatform TargetPlatform in TargetPlatforms)
 			{
-				CompilePluginWithUBT(HostProjectFile, HostProjectPluginFile, Plugin, "UE4Game", TargetType.Game, TargetPlatform, UnrealTargetConfiguration.Development, ReceiptFileNames, null);
-				CompilePluginWithUBT(HostProjectFile, HostProjectPluginFile, Plugin, "UE4Game", TargetType.Game, TargetPlatform, UnrealTargetConfiguration.Shipping, ReceiptFileNames, null);
+				if(Plugin.SupportsTargetPlatform(TargetPlatform))
+				{
+					CompilePluginWithUBT(HostProjectFile, HostProjectPluginFile, Plugin, "UE4Game", TargetType.Game, TargetPlatform, UnrealTargetConfiguration.Development, ReceiptFileNames, null);
+					CompilePluginWithUBT(HostProjectFile, HostProjectPluginFile, Plugin, "UE4Game", TargetType.Game, TargetPlatform, UnrealTargetConfiguration.Shipping, ReceiptFileNames, null);
+				}
 			}
 		}
 
diff --git a/Engine/Source/Programs/UnrealBuildTool/Configuration/UEBuildTarget.cs b/Engine/Source/Programs/UnrealBuildTool/Configuration/UEBuildTarget.cs
--- a/Engine/Source/Programs/UnrealBuildTool/Configuration/UEBuildTarget.cs
+++ b/Engine/Source/Programs/UnrealBuildTool/Configuration/UEBuildTarget.cs
@@ -1628,12 +1628,15 @@ namespace UnrealBuildTool
 						if(Module.Rules.bPrecompile)
 						{
 							FileReference PrecompiledManifestLocation = Module.PrecompiledManifestLocation;
-							Receipt.PrecompiledBuildDependencies.Add(PrecompiledManifestLocation);
-
-							PrecompiledManifest Manifest = PrecompiledManifest.Read(PrecompiledManifestLocation);
-							foreach(FileReference OutputFile in Manifest.OutputFiles)
+							if(FileReference.Exists(PrecompiledManifestLocation))
 							{
-								Receipt.PrecompiledBuildDependencies.Add(OutputFile);
+								Receipt.PrecompiledBuildDependencies.Add(PrecompiledManifestLocation);
+
+								PrecompiledManifest Manifest = PrecompiledManifest.Read(PrecompiledManifestLocation);
+								foreach(FileReference OutputFile in Manifest.OutputFiles)
+								{
+									Receipt.PrecompiledBuildDependencies.Add(OutputFile);
+								}
 							}
 						}
 					}
@@ -2348,9 +2351,9 @@ namespace UnrealBuildTool
 			}
 			else
 			{
-			// Build the target's binaries.
+				// Build the target's binaries.
 				foreach (UEBuildBinary Binary in Binaries)
-			{
+				{
 					OutputItems.AddRange(Binary.Build(Rules, TargetToolChain, GlobalCompileEnvironment, GlobalLinkEnvironment, SharedPCHs, WorkingSet, ActionGraph));
 				}
 
