fix for shipping IOS IPA name and TVOS IPA name

diff --git a/Engine/Source/Programs/AutomationTool/IOS/IOSPlatform.Automation.cs b/Engine/Source/Programs/AutomationTool/IOS/IOSPlatform.Automation.cs
--- a/Engine/Source/Programs/AutomationTool/IOS/IOSPlatform.Automation.cs
+++ b/Engine/Source/Programs/AutomationTool/IOS/IOSPlatform.Automation.cs
@@ -264,8 +264,13 @@ public class IOSPlatform : Platform
 
 	protected string MakeIPAFileName( UnrealTargetConfiguration TargetConfiguration, ProjectParams Params, DeploymentContext SC, bool bAllowDistroPrefix )
 	{
+        string ExeName = SC.StageExecutables[0];
+        if (!SC.IsCodeBasedProject)
+        {
+            ExeName = ExeName.Replace("UE4Game", Params.RawProjectPath.GetFileNameWithoutExtension());
+        }
 		return Path.Combine(Path.GetDirectoryName(Params.RawProjectPath.FullName), "Binaries", PlatformName, 
-			((bAllowDistroPrefix && Params.Distribution) ? "Distro_" : "") + Params.RawProjectPath.GetFileNameWithoutExtension() + ".ipa");
+			((bAllowDistroPrefix && Params.Distribution) ? "Distro_" : "") + ExeName + ".ipa");
 	}
 
 	// Determine if we should code sign
@@ -428,9 +433,9 @@ public class IOSPlatform : Platform
 			var ProjectIPA = MakeIPAFileName(TargetConfiguration, Params, SC, Params.Distribution);
 			var ProjectStub = Path.GetFullPath(ProjectGameExeFilename);
             var IPPProjectIPA = "";
-            if (ProjectStub.Contains("UE4Game.stub"))
+            if (ProjectStub.Contains("UE4Game"))
             {
-                IPPProjectIPA = Path.Combine(Path.GetDirectoryName(ProjectIPA), "UE4Game.ipa");
+                IPPProjectIPA = Path.Combine(Path.GetDirectoryName(ProjectIPA), Path.GetFileName(ProjectIPA).Replace(Params.RawProjectPath.GetFileNameWithoutExtension(), "UE4Game"));
             }
 
 			// package a .ipa from the now staged directory
@@ -439,7 +444,8 @@ public class IOSPlatform : Platform
 			LogLog("ProjectName={0}", Params.ShortProjectName);
 			LogLog("ProjectStub={0}", ProjectStub);
 			LogLog("ProjectIPA={0}", ProjectIPA);
-			LogLog("IPPExe={0}", IPPExe);
+            LogLog("IPPProjectIPA={0}", IPPProjectIPA);
+            LogLog("IPPExe={0}", IPPExe);
 
 			bool cookonthefly = Params.CookOnTheFly || Params.SkipCookOnTheFly;
 
