Reduce log spam when building PhysX for Mac/IOS/TVOS. Console output causes massive regression in build times under Mono (10m -> 3h).

#rnx

diff --git a/Engine/Source/Programs/AutomationTool/Scripts/BuildPhysX.Automation.cs b/Engine/Source/Programs/AutomationTool/Scripts/BuildPhysX.Automation.cs
--- a/Engine/Source/Programs/AutomationTool/Scripts/BuildPhysX.Automation.cs
+++ b/Engine/Source/Programs/AutomationTool/Scripts/BuildPhysX.Automation.cs
@@ -1083,7 +1083,7 @@ class BuildPhysX : BuildCommand
 
 		foreach (string BuildConfig in TargetConfigurations)
 		{
-			string CmdLine = String.Format("-project \"{0}\" -target=\"ALL_BUILD\" -configuration {1}", ProjectFile, BuildConfig);
+			string CmdLine = String.Format("-project \"{0}\" -target=\"ALL_BUILD\" -configuration {1} -quiet", ProjectFile, BuildConfig);
 			RunAndLog(BuildCommand.CmdEnv, "/usr/bin/xcodebuild", CmdLine);
 		}
 	}
