Delete any existing temporary workspace created for SyncDepotPath before trying to sync from it.

diff --git a/Engine/Source/Programs/AutomationTool/Scripts/SyncDepotPath.cs b/Engine/Source/Programs/AutomationTool/Scripts/SyncDepotPath.cs
--- a/Engine/Source/Programs/AutomationTool/Scripts/SyncDepotPath.cs
+++ b/Engine/Source/Programs/AutomationTool/Scripts/SyncDepotPath.cs
@@ -1,4 +1,4 @@
-﻿// Copyright 1998-2018 Epic Games, Inc. All Rights Reserved.
+// Copyright 1998-2018 Epic Games, Inc. All Rights Reserved.
 
 using System;
 using System.Collections.Generic;
@@ -35,6 +35,11 @@ namespace AutomationTool
 			List<KeyValuePair<string, string>> RequiredView = new List<KeyValuePair<string, string>>();
 			RequiredView.Add(new KeyValuePair<string, string>(DepotPath, "/..."));
 
+			if(P4.DoesClientExist(ClientName))
+			{
+				P4.DeleteClient(ClientName);
+			}
+
 			P4ClientInfo Client = new P4ClientInfo();
 			Client.Owner = P4Env.User;
 			Client.Host = Environment.MachineName;
