Merging 4214181 to Main
Fix staging of build products for Mac using the TagReceipt task. CL 3924915 introduces a change that will only include build products that are set to be staged as UFS by default, since StagedFileType is not initialized.

diff --git a/Engine/Source/Programs/AutomationTool/BuildGraph/Tasks/TagReceiptTask.cs b/Engine/Source/Programs/AutomationTool/BuildGraph/Tasks/TagReceiptTask.cs
--- a/Engine/Source/Programs/AutomationTool/BuildGraph/Tasks/TagReceiptTask.cs
+++ b/Engine/Source/Programs/AutomationTool/BuildGraph/Tasks/TagReceiptTask.cs
@@ -89,9 +89,15 @@ namespace AutomationTool.Tasks
 		/// </summary>
 		TagReceiptTaskParameters Parameters;
 
-		BuildProductType BuildProductType;
+		/// <summary>
+		/// The type of build products to enumerate. May be null.
+		/// </summary>
+		Nullable<BuildProductType> BuildProductType;
 
-		StagedFileType StagedFileType;
+		/// <summary>
+		/// The type of staged files to enumerate. May be null,
+		/// </summary>
+		Nullable<StagedFileType> StagedFileType;
 
 		/// <summary>
 		/// Constructor
@@ -117,8 +123,7 @@ namespace AutomationTool.Tasks
 		/// <param name="Job">Information about the current job</param>
 		/// <param name="BuildProducts">Set of build products produced by this node.</param>
 		/// <param name="TagNameToFileSet">Mapping from tag names to the set of files they include</param>
-		public override void Execute(JobContext Job, HashSet<FileReference
-			> BuildProducts, Dictionary<string, HashSet<FileReference>> TagNameToFileSet)
+		public override void Execute(JobContext Job, HashSet<FileReference> BuildProducts, Dictionary<string, HashSet<FileReference>> TagNameToFileSet)
 		{
 			// Set the Engine directory
 			DirectoryReference EngineDir = DirectoryReference.Combine(CommandUtils.RootDirectory, "Engine");
@@ -158,10 +163,15 @@ namespace AutomationTool.Tasks
 				{
 					foreach (BuildProduct BuildProduct in Receipt.BuildProducts)
 					{
-						if ((String.IsNullOrEmpty(Parameters.BuildProductType) && TargetReceipt.GetStageTypeFromBuildProductType(BuildProduct) == StagedFileType) || BuildProduct.Type == BuildProductType)
+						if(BuildProductType.HasValue && BuildProduct.Type != BuildProductType.Value)
 						{
-							Files.Add(BuildProduct.Path);
+							continue;
 						}
+						if(StagedFileType.HasValue && TargetReceipt.GetStageTypeFromBuildProductType(BuildProduct) != StagedFileType.Value)
+						{
+							continue;
+						}
+						Files.Add(BuildProduct.Path);
 					}
 				}
 
@@ -169,18 +179,25 @@ namespace AutomationTool.Tasks
 				{
 					foreach (RuntimeDependency RuntimeDependency in Receipt.RuntimeDependencies)
 					{
-						if (String.IsNullOrEmpty(Parameters.StagedFileType) || RuntimeDependency.Type == StagedFileType)
+						// Skip anything that doesn't match the files we want
+						if(BuildProductType.HasValue)
+						{
+							continue;
+						}
+						if(StagedFileType.HasValue && RuntimeDependency.Type != StagedFileType.Value)
+						{
+							continue;
+						}
+
+						// Only add files that exist as dependencies are assumed to always exist
+						FileReference DependencyPath = RuntimeDependency.Path;
+						if (FileReference.Exists(DependencyPath))
+						{
+							Files.Add(DependencyPath);
+						}
+						else
 						{
-							// Only add files that exist as dependencies are assumed to always exist
-							FileReference DependencyPath = RuntimeDependency.Path;
-							if (FileReference.Exists(DependencyPath))
-							{
-								Files.Add(DependencyPath);
-							}
-							else
-							{
-								CommandUtils.LogWarning("File listed as RuntimeDependency in {0} does not exist ({1})", TargetFile.FullName, DependencyPath.FullName);
-							}
+							CommandUtils.LogWarning("File listed as RuntimeDependency in {0} does not exist ({1})", TargetFile.FullName, DependencyPath.FullName);
 						}
 					}
 				}
