Merging CL# 4243662
Fix launch-on delta manifest paths

#ue4
#rnx

diff --git a/Engine/Source/Programs/AutomationTool/AutomationUtils/DeploymentContext.cs b/Engine/Source/Programs/AutomationTool/AutomationUtils/DeploymentContext.cs
--- a/Engine/Source/Programs/AutomationTool/AutomationUtils/DeploymentContext.cs
+++ b/Engine/Source/Programs/AutomationTool/AutomationUtils/DeploymentContext.cs
@@ -909,28 +909,28 @@ public class DeploymentContext //: ProjectParams
 	{
 		string SanitizedDeviceName = GetSanitizedDeviceName(DeviceName);
 
-		return Path.Combine(EngineRoot.FullName, "Intermediate", "UAT", UFSDeployDeltaFileName + SanitizedDeviceName);
+		return Path.Combine(StageDirectory.FullName, UFSDeployDeltaFileName + SanitizedDeviceName);
 	}
 
 	public String GetNonUFSDeploymentDeltaPath(string DeviceName)
 	{
 		string SanitizedDeviceName = GetSanitizedDeviceName(DeviceName);
 
-		return Path.Combine(EngineRoot.FullName, "Intermediate", "UAT", NonUFSDeployDeltaFileName + SanitizedDeviceName);
+		return Path.Combine(StageDirectory.FullName, NonUFSDeployDeltaFileName + SanitizedDeviceName);
 	}
 
 	public String GetUFSDeploymentObsoletePath(string DeviceName)
 	{
 		string SanitizedDeviceName = GetSanitizedDeviceName(DeviceName);
 
-		return Path.Combine(EngineRoot.FullName, "Intermediate", "UAT", UFSDeployObsoleteFileName + SanitizedDeviceName);
+		return Path.Combine(StageDirectory.FullName, UFSDeployObsoleteFileName + SanitizedDeviceName);
 	}
 
 	public String GetNonUFSDeploymentObsoletePath(string DeviceName)
 	{
 		string SanitizedDeviceName = GetSanitizedDeviceName(DeviceName);
 
-		return Path.Combine(EngineRoot.FullName, "Intermediate", "UAT", NonUFSDeployObsoleteFileName + SanitizedDeviceName);
+		return Path.Combine(StageDirectory.FullName, NonUFSDeployObsoleteFileName + SanitizedDeviceName);
 	}
 
 	public string UFSDeployedManifestFileName
