Copying //UE4/Dev-Main to //UE4/Main (Source: //Portal/Main/Engine @ 4247640)

============================
  MAJOR FEATURES & CHANGES
============================

Change 4247640 by [REDACTED_PERSON]

	BuildPatchTool: PackageChunks: Fixed issue with filenames not being set for chunkdbs in particular cases where the chunkdbs were small.

Change 4247246 by [REDACTED_PERSON]

	BuildPatchTool: Adding support to BPT PackageChunks to filter by tagsets, and also split chunkdb output by them too.

Change 4173518 by [REDACTED_PERSON]

	WebBrowser: Adding support to enable allowing net security expirations.

Change 4102451 by [REDACTED_PERSON]

	BuildPatchTool: PackageChunks: Adding tool mode JSON output for listing created chunkdb files.

Change 4099951 by [REDACTED_PERSON]

	BuildPatchTool: PackageChunks: Adding PrevManifestFile commandline support to BPT PackageChunks mode, allowing us to generate chunkdbs that only cover data required to perform an update.

Change 4051406 by [REDACTED_PERSON]

	BuildPatchServices: Update default number of simultaneous downloads for an installer to 16 which is now well proven in the wild.

Change 4036807 by [REDACTED_PERSON]

	BuildPatchServices: Added new message type for receiving updates about which files have been touched by the installation.
	BuildPatchServices: Cloud enumeration is now threaded to improve speed.

Change 4036477 by [REDACTED_PERSON]

	Thirdparty: AWSSDK: Update AWS SDK to version 3.3 as speculative fix for object is used after being disposed error.

	Confirmed this exception is being raised from within the AWS SDK, and _not_ in Epic code.
	Advice by Amazon on AWS forums for people experiencing this exception is always to update to the latest version of the SDK in the first instance.

Change 3922493 by [REDACTED_PERSON]

	Runtime: Added shortcircuit support to the ExpressionParser. More documentation to come.
	AgreementExpressionEvaluator has been updated to use the shortcircuit logic so it now processes expressions lef to right as expected in all cases.
	AgreementExpressionEvaluator was also updated to perform evaluation as a two pass system. The tow pass solution prevents over prompting eulas in certain scenarios where the expression would ultimately resolve true without them.

Change 3983713 by [REDACTED_PERSON]

	BuildPatchServices: Fix for verification logging error counts for unique runs. Solved the issue by using an external cache of the errors encountered each run to deduct from the total. Reporting only the delta.

Change 3966915 by [REDACTED_PERSON]

	WebBrowser: CEF: Potential fix for a rare issue encountered when we attempting to release resources outside of the game thread.

Change 3955168 by [REDACTED_PERSON]

	BuildPatchServices: Updated primary messaging of overall install status to now display Updating when patching an existing installation rather than Installing.

Change 3954610 by [REDACTED_PERSON]

	BuildPatchServices: Fixing issue with Cloud Chunk Source which would erroneously conclude that it needed to re-download a chunk due to external system failure.

	There is now an explicit concept of 'lost chunk' which is tracked and bubbled up by the system errors that cause the loss in the first place, so that each source knows exactly when it should be required to retrieve a chunk that it had already retrieved previously. Being explicit, these lost chunks can also now contribute to the total download required stat.
	Also fixing some tracking that was missing to update the total download required stat if a local Install Chunk Source failed to load data from the files on disk and so these chunks needed to be additionally downloaded.

Change 3947928 by [REDACTED_PERSON]

	UAT: Changing commandutils to attempt to find the Win8.1sdk signtool before the win10 one.  There's currently a bug related to vs2017/server2012r2/win10sdk signtool that causes it to exit with an undefined error when signing from a service account.  Using the win8.1sdk circumvents this issue.

Change 3942776 by [REDACTED_PERSON]

	Http: Fix for Mac sending up duplicate header strings

Change 3940306 by [REDACTED_PERSON]

	BuildPatchServices: Refactor to isolate CoreUObject dependency and be able to compile out usage.

Change 3936655 by [REDACTED_PERSON]

	Slate: Changed the invalid fontcache ensure in ShapedTextCache from always to only once to reduce ensure spamChange 3917840 by [REDACTED_PERSON]

	BuildPatchTool: DiffManifests mode now also saves info for New, Removed, Changed, and Unchanged file to the output json file.

Change 3911756 by [REDACTED_PERSON]

	WebBrowser: SWebBrowserView now ensures that it has a valid pointer to it's parent window in it's onpaint to avoid issues with the first frame being scaled incorrectly due to not being able to access the parent windows dpi scaling.

Change 3906670 by [REDACTED_PERSON]

	Slate: Change the ShapedTextCache to hold a weakptr to the Slate FontCache instead of a reference. This will allow it to detect if the FontCache has become invalid since it was linked with the ShapedTextCache, thus making it able to avoid crashing.

Change 3889008 by [REDACTED_PERSON]

	StandaloneRenderer: Made the SlateD3DRenderingPolicy more resilient to graphics device errors.

Change 3886969 by [REDACTED_PERSON]

	StandaloneRenderer:Changed SlateD3DConstantBuffer to no longer check on a D3DDevice failure and instead soft fail, so the application can go through the process of attempting to re-establish the D3DDevice.

Change 3886960 by [REDACTED_PERSON]

	WebBrowser: Made CEFWebBrowserWindow more resilient to issues with creating textures.

Change 3855821 by [REDACTED_PERSON]

	BuildPatchServices: Added logging of the configuration to the start of all installs.

Change 3839245 by [REDACTED_PERSON]

	WebBrowser: Adding support for web browser drag regions.  These are areas of a page tagged with -webkit-app-region: drag or -webkit-app-region: no-drag.  The application can now pass a handler function to the browser to handle window drag events.  This handler will be called if the browser detects mouse drag events inside of a tagged drag region.

Change 3835225 by [REDACTED_PERSON]

	Slate: Fix for SScrollBarTrack size issue

Change 3824320 by [REDACTED_PERSON]

	WebBrowser: Fix for reported deadlock in WebBrowserSingleton.  Associated with github pull request #4303.
	#4303

Change 3811191 by [REDACTED_PERSON]

	BuildPatchServices: Expanded MF01-X into MF01-X-X and MF02-X where X is the os error codes

Change 3807662 by [REDACTED_PERSON]

	BuildPatchServices: Broke down build verification errors into 4 new cases.

Change 3805698 by [REDACTED_PERSON]

	BuildPatchServices: Speculative fixes for Unit test crashes / failures.

Change 3804175 by [REDACTED_PERSON]

	ThirdParty: CEF: Adding browser locale pak files for es-MX, and es-ES as the typical mapping/fallback does not seem functional on mac browser.  They are copies of es_419 and es respectively.

Change 3786628 by [REDACTED_PERSON]

	WebBrowser: Exposing ability to customise tool tip widgets produced from SWebBrowserView.

Change 3775678 by [REDACTED_PERSON]

	BuildPatchServices: Allow a Prerequisite install only mode.

Change 3774365 by [REDACTED_PERSON]

	BuildPatchServices: Updated Build Stat report that is logged after every installation to use FText::AsMemory instead of the UnitConversion logic, and now it outputs multiple unit types for convenience.

Change 3774361 by [REDACTED_PERSON]

	Http: Remove Pragma: no-cache header from libcurl requests

Change 3774258 by [REDACTED_PERSON]

	BuildPatchServices: Fix for destructive patch destroy files that contain useful data.

Change 3766156 by [REDACTED_PERSON]

	Http: Various lower changes to enable Hardware testing and to enable proxy configuration.

Change 3756723 by [REDACTED_PERSON]

	BuildPatchServices: Hooking up disk chunk store operation states to the installer statistics.
	BuildPatchServices: Memory chunk store statistics fix for booted chunks that have been reloaded.

Change 3756320 by [REDACTED_PERSON]

	Http: Add default headers added to every HTTP request.

Change 3741274 by [REDACTED_PERSON]

	WebBrowser: Release CEF related references prior to CEF shutdown.

Change 3738003 by [REDACTED_PERSON]

	BuildPatchServices: Fix-ups for install stats when failures are occurring
	NumFilesOutdated now only set on first run, so it is not set to the number of files that are retried.
	Total downloaded data and total download requirement stats fixed up for runtime as well as final values.
	Initial chunk counters now only set for first run.
	Moving the GetBytesDownloaded api from cloud source to download service so that it correctly accumulates.
	InstallSource was multiply attempting, and counting, recycle failures, throwing that stat out of proportion.

Change 3729851 by [REDACTED_PERSON]

	BuildPatchServices: Changed the installer's MoveFile method to default to not retry. We currently only have uses of move file which already handle retry.

Change 3725611 by [REDACTED_PERSON]

	Core: FText::AsMemory - Fix for numerical edges.
	Added unit tests to check all edges up to full uint64 range.

Change 3725127 by [REDACTED_PERSON]

	BuildPatchServices: ManifestDiff: Correcting string format padding for new data size output uints.

Change 3725126 by [REDACTED_PERSON]

	Core: FText::AsMemory fix. Shifting (equivalent of divide 2 per shift) does not work for calculating SI units which are base 10.

Change 3721926 by [REDACTED_PERSON]

	ThirdParty: LibCurl update performed by [REDACTED_PERSON].
	We now have 100MB/sec download speed with libcurl on Win64, compared to 3MB/sec originally

Change 3700670 by [REDACTED_PERSON]

	SlateReflector: Fixed mouse click highlighting in Widget Reflector's Demo Mode in high DPI

Change 3697526 by [REDACTED_PERSON]

	BuildPatchServices: Exposing a suite of runtime statistics for BuildPatchInstallers.

Change 3686439 by [REDACTED_PERSON]

	BuildPatchServices: Stop installers from always logging a shutdown error on destruction. This should only occur if the installer is actually running.

Change 3684747 by [REDACTED_PERSON]

	BuildPatchTool: Fix file ignore list to parse using platform agnostic method.

Change 3643038 by [REDACTED_PERSON]

	Core: Don't defer Cocoa calls in FMacWindow Show and Hide to make sure both actions complete before we exit these functions. This solves the problem with the blocks being called after window was destroyed.

Change 3639692 by [REDACTED_PERSON]

	Fixes for a couple of issues found by address sanitizer

Change 3625568 by [REDACTED_PERSON]

	BuildPatchServices: Fixing numerical limits problem with double -> uint64 in FStatsCollector::SecondsToCycles().

Change 3617948 by [REDACTED_PERSON]

	BuildPatchServices: Disk space requirement can now be lower when patching if destructive patch mode is enabled. This mode will delete existing old files once they are not needed.
	BuildPatchServices: Adding new installation mode setting on installer config.
	BuildPatchServices: If destructive installation mode is enabled, the file constructor will delete old existing files after completing the new one.
	BuildPatchServices: Adding missing file path length check for install location to cover situation where staging directory is outside the install directory.

Change 3593632 by [REDACTED_PERSON]

	BuildPatchServices: Adding additional installation tracking to BuildPatchServices. See Engine/Source/Runtime/Online/BuildPatchServices/Public/Interfaces/IBuildInstaller.h

diff --git a/Engine/Plugins/Runtime/HTTPChunkInstaller/Source/Private/HTTPChunkInstaller.cpp b/Engine/Plugins/Runtime/HTTPChunkInstaller/Source/Private/HTTPChunkInstaller.cpp
--- a/Engine/Plugins/Runtime/HTTPChunkInstaller/Source/Private/HTTPChunkInstaller.cpp
+++ b/Engine/Plugins/Runtime/HTTPChunkInstaller/Source/Private/HTTPChunkInstaller.cpp
@@ -1397,9 +1397,7 @@ void FHTTPChunkInstall::ParseTitleFileManifest(const FString& ManifestFileHash)
 		return;
 	}
 #endif
-	FString JsonBuffer;
-	FFileHelper::BufferToString(JsonBuffer, FileContentBuffer.GetData(), FileContentBuffer.Num());
-	auto RemoteManifest = BPSModule->MakeManifestFromJSON(JsonBuffer);
+	auto RemoteManifest = BPSModule->MakeManifestFromData(FileContentBuffer);
 	if (!RemoteManifest.IsValid())
 	{
 		UE_LOG(LogHTTPChunkInstaller, Warning, TEXT("Manifest was invalid"));
diff --git a/Engine/Plugins/Runtime/MobilePatchingUtils/Source/MobilePatchingUtils/Private/MobilePatchingLibrary.cpp b/Engine/Plugins/Runtime/MobilePatchingUtils/Source/MobilePatchingUtils/Private/MobilePatchingLibrary.cpp
--- a/Engine/Plugins/Runtime/MobilePatchingUtils/Source/MobilePatchingUtils/Private/MobilePatchingLibrary.cpp
+++ b/Engine/Plugins/Runtime/MobilePatchingUtils/Source/MobilePatchingUtils/Private/MobilePatchingLibrary.cpp
@@ -34,6 +34,51 @@ static IBuildPatchServicesModule* GetBuildPatchServices()
 	return BuildPatchServices;
 }
 
+static const FText& GetBuildPatchStateText(const BuildPatchServices::EBuildPatchState& State)
+{
+	static const FText Queued = LOCTEXT("EBuildPatchState_Queued", "Queued");
+	static const FText Initializing = LOCTEXT("EBuildPatchState_Initialising", "Initializing");
+	static const FText Resuming = LOCTEXT("EBuildPatchState_Resuming", "Resuming");
+	static const FText Downloading = LOCTEXT("EBuildPatchState_Downloading", "Downloading");
+	static const FText Installing = LOCTEXT("EBuildPatchState_Installing", "Installing");
+	static const FText BuildVerification = LOCTEXT("EBuildPatchState_BuildVerification", "Verifying");
+	static const FText CleanUp = LOCTEXT("EBuildPatchState_CleanUp", "Cleaning up");
+	static const FText PrerequisitesInstall = LOCTEXT("EBuildPatchState_PrerequisitesInstall", "Prerequisites");
+	static const FText Completed = LOCTEXT("EBuildPatchState_Complete", "Complete");
+	static const FText Paused = LOCTEXT("EBuildPatchState_Paused", "Paused");
+	static const FText Empty = FText::GetEmpty();
+
+	switch (State)
+	{
+		case BuildPatchServices::EBuildPatchState::Queued:
+			return Queued;
+		case BuildPatchServices::EBuildPatchState::Initializing:
+			return Initializing;
+		case BuildPatchServices::EBuildPatchState::Resuming:
+			return Resuming;
+		case BuildPatchServices::EBuildPatchState::Downloading:
+			return Downloading;
+		case BuildPatchServices::EBuildPatchState::Installing:
+			// falls through
+		case BuildPatchServices::EBuildPatchState::MovingToInstall:
+			// falls through
+		case BuildPatchServices::EBuildPatchState::SettingAttributes:
+			return Installing;
+		case BuildPatchServices::EBuildPatchState::BuildVerification:
+			return BuildVerification;
+		case BuildPatchServices::EBuildPatchState::CleanUp:
+			return CleanUp;
+		case BuildPatchServices::EBuildPatchState::PrerequisitesInstall:
+			return PrerequisitesInstall;
+		case BuildPatchServices::EBuildPatchState::Completed:
+			return Completed;
+		case BuildPatchServices::EBuildPatchState::Paused:
+			return Paused;
+		default:
+			return Empty;
+	}
+}
+
 static const FText& GetRequestContentErrorText(ERequestContentError ErrorCode)
 {
 	static const FText NoError(LOCTEXT("RequestContentError_NoError", "The operation was successful."));
@@ -180,7 +225,7 @@ FText UMobilePendingContent::GetDownloadStatusText()
 {
 	if (Installer.IsValid())
 	{
-		return Installer->GetStatusText();
+		return GetBuildPatchStateText(Installer->GetState());
 	}
 	return FText::GetEmpty();
 }
diff --git a/Engine/Source/Developer/SlateReflector/Private/Widgets/SWidgetReflector.cpp b/Engine/Source/Developer/SlateReflector/Private/Widgets/SWidgetReflector.cpp
--- a/Engine/Source/Developer/SlateReflector/Private/Widgets/SWidgetReflector.cpp
+++ b/Engine/Source/Developer/SlateReflector/Private/Widgets/SWidgetReflector.cpp
@@ -1244,6 +1244,7 @@ int32 SWidgetReflector::VisualizeCursorAndKeys(FSlateWindowElementList& OutDrawE
 
 			FGeometry CursorHighlightGeometry = FGeometry::MakeRoot(PingSize, FSlateLayoutTransform(CursorPosDesktopSpace - PingSize / 2));
 			CursorHighlightGeometry.AppendTransform(Inverse(WindowBeingDrawn->GetLocalToScreenTransform()));
+			CursorHighlightGeometry.AppendTransform(FSlateLayoutTransform(WindowBeingDrawn->GetDPIScaleFactor(), FVector2D::ZeroVector));
 
 			FSlateDrawElement::MakeBox(
 				OutDrawElements,
diff --git a/Engine/Source/Developer/StandaloneRenderer/Private/Windows/D3D/SlateD3DConstantBuffer.h b/Engine/Source/Developer/StandaloneRenderer/Private/Windows/D3D/SlateD3DConstantBuffer.h
--- a/Engine/Source/Developer/StandaloneRenderer/Private/Windows/D3D/SlateD3DConstantBuffer.h
+++ b/Engine/Source/Developer/StandaloneRenderer/Private/Windows/D3D/SlateD3DConstantBuffer.h
@@ -21,7 +21,12 @@ public:
 		Desc.StructureByteStride = 0;
 
 		HRESULT Hr = GD3DDevice->CreateBuffer(&Desc, NULL, Buffer.GetInitReference() );
-		checkf( SUCCEEDED(Hr), TEXT("D3D11 Error Result %X"), Hr );
+		if (FAILED(Hr))
+		{
+			LogSlateD3DRendererFailure(TEXT("FSlateD3DConstantBuffer::Create() - ID3D11Device::CreateBuffer"), Hr);
+			GEncounteredCriticalD3DDeviceError = true;
+			return;
+		}
 	}
 
 	TRefCountPtr<ID3D11Buffer> GetResource() { return Buffer; }
diff --git a/Engine/Source/Developer/StandaloneRenderer/Private/Windows/D3D/SlateD3DRenderer.cpp b/Engine/Source/Developer/StandaloneRenderer/Private/Windows/D3D/SlateD3DRenderer.cpp
--- a/Engine/Source/Developer/StandaloneRenderer/Private/Windows/D3D/SlateD3DRenderer.cpp
+++ b/Engine/Source/Developer/StandaloneRenderer/Private/Windows/D3D/SlateD3DRenderer.cpp
@@ -140,15 +140,25 @@ bool FSlateD3DRenderer::Initialize()
 		bool bResult = CreateDevice();
 		if (bResult)
 		{
+			GEncounteredCriticalD3DDeviceError = false;
+
 			TextureManager = MakeShareable(new FSlateD3DTextureManager);
 			FSlateDataPayload::ResourceManager = TextureManager.Get();
 
-			TextureManager->LoadUsedTextures();
+			if (!GEncounteredCriticalD3DDeviceError)
+			{
+				TextureManager->LoadUsedTextures();
+			}
 
-			RenderingPolicy = MakeShareable(new FSlateD3D11RenderingPolicy(SlateFontServices.ToSharedRef(), TextureManager.ToSharedRef()));
+			if (!GEncounteredCriticalD3DDeviceError)
+			{
+				RenderingPolicy = MakeShareable(new FSlateD3D11RenderingPolicy(SlateFontServices.ToSharedRef(), TextureManager.ToSharedRef()));
+			}
 
-			ElementBatcher = MakeShareable(new FSlateElementBatcher(RenderingPolicy.ToSharedRef()));
-			GEncounteredCriticalD3DDeviceError = false;
+			if (!GEncounteredCriticalD3DDeviceError)
+			{
+				ElementBatcher = MakeShareable(new FSlateElementBatcher(RenderingPolicy.ToSharedRef()));
+			}
 		}
 		else
 		{
diff --git a/Engine/Source/Developer/StandaloneRenderer/Private/Windows/D3D/SlateD3DRenderingPolicy.cpp b/Engine/Source/Developer/StandaloneRenderer/Private/Windows/D3D/SlateD3DRenderingPolicy.cpp
--- a/Engine/Source/Developer/StandaloneRenderer/Private/Windows/D3D/SlateD3DRenderingPolicy.cpp
+++ b/Engine/Source/Developer/StandaloneRenderer/Private/Windows/D3D/SlateD3DRenderingPolicy.cpp
@@ -43,6 +43,9 @@ FSlateD3D11RenderingPolicy::~FSlateD3D11RenderingPolicy()
 
 void FSlateD3D11RenderingPolicy::InitResources()
 {
+	VertexShader = nullptr;
+	PixelShader = nullptr;
+
 	D3D11_SAMPLER_DESC SamplerDesc;
 	FMemory::Memzero( &SamplerDesc, sizeof(D3D11_SAMPLER_DESC) );
 	SamplerDesc.Filter = D3D11_FILTER_MIN_MAG_MIP_POINT;
@@ -170,12 +173,27 @@ void FSlateD3D11RenderingPolicy::ReleaseResources()
 	VertexBuffer.DestroyBuffer();
 	IndexBuffer.DestroyBuffer();
 
-	delete VertexShader;
-	delete PixelShader;
+	if (VertexShader != nullptr)
+	{
+		delete VertexShader;
+		VertexShader = nullptr;
+	}
+
+	if (PixelShader != nullptr)
+	{
+		delete PixelShader;
+		PixelShader = nullptr;
+	}
 }
 
 void FSlateD3D11RenderingPolicy::UpdateVertexAndIndexBuffers( FSlateBatchData& InBatchData )
 {	
+	if (!VertexBuffer.GetResource().IsValid() ||
+		!IndexBuffer.GetResource().IsValid())
+	{
+		return;
+	}
+
 	if( InBatchData.GetRenderBatches().Num() > 0 )
 	{
 		if( InBatchData.GetNumBatchedVertices() > 0 )
@@ -230,6 +248,12 @@ void FSlateD3D11RenderingPolicy::UpdateVertexAndIndexBuffers( FSlateBatchData& I
 
 void FSlateD3D11RenderingPolicy::DrawElements( const FMatrix& ViewProjectionMatrix, const TArray<FSlateRenderBatch>& RenderBatches )
 {
+	if (!VertexBuffer.GetResource().IsValid() ||
+		!IndexBuffer.GetResource().IsValid())
+	{
+		return;
+	}
+
 	VertexShader->BindShader();
 
 	ID3D11Buffer* Buffer = VertexBuffer.GetResource();
diff --git a/Engine/Source/Programs/AutomationTool/AutomationUtils/MCPPublic.cs b/Engine/Source/Programs/AutomationTool/AutomationUtils/MCPPublic.cs
--- a/Engine/Source/Programs/AutomationTool/AutomationUtils/MCPPublic.cs
+++ b/Engine/Source/Programs/AutomationTool/AutomationUtils/MCPPublic.cs
@@ -81,41 +81,41 @@ namespace EpicGames.MCP.Automation
         }
     }
 
-    /// <summary>
-    /// Enum that defines the MCP backend-compatible platform
-    /// </summary>
-    public enum MCPPlatform
-    {
-        /// <summary>
+	/// <summary>
+	/// Enum that defines the MCP backend-compatible platform
+	/// </summary>
+	public enum MCPPlatform
+	{
+		/// <summary>
 		/// MCP uses Windows for Win64
-        /// </summary>
-        Windows,
+		/// </summary>
+		Windows,
 
-        /// <summary>
+		/// <summary>
 		/// 32 bit Windows
 		/// </summary>
 		Win32,
 
 		/// <summary>
-        /// Mac platform.
-        /// </summary>
-        Mac,
+		/// Mac platform.
+		/// </summary>
+		Mac,
 
 		/// <summary>
 		/// Linux platform.
 		/// </summary>
 		Linux,
 
-        /// <summary>
-        /// Linux platform.
-        /// </summary>
-        IOS,
+		/// <summary>
+		/// IOS platform.
+		/// </summary>
+		IOS,
 
-        /// <summary>
-        /// Linux platform.
-        /// </summary>
-        Android
-    }
+		/// <summary>
+		/// Android platform.
+		/// </summary>
+		Android
+	}
 
     /// <summary>
     /// Enum that defines CDN types
@@ -727,11 +727,27 @@ namespace EpicGames.MCP.Automation
 			public class ManifestDiff
 			{
 				/// <summary>
-				/// The list of cloud directory relative paths for all new chunks required for a patch.
+				/// The list of build relative paths for files which were added by the patch from ManifestA to ManifestB, subject to using the tags that were provided.
+				/// </summary>
+				public List<string> NewFilePaths;
+				/// <summary>
+				/// The list of build relative paths for files which were removed by the patch from ManifestA to ManifestB, subject to using the tags that were provided.
+				/// </summary>
+				public List<string> RemovedFilePaths;
+				/// <summary>
+				/// The list of build relative paths for files which were changed between ManifestA and ManifestB, subject to using the tags that were provided.
+				/// </summary>
+				public List<string> ChangedFilePaths;
+				/// <summary>
+				/// The list of build relative paths for files which were unchanged between ManifestA and ManifestB, subject to using the tags that were provided.
+				/// </summary>
+				public List<string> UnchangedFilePaths;
+				/// <summary>
+				/// The list of cloud directory relative paths for all new chunks required by the patch from ManifestA to ManifestB, subject to using the tags that were provided.
 				/// </summary>
 				public List<string> NewChunkPaths;
 				/// <summary>
-				/// The required download size for the patch between ManifestA and ManifestB, subject to using the tags that were provided.
+				/// The required download size for the patch from ManifestA to ManifestB, subject to using the tags that were provided.
 				/// </summary>
 				public ulong DeltaDownloadSize;
 				/// <summary>
@@ -769,6 +785,10 @@ namespace EpicGames.MCP.Automation
 			/// </summary>
 			public string ManifestFile;
 			/// <summary>
+			/// Specifies the file path to a manifest for a previous build, this will be used to filter out chunks.
+			/// </summary>
+			public string PrevManifestFile;
+			/// <summary>
 			/// Specifies the file path to the output package.  An extension of .chunkdb will be added if not present.
 			/// </summary>
 			public string OutputFile;
@@ -784,6 +804,31 @@ namespace EpicGames.MCP.Automation
 			/// required for highest numbered part.
 			/// </summary>
 			public ulong? MaxOutputFileSize;
+			/// <summary>
+			/// Optionally provide a list of tagsets to split chunkdb files on. First all data from the tagset at index 0 will be saved, then any extra data needed
+			/// for tagset at index 1, and so on. Note that this means the chunkdb files produced for tagset at index 1 will not contain some required data for that tagset if
+			/// the data already got saved out as part of tagset at index 0, and thus the chunkdb files are additive with no dupes.
+			/// If it is desired that each tagset's chunkdb files contain the duplicate data, then PackageChunks should be executed once per tagset rather than once will all tagsets.
+			/// An empty string must be included in one of the tagsets to include untagged file data in that tagset.
+			/// Leaving this variable null will include data for all files.
+			/// </summary>
+			public List<HashSet<string>> TagSetSplit;
+		}
+
+		public class PackageChunksOutput
+		{
+			/// <summary>
+			/// The list of full filepaths of all created chunkdb files.
+			/// </summary>
+			public List<string> ChunkDbFilePaths;
+			/// <summary>
+			/// If PackageChunksOptions.TagSetSplit was provided, then this variable will contain a lookup table of TagSetSplit index to List of ChunkDbFilePaths indices.
+			/// e.g.
+			/// TagSetLookupTable[0] = [ 0, 1, 2, 3, ..., n ]
+			/// TagSetLookupTable[1] = [] an empty List would mean that all data for this tagset was already included in the chunkdb(s) for previous tagset(s).
+			/// TagSetLookupTable[2] = [ n+1, n+2, ..., n+m ]
+			/// </summary>
+			public List<List<int>> TagSetLookupTable;
 		}
 
 		static BuildPatchToolBase Handler = null;
@@ -861,10 +906,12 @@ namespace EpicGames.MCP.Automation
 		/// Runs the Build Patch Tool executable to create ChunkDB file(s) consisting of multiple chunks to allow installing / patching to a specific build.
 		/// </summary>
 		/// <param name="Opts">Parameters which will be passed to the patch tool package chunks process.</param>
+		/// <param name="Output">Will receive the data back for the packaging.</param>
 		/// <param name="Version">Which version of BuildPatchTool is desired.</param>
-		public abstract void Execute(PackageChunksOptions Opts, ToolVersion Version = ToolVersion.Live);
+		public abstract void Execute(PackageChunksOptions Opts, out PackageChunksOutput Output, ToolVersion Version = ToolVersion.Live);
 	}
 
+
 	/// <summary>
 	/// Class that provides programmatic access to the metadata field from build info.
 	/// </summary>
@@ -883,33 +930,33 @@ namespace EpicGames.MCP.Automation
 	/// Helper class
 	/// </summary>
 	public abstract class BuildInfoPublisherBase
-    {
-        static BuildInfoPublisherBase Handler = null;
- 
-        public static BuildInfoPublisherBase Get()
-        {
-            if (Handler == null)
-            {
-                Assembly[] LoadedAssemblies = AppDomain.CurrentDomain.GetAssemblies();
-                foreach (var Dll in LoadedAssemblies)
-                {
-                    Type[] AllTypes = Dll.SafeGetLoadedTypes();
-                    foreach (var PotentialConfigType in AllTypes)
-                    {
-                        if (PotentialConfigType != typeof(BuildInfoPublisherBase) && typeof(BuildInfoPublisherBase).IsAssignableFrom(PotentialConfigType))
-                        {
-                            Handler = Activator.CreateInstance(PotentialConfigType) as BuildInfoPublisherBase;
-                            break;
-                        }
-                    }
-                }
-                if (Handler == null)
-                {
-                    throw new AutomationException("Attempt to use BuildInfoPublisherBase.Get() and it doesn't appear that there are any modules that implement this class.");
-                }
-            }
-            return Handler;
-        }
+	{
+		static BuildInfoPublisherBase Handler = null;
+
+		public static BuildInfoPublisherBase Get()
+		{
+			if (Handler == null)
+			{
+				Assembly[] LoadedAssemblies = AppDomain.CurrentDomain.GetAssemblies();
+				foreach (var Dll in LoadedAssemblies)
+				{
+					Type[] AllTypes = Dll.SafeGetLoadedTypes();
+					foreach (var PotentialConfigType in AllTypes)
+					{
+						if (PotentialConfigType != typeof(BuildInfoPublisherBase) && typeof(BuildInfoPublisherBase).IsAssignableFrom(PotentialConfigType))
+						{
+							Handler = Activator.CreateInstance(PotentialConfigType) as BuildInfoPublisherBase;
+							break;
+						}
+					}
+				}
+				if (Handler == null)
+				{
+					throw new AutomationException("Attempt to use BuildInfoPublisherBase.Get() and it doesn't appear that there are any modules that implement this class.");
+				}
+			}
+			return Handler;
+		}
 
 		/// <summary>
 		/// Creates a metadata object implementation, initialized with the provided JSON object.
@@ -926,11 +973,11 @@ namespace EpicGames.MCP.Automation
 		/// <returns>true if the build is registered, false otherwise</returns>
 		abstract public bool BuildExists(BuildPatchToolStagingInfo StagingInfo, string McpConfigName);
 
-        /// <summary>
-        /// Given a MCPStagingInfo defining our build info, posts the build to the MCP BuildInfo Service.
-        /// </summary>
-        /// <param name="stagingInfo">Staging Info describing the BuildInfo to post.</param>
-        abstract public void PostBuildInfo(BuildPatchToolStagingInfo stagingInfo);
+		/// <summary>
+		/// Given a MCPStagingInfo defining our build info, posts the build to the MCP BuildInfo Service.
+		/// </summary>
+		/// <param name="stagingInfo">Staging Info describing the BuildInfo to post.</param>
+		abstract public void PostBuildInfo(BuildPatchToolStagingInfo stagingInfo);
 
 		/// <summary>
 		/// Given a MCPStagingInfo defining our build info and a MCP config name, posts the build to the requested MCP BuildInfo Service.
@@ -944,6 +991,7 @@ namespace EpicGames.MCP.Automation
 		/// </summary>
 		/// <param name="BuildVersion">Build version to return labels for.</param>
 		/// <param name="McpConfigName">Which BuildInfo backend to get labels from for this promotion attempt.</param>
+		/// <returns>The list of build labels applied.</returns>
 		abstract public List<string> GetBuildLabels(BuildPatchToolStagingInfo StagingInfo, string McpConfigName);
 
 		/// <summary>
@@ -951,7 +999,7 @@ namespace EpicGames.MCP.Automation
 		/// </summary>
 		/// <param name="StagingInfo">Staging Info describing the BuildInfo to query.</param>
 		/// <param name="McpConfigName">Name of which MCP config to query.</param>
-		/// <returns></returns>
+		/// <returns>The manifest url.</returns>
 		abstract public string GetBuildManifestUrl(BuildPatchToolStagingInfo StagingInfo, string McpConfigName);
 
 		/// <summary>
@@ -976,6 +1024,7 @@ namespace EpicGames.MCP.Automation
 		/// </summary>
 		/// <param name="DestinationLabel">Base of label</param>
 		/// <param name="Platform">Platform to add to base label.</param>
+		/// <returns>The label string including platform postfix.</returns>
 		abstract public string GetLabelWithPlatform(string DestinationLabel, MCPPlatform Platform);
 
 		/// <summary>
@@ -983,15 +1032,16 @@ namespace EpicGames.MCP.Automation
 		/// </summary>
 		/// <param name="DestinationLabel">Base of label</param>
 		/// <param name="Platform">Platform to add to base label.</param>
+		/// <returns>The BuildVersion string including platform postfix.</returns>
 		abstract public string GetBuildVersionWithPlatform(BuildPatchToolStagingInfo StagingInfo);
 
 		/// <summary>
-		/// 
+		/// Get the BuildVersion for a build that is labeled under a specific appname.
 		/// </summary>
 		/// <param name="AppName">Application name to check the label in</param>
-		/// <param name="LabelName">Label name to get the build for</param>
+		/// <param name="LabelName">Label name to get the build version for</param>
 		/// <param name="McpConfigName">Which BuildInfo backend to label the build in.</param>
-		/// <returns></returns>
+		/// <returns>The BuildVersion or null if no build labeled.</returns>
 		abstract public string GetLabeledBuildVersion(string AppName, string LabelName, string McpConfigName);
 
 		/// <summary>
@@ -1002,16 +1052,16 @@ namespace EpicGames.MCP.Automation
 		/// <param name="McpConfigName">Which BuildInfo backend to label the build in.</param>
 		abstract public void LabelBuild(BuildPatchToolStagingInfo StagingInfo, string DestinationLabelWithPlatform, string McpConfigName);
 
-        /// <summary>
-        /// Informs Patcher Service of a new build availability after async labeling is complete
-        /// (this usually means the build was copied to a public file server before the label could be applied).
-        /// </summary>
-        /// <param name="Command">Parent command</param>
-        /// <param name="AppName">Application name that the patcher service will use.</param>
-        /// <param name="BuildVersion">BuildVersion string that the patcher service will use.</param>
-        /// <param name="ManifestRelativePath">Relative path to the Manifest file relative to the global build root (which is like P:\Builds) </param>
-        /// <param name="LabelName">Name of the label that we will be setting.</param>
-        abstract public void BuildPromotionCompleted(BuildPatchToolStagingInfo stagingInfo, string AppName, string BuildVersion, string ManifestRelativePath, string PlatformName, string LabelName);
+		/// <summary>
+		/// Informs Patcher Service of a new build availability after async labeling is complete
+		/// (this usually means the build was copied to a public file server before the label could be applied).
+		/// </summary>
+		/// <param name="Command">Parent command</param>
+		/// <param name="AppName">Application name that the patcher service will use.</param>
+		/// <param name="BuildVersion">BuildVersion string that the patcher service will use.</param>
+		/// <param name="ManifestRelativePath">Relative path to the Manifest file relative to the global build root (which is like P:\Builds) </param>
+		/// <param name="LabelName">Name of the label that we will be setting.</param>
+		abstract public void BuildPromotionCompleted(BuildPatchToolStagingInfo stagingInfo, string AppName, string BuildVersion, string ManifestRelativePath, string PlatformName, string LabelName);
 	}
 
     /// <summary>
diff --git a/Engine/Source/Programs/AutomationTool/Scripts/AutomationScripts.Automation.csproj b/Engine/Source/Programs/AutomationTool/Scripts/AutomationScripts.Automation.csproj
--- a/Engine/Source/Programs/AutomationTool/Scripts/AutomationScripts.Automation.csproj
+++ b/Engine/Source/Programs/AutomationTool/Scripts/AutomationScripts.Automation.csproj
@@ -98,7 +98,6 @@
     <Compile Include="ReplaceAssetsUsingManifest.cs" />
     <Compile Include="ResavePackagesCommand.Automation.cs" />
     <Compile Include="IPhonePackager.Automation.cs" />
-    <Compile Include="LauncherLocalization.Automation.cs" />
     <Compile Include="ListMobileDevices.Automation.cs">
       <SubType>Code</SubType>
     </Compile>
@@ -138,10 +137,6 @@
       <Project>{c2f959b1-1727-4598-9c32-3d9acce72750}</Project>
       <Name>Localization.Automation</Name>
     </ProjectReference>
-    <ProjectReference Include="..\OneSkyLocalization\OneSkyLocalization.Automation.csproj">
-      <Project>{39369368-5603-46da-baae-766bb3cb2b51}</Project>
-      <Name>OneSkyLocalization.Automation</Name>
-    </ProjectReference>
   </ItemGroup>
   <Import Project="$(MSBuildToolsPath)\Microsoft.CSharp.targets" />
 </Project>
\ No newline at end of file
diff --git a/Engine/Source/Programs/AutomationTool/XLocLocalization/XLocLocalizationProvider.cs b/Engine/Source/Programs/AutomationTool/XLocLocalization/XLocLocalizationProvider.cs
--- a/Engine/Source/Programs/AutomationTool/XLocLocalization/XLocLocalizationProvider.cs
+++ b/Engine/Source/Programs/AutomationTool/XLocLocalization/XLocLocalizationProvider.cs
@@ -23,7 +23,7 @@ public struct XLocConfig
 	public string LocalizationId;
 };
 
-struct XLocUtils
+public struct XLocUtils
 {
 	public static string MD5HashString(string Str)
 	{
@@ -288,7 +288,7 @@ public abstract class XLocLocalizationProvider : LocalizationProvider
 		}
 	}
 
-	private XLocApiClient CreateXLocApiClient()
+	protected XLocApiClient CreateXLocApiClient()
 	{
 		var Binding = new BasicHttpBinding();
 		Binding.Name = "basicHttpXLocApi";
@@ -317,7 +317,7 @@ public abstract class XLocLocalizationProvider : LocalizationProvider
 		return new XLocApiClient(Binding, Endpoint);
 	}
 
-	private TransferServiceClient CreateTransferServiceClient()
+	protected TransferServiceClient CreateTransferServiceClient()
 	{
 		var Binding = new BasicHttpBinding();
 		Binding.Name = "transfer";
@@ -329,12 +329,12 @@ public abstract class XLocLocalizationProvider : LocalizationProvider
 		return new TransferServiceClient(Binding, Endpoint);
 	}
 
-	private string RequestAuthToken(XLocApiClient XLocApiClient)
+	protected string RequestAuthToken(XLocApiClient XLocApiClient)
 	{
 		return XLocApiClient.GetAuthToken(Config.APIKey, XLocUtils.MD5HashString(Config.APIKey + Config.APISecret));
 	}
 
-	private string RequestAuthTokenWithRetry(XLocApiClient XLocApiClient)
+	protected string RequestAuthTokenWithRetry(XLocApiClient XLocApiClient)
 	{
 		const int MAX_COUNT = 3;
 
@@ -359,7 +359,7 @@ public abstract class XLocLocalizationProvider : LocalizationProvider
 		}
 	}
 
-	private string RequestLatestBuild(XLocApiClient XLocApiClient, string AuthToken, string LanguageId, string RemoteFilename)
+	protected string RequestLatestBuild(XLocApiClient XLocApiClient, string AuthToken, string LanguageId, string RemoteFilename)
 	{
 		return XLocApiClient.GetLatestBuildByFile(Config.APIKey, AuthToken, XLocUtils.MD5HashString(Config.APIKey + Config.APISecret + Config.LocalizationId + LanguageId + RemoteFilename), Config.LocalizationId, LanguageId, RemoteFilename);
 	}
diff --git a/Engine/Source/Programs/BuildPatchTool/BuildPatchTool.Target.cs b/Engine/Source/Programs/BuildPatchTool/BuildPatchTool.Target.cs
--- a/Engine/Source/Programs/BuildPatchTool/BuildPatchTool.Target.cs
+++ b/Engine/Source/Programs/BuildPatchTool/BuildPatchTool.Target.cs
@@ -23,5 +23,6 @@ public class BuildPatchToolTarget : TargetRules
 		bUseChecksInShipping = true;
 		bIsBuildingConsoleApplication = true;
 		bHasExports = false;
+		bWithServerCode = false;
 	}
 }
diff --git a/Engine/Source/Programs/BuildPatchTool/Private/BuildPatchTool.cpp b/Engine/Source/Programs/BuildPatchTool/Private/BuildPatchTool.cpp
--- a/Engine/Source/Programs/BuildPatchTool/Private/BuildPatchTool.cpp
+++ b/Engine/Source/Programs/BuildPatchTool/Private/BuildPatchTool.cpp
@@ -16,14 +16,16 @@ class FBuildPatchOutputDevice : public FOutputDevice
 public:
 	virtual void Serialize( const TCHAR* V, ELogVerbosity::Type Verbosity, const class FName& Category ) override
 	{
-#if PLATFORM_WINDOWS
+		// Only forward verbosities higher than Display as they will already be sent to stdout.
+		if (Verbosity > ELogVerbosity::Display)
+		{
 #if PLATFORM_USE_LS_SPEC_FOR_WIDECHAR
-		printf("\n%ls", *FOutputDeviceHelper::FormatLogLine(Verbosity, Category, V, GPrintLogTimes));
+			printf("\n%ls", *FOutputDeviceHelper::FormatLogLine(Verbosity, Category, V, GPrintLogTimes));
 #else
-		wprintf(TEXT("\n%s"), *FOutputDeviceHelper::FormatLogLine(Verbosity, Category, V, GPrintLogTimes));
-#endif
-		fflush( stdout );
+			wprintf(TEXT("\n%s"), *FOutputDeviceHelper::FormatLogLine(Verbosity, Category, V, GPrintLogTimes));
 #endif
+			fflush( stdout );
+		}
 	}
 };
 
@@ -69,14 +71,17 @@ const TCHAR* HandleLegacyCommandline(const TCHAR* CommandLine)
 
 EReturnCode RunBuildPatchTool()
 {
-	// Load the BuildPatchServices Module
-	IBuildPatchServicesModule& BuildPatchServicesModule = FModuleManager::LoadModuleChecked<IBuildPatchServicesModule>(TEXT("BuildPatchServices"));
-
-	// Initialise the UObject system and process our uobject classes
+	// Initialise the UObject module.
 	FModuleManager::Get().LoadModule(TEXT("CoreUObject"));
 	FCoreDelegates::OnInit.Broadcast();
+
+	// Load the BuildPatchServices Module.
+	IBuildPatchServicesModule& BuildPatchServicesModule = FModuleManager::LoadModuleChecked<IBuildPatchServicesModule>(TEXT("BuildPatchServices"));
+
+	// Make sure we have processed UObjects from BPS.
 	ProcessNewlyLoadedUObjects();
 
+	// Instantiate and execute the tool.
 	TSharedRef<IToolMode> ToolMode = FToolModeFactory::Create(BuildPatchServicesModule);
 	return ToolMode->Execute();
 }
@@ -84,7 +89,10 @@ EReturnCode RunBuildPatchTool()
 EReturnCode BuildPatchToolMain(const TCHAR* CommandLine)
 {
 	// Add log device for stdout
-	GLog->AddOutputDevice(new FBuildPatchOutputDevice());
+	if (FParse::Param(CommandLine, TEXT("stdout")))
+	{
+		GLog->AddOutputDevice(new FBuildPatchOutputDevice());
+	}
 
 	// Handle legacy commandlines
 	CommandLine = HandleLegacyCommandline(CommandLine);
diff --git a/Engine/Source/Programs/BuildPatchTool/Private/ToolModes/PackageChunksMode.cpp b/Engine/Source/Programs/BuildPatchTool/Private/ToolModes/PackageChunksMode.cpp
--- a/Engine/Source/Programs/BuildPatchTool/Private/ToolModes/PackageChunksMode.cpp
+++ b/Engine/Source/Programs/BuildPatchTool/Private/ToolModes/PackageChunksMode.cpp
@@ -39,8 +39,12 @@ public:
 			UE_LOG(LogBuildPatchTool, Log, TEXT("  -OutputFile=\"\"       Specifies in quotes the file path the output package. Extension of .chunkdb will be added if not present."));
 			UE_LOG(LogBuildPatchTool, Log, TEXT(""));
 			UE_LOG(LogBuildPatchTool, Log, TEXT("Optional arguments:"));
+			UE_LOG(LogBuildPatchTool, Log, TEXT("  -PrevManifestFile=\"\" Specifies in quotes the file path to a manifest for a previous build, this will be used to filter out chunks, such that the"));
+			UE_LOG(LogBuildPatchTool, Log, TEXT("                       produced chunkdb files will only contain chunks required to patch from this build to the one described by ManifestFile."));
 			UE_LOG(LogBuildPatchTool, Log, TEXT("  -CloudDir=\"\"         Specifies in quotes the cloud directory where chunks to be packaged can be found."));
 			UE_LOG(LogBuildPatchTool, Log, TEXT("  -MaxOutputFileSize=  When specified, the size of each output file (in bytes) will be limited to a maximum of the provided value."));
+			UE_LOG(LogBuildPatchTool, Log, TEXT("  -ResultDataFile=\"\"   Specifies in quotes the file path where the results will be exported as a JSON object."));
+			UE_LOG(LogBuildPatchTool, Log, TEXT("  -TagSets=\"t1,t2\"     Specifies in quotes a comma seperated list of tags for filtering of data saved. Multiple sets can be provided to split the chunkdb files by tagsets."));
 			UE_LOG(LogBuildPatchTool, Log, TEXT(""));
 			UE_LOG(LogBuildPatchTool, Log, TEXT("NB: If CloudDir is not specified, the manifest file location will be used as the cloud directory."));
 			UE_LOG(LogBuildPatchTool, Log, TEXT("NB: MaxOutputFileSize is recommended to be as large as possible. The minimum individual chunkdb filesize is equal to one chunk plus chunkdb"));
@@ -50,11 +54,17 @@ public:
 			UE_LOG(LogBuildPatchTool, Log, TEXT("    have the number of digits required for highest numbered part."));
 			UE_LOG(LogBuildPatchTool, Log, TEXT("NB: If MaxOutputFileSize is specified, then each part can be equal to or less than the specified size, depending on the size of the last chunk"));
 			UE_LOG(LogBuildPatchTool, Log, TEXT("    that fits."));
+			UE_LOG(LogBuildPatchTool, Log, TEXT("NB: When providing multiple -TagSets= arguments, all data from the first -TagSets= arg will be saved first, followed by any extra data needed for the second -TagSets= arg, and so on in separated chunkdb files."));
+			UE_LOG(LogBuildPatchTool, Log, TEXT("    Note that this means the chunkdb files produced for the second -TagSets= arg and later will not contain some required data for that tagset if the data already got saved out as part of a previous tagset."));
+			UE_LOG(LogBuildPatchTool, Log, TEXT("    The chunkdb files are thus additive with no dupes."));
+			UE_LOG(LogBuildPatchTool, Log, TEXT("    If it is desired that each tagset's chunkdb files contain the duplicate data, then PackageChunks should be executed once per -TagSets= arg rather than once will all -TagSets= args."));
+			UE_LOG(LogBuildPatchTool, Log, TEXT("    An empty tag must be included in one of the -TagSets= args to include untagged file data in that tagset, e.g. -TagSets=\" , t1\"."));
+			UE_LOG(LogBuildPatchTool, Log, TEXT("    Adding no -TagSets= args will include all data."));
 			return EReturnCode::OK;
 		}
 
 		// Run the enumeration routine.
-		bool bSuccess = BpsInterface.PackageChunkData(ManifestFile, OutputFile, CloudDir, MaxOutputFileSize);
+		bool bSuccess = BpsInterface.PackageChunkData(ManifestFile, PrevManifestFile, TagSetArray, OutputFile, CloudDir, MaxOutputFileSize, ResultDataFile);
 		return bSuccess ? EReturnCode::OK : EReturnCode::ToolFailure;
 	}
 
@@ -62,8 +72,9 @@ private:
 
 	bool ProcessCommandline()
 	{
-#define HAS_SWITCH(Switch) (Algo::FindByPredicate(Switches, [](const FString& Elem){ return Elem.StartsWith(TEXT(#Switch "="));}) != nullptr)
-#define PARSE_SWITCH(Switch) ParseSwitch(TEXT(#Switch "="), Switch, Switches)
+#define HAS_SWITCH(SwitchVar) (Algo::FindByPredicate(Switches, [](const FString& Elem){ return Elem.StartsWith(TEXT(#SwitchVar "="));}) != nullptr)
+#define PARSE_SWITCH(SwitchVar) ParseSwitch(TEXT(#SwitchVar "="), SwitchVar, Switches)
+#define PARSE_SWITCHES(SwitchesVar) ParseSwitches(TEXT(#SwitchesVar "="), SwitchesVar, Switches)
 		TArray<FString> Tokens, Switches;
 		FCommandLine::Parse(FCommandLine::Get(), Tokens, Switches);
 
@@ -84,6 +95,12 @@ private:
 		FPaths::NormalizeFilename(OutputFile);
 
 		// Get optional parameters.
+		PARSE_SWITCH(PrevManifestFile);
+		PARSE_SWITCH(ResultDataFile);
+		FPaths::NormalizeFilename(PrevManifestFile);
+		FPaths::NormalizeFilename(ResultDataFile);
+		PARSE_SWITCHES(TagSets);
+
 		if (!PARSE_SWITCH(CloudDir))
 		{
 			// If not provided we use the location of the manifest file.
@@ -106,7 +123,29 @@ private:
 			MaxOutputFileSize = MAX_uint64;
 		}
 
+		// Process the tagsets that we parsed.
+		if (TagSets.Num() > 0)
+		{
+			for (const FString& TagSet : TagSets)
+			{
+				TArray<FString> Tags;
+				const bool bCullEmpty = false;
+				TagSet.ParseIntoArray(Tags, TEXT(","), bCullEmpty);
+				for (FString& Tag : Tags)
+				{
+					Tag.TrimStartAndEndInline();
+				}
+				// If we ended up with an empty array, the intention would have been to pass a tagset that contains just an empty string, so make that fixup.
+				if (Tags.Num() == 0)
+				{
+					Tags.Add(TEXT(""));
+				}
+				TagSetArray.Emplace(Tags);
+			}
+		}
+
 		return true;
+#undef PARSE_SWITCHES
 #undef PARSE_SWITCH
 #undef HAS_SWITCH
 	}
@@ -115,9 +154,13 @@ private:
 	IBuildPatchServicesModule& BpsInterface;
 	bool bHelp;
 	FString ManifestFile;
+	FString PrevManifestFile;
 	FString OutputFile;
+	FString ResultDataFile;
 	FString CloudDir;
 	uint64 MaxOutputFileSize;
+	TArray<FString> TagSets;
+	TArray<TSet<FString>> TagSetArray;
 };
 
 BuildPatchTool::IToolModeRef BuildPatchTool::FPackageChunksToolModeFactory::Create(IBuildPatchServicesModule& BpsInterface)
diff --git a/Engine/Source/Programs/BuildPatchTool/Public/Interfaces/ToolMode.h b/Engine/Source/Programs/BuildPatchTool/Public/Interfaces/ToolMode.h
--- a/Engine/Source/Programs/BuildPatchTool/Public/Interfaces/ToolMode.h
+++ b/Engine/Source/Programs/BuildPatchTool/Public/Interfaces/ToolMode.h
@@ -21,10 +21,10 @@ namespace BuildPatchTool
 		 * Helper for parsing a switch from an array of switches, usually produced using FCommandLine::Parse(..)
 		 *
 		 * @param InSwitch      The switch name, ending with =. E.g. option=, foo=. It would usually be a compile time const.
-		 * @param Value         FString to receive the value for the switch
-		 * @param Switches      The array of switches to search through
-		 * @return true if the switch was found
-		**/
+		 * @param Value         Receives the value from the switch.
+		 * @param Switches      The array of switches to search through.
+		 * @return true if the switch was found.
+		 */
 		template <typename TValueType>
 		bool ParseSwitch(const TCHAR* InSwitch, TValueType& Value, const TArray<FString>& Switches)
 		{
@@ -44,6 +44,39 @@ namespace BuildPatchTool
 			return false;
 		}
 
+		/**
+		 * Helper for parsing an array of multiple same name switches from the full array of switches, usually produced using FCommandLine::Parse(..)
+		 *
+		 * @param InSwitch      The switch name, ending with =. E.g. option=, foo=. It would usually be a compile time const.
+		 * @param Values        Receives the values from the switches.
+		 * @param Switches      The array of switches to search through.
+		 * @return true if at least one match was found.
+		 */
+		template <typename TValueType>
+		bool ParseSwitches(const TCHAR* InSwitch, TArray<TValueType>& Values, const TArray<FString>& Switches)
+		{
+			// Debug check requirements for InSwitch
+			checkSlow(InSwitch != nullptr);
+			checkSlow(InSwitch[FCString::Strlen(InSwitch)-1] == TEXT('='));
+
+			bool bFoundValue = false;
+			for (const FString& Switch : Switches)
+			{
+				if (Switch.StartsWith(InSwitch))
+				{
+					FString StringValue;
+					TValueType Value;
+					Switch.Split(EqualsStr, nullptr, &StringValue);
+					if (ParseValue(StringValue, Value))
+					{
+						Values.Emplace(MoveTemp(Value));
+						bFoundValue = true;
+					}
+				}
+			}
+			return bFoundValue;
+		}
+
 		bool ParseOption(const TCHAR* InSwitch, const TArray<FString>& Switches)
 		{
 			return Switches.Contains(InSwitch);
diff --git a/Engine/Source/Programs/UnrealBuildTool/Platform/Mac/MacToolChain.cs b/Engine/Source/Programs/UnrealBuildTool/Platform/Mac/MacToolChain.cs
--- a/Engine/Source/Programs/UnrealBuildTool/Platform/Mac/MacToolChain.cs
+++ b/Engine/Source/Programs/UnrealBuildTool/Platform/Mac/MacToolChain.cs
@@ -1196,7 +1196,7 @@ namespace UnrealBuildTool
 						AppendMacLine(FinalizeAppBundleScript, FormatCopyCommand(String.Format("{0}/Runtime/Launch/Resources/Mac/UProject.icns", EngineSourcePath), String.Format("{0}.app/Contents/Resources/UProject.icns", ExeName)));
 					}
 
-					string InfoPlistFile = CustomResourcesPath + (bBuildingEditor ? "Info-Editor.plist" : "Info.plist");
+					string InfoPlistFile = CustomResourcesPath + (bBuildingEditor ? "/Info-Editor.plist" : "/Info.plist");
 					if (!File.Exists(InfoPlistFile))
 					{
 						InfoPlistFile = EngineSourcePath + "/Runtime/Launch/Resources/Mac/" + (bBuildingEditor ? "Info-Editor.plist" : "Info.plist");
diff --git a/Engine/Source/Runtime/Analytics/AnalyticsET/Private/IAnalyticsProviderET.cpp b/Engine/Source/Runtime/Analytics/AnalyticsET/Private/IAnalyticsProviderET.cpp
--- a/Engine/Source/Runtime/Analytics/AnalyticsET/Private/IAnalyticsProviderET.cpp
+++ b/Engine/Source/Runtime/Analytics/AnalyticsET/Private/IAnalyticsProviderET.cpp
@@ -20,6 +20,7 @@
 #include "HttpModule.h"
 #include "PlatformHttp.h"
 #include "Misc/EngineVersion.h"
+#include "HttpRetrySystem.h"
 
 
 /** When enabled (and -AnalyticsTrackPerf is specified on the command line, will log out analytics flush timings on a regular basis to Saved/AnalyticsTiming.csv. */
@@ -200,6 +201,8 @@ public:
 	virtual const TArray<FAnalyticsEventAttribute>& GetDefaultEventAttributes() const override;
 	virtual void SetEventCallback(const OnEventRecorded& Callback) override;
 
+	virtual void SetURLEndpoint(const FString& UrlEndpoint, const TArray<FString>& AltDomains) override;
+
 	virtual ~FAnalyticsProviderET();
 
 	FString GetAPIKey() const { return APIKey; }
@@ -217,6 +220,9 @@ private:
 		return CachedEvents.Num() > 1;
 	}
 
+	/** Create a request utilizing HttpRetry domains */
+	TSharedRef<IHttpRequest> CreateRequest();
+
 	bool bSessionInProgress;
 	/** ET Game API Key - Get from your account manager */
 	FString APIKey;
@@ -292,6 +298,9 @@ private:
 	* Delegate called when an event Http request completes
 	*/
 	void EventRequestComplete(FHttpRequestPtr HttpRequest, FHttpResponsePtr HttpResponse, bool bSucceeded, TSharedPtr< TArray<FAnalyticsEventEntry> > FlushedEvents);
+
+	TSharedPtr<class FHttpRetrySystem::FManager> HttpRetryManager;
+	FHttpRetrySystem::FRetryDomainsPtr RetryServers;
 };
 
 TSharedPtr<IAnalyticsProviderET> FAnalyticsET::CreateAnalyticsProvider(const Config& ConfigValues) const
@@ -327,6 +336,28 @@ FAnalyticsProviderET::FAnalyticsProviderET(const FAnalyticsET::Config& ConfigVal
 		UE_LOG(LogAnalytics, Fatal, TEXT("AnalyticsET: APIKey (%s) and APIServer (%s) cannot be empty!"), *APIKey, *APIServer);
 	}
 
+	// Set the number of retries to the number of retry URLs that have been passed in.
+	uint32 RetryLimitCount = ConfigValues.AltAPIServersET.Num();
+
+	HttpRetryManager = MakeShared<FHttpRetrySystem::FManager>(
+		FHttpRetrySystem::FRetryLimitCountSetting::Create(RetryLimitCount),
+		FHttpRetrySystem::FRetryTimeoutRelativeSecondsSetting::Unused()
+		);
+
+	// If we have retry domains defined, insert the default domain into the list
+	if (RetryLimitCount > 0)
+	{
+		TArray<FString> TmpAltAPIServers = ConfigValues.AltAPIServersET;
+
+		FString DefaultUrlDomain = FPlatformHttp::GetUrlDomain(APIServer);
+		if (!TmpAltAPIServers.Contains(DefaultUrlDomain))
+		{
+			TmpAltAPIServers.Insert(DefaultUrlDomain, 0);
+		}
+
+		RetryServers = MakeShared<FHttpRetrySystem::FRetryDomains, ESPMode::ThreadSafe>(MoveTemp(TmpAltAPIServers));
+	}
+
 	// force very verbose logging if we are force-disabling events.
 	bool bForceDisableCaching = FParse::Param(FCommandLine::Get(), TEXT("ANALYTICSDISABLECACHING"));
 	if (bForceDisableCaching)
@@ -377,6 +408,8 @@ bool FAnalyticsProviderET::Tick(float DeltaSeconds)
 {
 	QUICK_SCOPE_CYCLE_COUNTER(STAT_FAnalyticsProviderET_Tick);
 
+	HttpRetryManager->Update();
+
 	// There are much better ways to do this, but since most events are recorded and handled on the same (game) thread,
 	// this is probably mostly fine for now, and simply favoring not crashing at the moment
 	FScopeLock ScopedLock(&CachedEventsCS);
@@ -468,6 +501,18 @@ void FAnalyticsProviderET::EndSession()
 	bSessionInProgress = false;
 }
 
+TSharedRef<IHttpRequest> FAnalyticsProviderET::CreateRequest()
+{
+	// TODO add config values for retries, for now, using default
+	TSharedRef<IHttpRequest> HttpRequest = HttpRetryManager->CreateRequest(FHttpRetrySystem::FRetryLimitCountSetting::Unused(),
+		FHttpRetrySystem::FRetryTimeoutRelativeSecondsSetting::Unused(),
+		FHttpRetrySystem::FRetryResponseCodes(),
+		FHttpRetrySystem::FRetryVerbs(),
+		RetryServers);
+
+	return HttpRequest;
+}
+
 void FAnalyticsProviderET::FlushEvents()
 {
 	// Warn if this takes more than 2 ms
@@ -620,12 +665,12 @@ void FAnalyticsProviderET::FlushEvents()
 			}
 
 			// Create/send Http request for an event
-			TSharedRef<IHttpRequest> HttpRequest = FHttpModule::Get().CreateRequest();
+			TSharedRef<IHttpRequest> HttpRequest = CreateRequest();
 			HttpRequest->SetHeader(TEXT("Content-Type"), TEXT("application/json; charset=utf-8"));
-
-			HttpRequest->SetURL(APIServer + URLPath);
+			HttpRequest->SetURL(APIServer / URLPath);
 			HttpRequest->SetVerb(TEXT("POST"));
 			HttpRequest->SetContentAsString(Payload);
+
 			// Don't set a response callback if we are in our destructor, as the instance will no longer be there to call.
 			if (!bInDestructor)
 			{
@@ -678,7 +723,7 @@ void FAnalyticsProviderET::FlushEvents()
 						*EventParams);
 
 					// Create/send Http request for an event
-					TSharedRef<IHttpRequest> HttpRequest = FHttpModule::Get().CreateRequest();
+					TSharedRef<IHttpRequest> HttpRequest = CreateRequest();
 					HttpRequest->SetHeader(TEXT("Content-Type"), TEXT("text/plain"));
 
 					// Don't need to URL encode the APIServer or the EventParams, which are already encoded, and contain parameter separaters that we DON'T want encoded.
@@ -864,29 +909,57 @@ void FAnalyticsProviderET::EventRequestComplete(FHttpRequestPtr HttpRequest, FHt
 
 		// if FlushedEvents is passed, re-queue the events for next time
 		if (FlushedEvents.IsValid())
-	{
-		// add a dropped submission event so we can see how often this is happening
-		if (bShouldCacheEvents && CachedEvents.Num() < 1024)
 		{
-			TArray<FAnalyticsEventAttribute> Attributes;
-			Attributes.Emplace(FAnalyticsEventAttribute(FString(TEXT("HTTP_STATUS")), FString::Printf(TEXT("%d"), HttpResponse.IsValid() ? HttpResponse->GetResponseCode() : 0)));
-			Attributes.Emplace(FAnalyticsEventAttribute(FString(TEXT("EVENTS_IN_BATCH")), FString::Printf(TEXT("%d"), FlushedEvents->Num())));
-			Attributes.Emplace(FAnalyticsEventAttribute(FString(TEXT("EVENTS_QUEUED")), FString::Printf(TEXT("%d"), CachedEvents.Num())));
-			CachedEvents.Emplace(FAnalyticsEventEntry(FString(TEXT("ET.DroppedSubmission")), MoveTemp(Attributes), false, false));
-		}
+			// add a dropped submission event so we can see how often this is happening
+			if (bShouldCacheEvents && CachedEvents.Num() < 1024)
+			{
+				TArray<FAnalyticsEventAttribute> Attributes;
+				Attributes.Emplace(FAnalyticsEventAttribute(FString(TEXT("HTTP_STATUS")), FString::Printf(TEXT("%d"), HttpResponse.IsValid() ? HttpResponse->GetResponseCode() : 0)));
+				Attributes.Emplace(FAnalyticsEventAttribute(FString(TEXT("EVENTS_IN_BATCH")), FString::Printf(TEXT("%d"), FlushedEvents->Num())));
+				Attributes.Emplace(FAnalyticsEventAttribute(FString(TEXT("EVENTS_QUEUED")), FString::Printf(TEXT("%d"), CachedEvents.Num())));
+				CachedEvents.Emplace(FAnalyticsEventEntry(FString(TEXT("ET.DroppedSubmission")), MoveTemp(Attributes), false, false));
+			}
 
-		// if we're being super spammy or have been offline forever, just leave it at the ET.DroppedSubmission event
-		if (bShouldCacheEvents && CachedEvents.Num() < 256)
-		{
-			UE_LOG(LogAnalytics, Log, TEXT("[%s] ET Requeuing %d analytics events due to failure to send"), *APIKey, FlushedEvents->Num());
+			// if we're being super spammy or have been offline forever, just leave it at the ET.DroppedSubmission event
+			if (bShouldCacheEvents && CachedEvents.Num() < 256)
+			{
+				UE_LOG(LogAnalytics, Log, TEXT("[%s] ET Requeuing %d analytics events due to failure to send"), *APIKey, FlushedEvents->Num());
 
-			// put them at the beginning since it should include a default attributes entry and we don't want to change the current default attributes
-			CachedEvents.Insert(*FlushedEvents, 0);
+				// put them at the beginning since it should include a default attributes entry and we don't want to change the current default attributes
+				CachedEvents.Insert(*FlushedEvents, 0);
+			}
+			else
+			{
+				UE_LOG(LogAnalytics, Error, TEXT("[%s] ET dropping %d analytics events due to too many in queue (%d)"), *APIKey, FlushedEvents->Num(), CachedEvents.Num());
+			}
 		}
-		else
+	}
+}
+
+void FAnalyticsProviderET::SetURLEndpoint(const FString& UrlEndpoint, const TArray<FString>& AltDomains)
+{
+	APIServer = UrlEndpoint;
+
+	// Set the number of retries to the number of retry URLs that have been passed in.
+	uint32 RetryLimitCount = AltDomains.Num();
+
+	HttpRetryManager->SetDefaultRetryLimit(RetryLimitCount);
+
+	TArray<FString> TmpAltAPIServers = AltDomains;
+
+	// If we have retry domains defined, insert the default domain into the list
+	if (RetryLimitCount > 0)
+	{
+		FString DefaultUrlDomain = FPlatformHttp::GetUrlDomain(APIServer);
+		if (!TmpAltAPIServers.Contains(DefaultUrlDomain))
 		{
-			UE_LOG(LogAnalytics, Error, TEXT("[%s] ET dropping %d analytics events due to too many in queue (%d)"), *APIKey, FlushedEvents->Num(), CachedEvents.Num());
+			TmpAltAPIServers.Insert(DefaultUrlDomain, 0);
 		}
+
+		RetryServers = MakeShared<FHttpRetrySystem::FRetryDomains, ESPMode::ThreadSafe>(MoveTemp(TmpAltAPIServers));
 	}
-}
-}
+	else
+	{
+		RetryServers.Reset();
+	}
+}
\ No newline at end of file
diff --git a/Engine/Source/Runtime/Analytics/AnalyticsET/Public/AnalyticsET.h b/Engine/Source/Runtime/Analytics/AnalyticsET/Public/AnalyticsET.h
--- a/Engine/Source/Runtime/Analytics/AnalyticsET/Public/AnalyticsET.h
+++ b/Engine/Source/Runtime/Analytics/AnalyticsET/Public/AnalyticsET.h
@@ -40,6 +40,8 @@ public:
 		FString APIKeyET;
 		/** ET API Server - Base URL to send events. */
 		FString APIServerET;
+		/** ET Alt API Servers - Base URLs to send events on retry. */
+		TArray<FString> AltAPIServersET;
 		/** 
 		 * AppVersion - defines the app version passed to the provider. By default this will be FEngineVersion::Current(), but you can supply your own. 
 		 * As a convenience, you can use -AnalyticsAppVersion=XXX to force the AppVersion to a specific value. Useful for playtest etc where you want to define a specific version string dynamically.
@@ -54,13 +56,16 @@ public:
 		FString AppEnvironment;
 		/** The UploadType that the data router should use. Defaults to GetDefaultUploadType. */
 		FString UploadType;
+		/** Maximum number of retries to attempt. */
+		uint32 RetryLimitCount;
 
 		/** Default ctor to ensure all values have their proper default. */
 		Config() : UseLegacyProtocol(false) {}
 		/** Ctor exposing common configurables . */
-		Config(FString InAPIKeyET, FString InAPIServerET, FString InAppVersionET = FString(), bool InUseLegacyProtocol = false, FString InAppEnvironment = FString(), FString InUploadType = FString()) 
+		Config(FString InAPIKeyET, FString InAPIServerET, FString InAppVersionET = FString(), bool InUseLegacyProtocol = false, FString InAppEnvironment = FString(), FString InUploadType = FString(), TArray<FString> InAltApiServers = TArray<FString>()) 
 			: APIKeyET(MoveTemp(InAPIKeyET))
 			, APIServerET(MoveTemp(InAPIServerET))
+			, AltAPIServersET(MoveTemp(InAltApiServers))
 			, AppVersionET(MoveTemp(InAppVersionET))
 			, UseLegacyProtocol(InUseLegacyProtocol)
 			, AppEnvironment(MoveTemp(InAppEnvironment))
diff --git a/Engine/Source/Runtime/Analytics/AnalyticsET/Public/IAnalyticsProviderET.h b/Engine/Source/Runtime/Analytics/AnalyticsET/Public/IAnalyticsProviderET.h
--- a/Engine/Source/Runtime/Analytics/AnalyticsET/Public/IAnalyticsProviderET.h
+++ b/Engine/Source/Runtime/Analytics/AnalyticsET/Public/IAnalyticsProviderET.h
@@ -88,6 +88,11 @@ public:
 	*/
 	virtual const TArray<FAnalyticsEventAttribute>& GetDefaultEventAttributes() const = 0;
 
+	/**
+	 * Updates the default URL endpoint and AltDomains.
+	 */
+	virtual void SetURLEndpoint(const FString& UrlEndpoint, const TArray<FString>& AltDomains) = 0;
+
 	/**
 	* Set a callback to be invoked any time an event is queued.
 	* 
diff --git a/Engine/Source/Runtime/ApplicationCore/Private/Windows/WindowsApplication.cpp b/Engine/Source/Runtime/ApplicationCore/Private/Windows/WindowsApplication.cpp
--- a/Engine/Source/Runtime/ApplicationCore/Private/Windows/WindowsApplication.cpp
+++ b/Engine/Source/Runtime/ApplicationCore/Private/Windows/WindowsApplication.cpp
@@ -1329,11 +1329,13 @@ int32 FWindowsApplication::ProcessMessage( HWND hwnd, uint32 msg, WPARAM wParam,
 					BorderHeight = BorderRect.bottom - BorderRect.top;
 				}
 
+				const float DPIScaleFactor = CurrentNativeEventWindow->GetDPIScaleFactor();
+
 				// We always apply BorderWidth and BorderHeight since Slate always works with client area window sizes
-				MinMaxInfo->ptMinTrackSize.x = FMath::RoundToInt( SizeLimits.GetMinWidth().Get(MinMaxInfo->ptMinTrackSize.x) );
-				MinMaxInfo->ptMinTrackSize.y = FMath::RoundToInt( SizeLimits.GetMinHeight().Get(MinMaxInfo->ptMinTrackSize.y) );
-				MinMaxInfo->ptMaxTrackSize.x = FMath::RoundToInt( SizeLimits.GetMaxWidth().Get(MinMaxInfo->ptMaxTrackSize.x) ) + BorderWidth;
-				MinMaxInfo->ptMaxTrackSize.y = FMath::RoundToInt( SizeLimits.GetMaxHeight().Get(MinMaxInfo->ptMaxTrackSize.y) ) + BorderHeight;
+				MinMaxInfo->ptMinTrackSize.x = FMath::RoundToInt( SizeLimits.GetMinWidth().Get(MinMaxInfo->ptMinTrackSize.x) * DPIScaleFactor );
+				MinMaxInfo->ptMinTrackSize.y = FMath::RoundToInt( SizeLimits.GetMinHeight().Get(MinMaxInfo->ptMinTrackSize.y) * DPIScaleFactor );
+				MinMaxInfo->ptMaxTrackSize.x = FMath::RoundToInt( SizeLimits.GetMaxWidth().Get(MinMaxInfo->ptMaxTrackSize.x) * DPIScaleFactor ) + BorderWidth;
+				MinMaxInfo->ptMaxTrackSize.y = FMath::RoundToInt( SizeLimits.GetMaxHeight().Get(MinMaxInfo->ptMaxTrackSize.y) * DPIScaleFactor ) + BorderHeight;
 				return 0;
 			}
 			break;
diff --git a/Engine/Source/Runtime/Core/Private/Internationalization/Text.cpp b/Engine/Source/Runtime/Core/Private/Internationalization/Text.cpp
--- a/Engine/Source/Runtime/Core/Private/Internationalization/Text.cpp
+++ b/Engine/Source/Runtime/Core/Private/Internationalization/Text.cpp
@@ -678,32 +678,46 @@ FText FText::AsTimespan(const FTimespan& Timespan, const FCulturePtr& TargetCult
 	}
 }
 
-FText FText::AsMemory(uint64 NumBytes, const FNumberFormattingOptions* const Options, const FCulturePtr& TargetCulture)
+FText FText::AsMemory(uint64 NumBytes, const FNumberFormattingOptions* const Options, const FCulturePtr& TargetCulture, EMemoryUnitStandard UnitStandard)
 {
 	checkf(FInternationalization::Get().IsInitialized() == true, TEXT("FInternationalization is not initialized. An FText formatting method was likely used in static object initialization - this is not supported."));
 	FFormatNamedArguments Args;
 
-	if (NumBytes < 1024)
+	static const TCHAR* Prefixes = TEXT("KMGTPEZYkMGTPEZY");
+	int32 Prefix = 0;
+	FString Suffix = TEXT("iB");
+	uint64 Unit = 1024;
+	if (UnitStandard == EMemoryUnitStandard::SI)
+	{
+		Prefix = 8;
+		Suffix = TEXT("B");
+		Unit = 1000;
+	}
+
+	// We consistently use decimal magnitude for testing, so that both IEC and SI remain friendly as a stringified decimal number.
+	if (NumBytes < 1000)
 	{
 		Args.Add( TEXT("Number"), FText::AsNumber( NumBytes, Options, TargetCulture) );
 		Args.Add( TEXT("Unit"), FText::FromString( FString( TEXT("B") ) ) );
 		return FText::Format( NSLOCTEXT("Internationalization", "ComputerMemoryFormatting", "{Number} {Unit}"), Args );
 	}
-
-	static const TCHAR* Prefixes = TEXT("kMGTPEZY");
-	int32 Prefix = 0;
-
-	for (; NumBytes > 1024 * 1024; NumBytes >>= 10)
+	while (NumBytes >= (1000000))
 	{
+		NumBytes /= Unit;
 		++Prefix;
 	}
 
-	const double MemorySizeAsDouble = (double)NumBytes / 1024.0;
+	const double MemorySizeAsDouble = (double)NumBytes / (double)Unit;
 	Args.Add( TEXT("Number"), FText::AsNumber( MemorySizeAsDouble, Options, TargetCulture) );
-	Args.Add( TEXT("Unit"), FText::FromString( FString( 1, &Prefixes[Prefix] ) + TEXT("B") ) );
+	Args.Add( TEXT("Unit"), FText::FromString( FString( 1, &Prefixes[Prefix] ) + Suffix) );
 	return FText::Format( NSLOCTEXT("Internationalization", "ComputerMemoryFormatting", "{Number} {Unit}"), Args);
 }
 
+FText FText::AsMemory(uint64 NumBytes, EMemoryUnitStandard UnitStandard)
+{
+	return FText::AsMemory(NumBytes, nullptr, nullptr, UnitStandard);
+}
+
 FString FText::GetInvariantTimeZone()
 {
 	return TEXT("Etc/Unknown");
diff --git a/Engine/Source/Runtime/Core/Private/Mac/MacPlatformMemory.cpp b/Engine/Source/Runtime/Core/Private/Mac/MacPlatformMemory.cpp
--- a/Engine/Source/Runtime/Core/Private/Mac/MacPlatformMemory.cpp
+++ b/Engine/Source/Runtime/Core/Private/Mac/MacPlatformMemory.cpp
@@ -115,8 +115,10 @@ FMalloc* FMacPlatformMemory::BaseAllocator()
 	case EMemoryAllocatorToUse::Stomp:
 		return new FMallocStomp();
 #endif
+#if TBB_ALLOCATOR_ALLOWED
 	case EMemoryAllocatorToUse::TBB:
 		return new FMallocTBB();
+#endif
 	case EMemoryAllocatorToUse::Binned2:
 		return new FMallocBinned2();
 
diff --git a/Engine/Source/Runtime/Core/Private/Math/UnitConversion.cpp b/Engine/Source/Runtime/Core/Private/Math/UnitConversion.cpp
--- a/Engine/Source/Runtime/Core/Private/Math/UnitConversion.cpp
+++ b/Engine/Source/Runtime/Core/Private/Math/UnitConversion.cpp
@@ -100,7 +100,7 @@ const TCHAR* const DisplayStrings[] = {
 
 	TEXT("Hz"), TEXT("KHz"), TEXT("MHz"), TEXT("GHz"), TEXT("rpm"),
 
-	TEXT("B"), TEXT("KB"), TEXT("MB"), TEXT("GB"), TEXT("TB"),
+	TEXT("B"), TEXT("KiB"), TEXT("MiB"), TEXT("GiB"), TEXT("TiB"),
 
 	TEXT("lm"),
 
diff --git a/Engine/Source/Runtime/Core/Private/Misc/ExpressionParser.cpp b/Engine/Source/Runtime/Core/Private/Misc/ExpressionParser.cpp
--- a/Engine/Source/Runtime/Core/Private/Misc/ExpressionParser.cpp
+++ b/Engine/Source/Runtime/Core/Private/Misc/ExpressionParser.cpp
@@ -2,6 +2,7 @@
 
 #include "Misc/ExpressionParser.h"
 #include "Misc/AutomationTest.h"
+#include "Math/BasicMathExpressionEvaluator.h"
 
 #define LOCTEXT_NAMESPACE "ExpressionParser"
 
@@ -447,8 +448,8 @@ struct FExpressionCompiler
 
 	struct FWrappedOperator : FNoncopyable
 	{
-		FWrappedOperator(FCompiledToken InToken, int32 InPrecedence = 0)
-			: Token(MoveTemp(InToken)), Precedence(InPrecedence)
+		FWrappedOperator(FCompiledToken InToken, int32 InPrecedence = 0, int32 InShortCircuitIndex = INDEX_NONE)
+			: Token(MoveTemp(InToken)), Precedence(InPrecedence), ShortCircuitIndex(InShortCircuitIndex)
 		{}
 
 		FWrappedOperator(FWrappedOperator&& In) : Token(MoveTemp(In.Token)), Precedence(In.Precedence) {}
@@ -458,6 +459,7 @@ struct FExpressionCompiler
 
 		FCompiledToken Token;
 		int32 Precedence;
+		int32 ShortCircuitIndex;
 	};
 
 	TOptional<FExpressionError> CompileGroup(const FExpressionToken* GroupStart, const FGuid* StopAt)
@@ -467,6 +469,17 @@ struct FExpressionCompiler
 		TArray<FWrappedOperator> OperatorStack;
 		OperatorStack.Reserve(Tokens.Num() - CurrentTokenIndex);
 
+		auto PopOperator = [&]
+		{
+			int32 ShortCircuitIndex = OperatorStack.Last().ShortCircuitIndex;
+
+			Commands.Add(OperatorStack.Pop(false).Steal());
+			if (ShortCircuitIndex != INDEX_NONE)
+			{
+				Commands[ShortCircuitIndex].ShortCircuitIndex = Commands.Num() - 1;
+			}
+		};
+
 		bool bFoundEndOfGroup = StopAt == nullptr;
 
 		// Start off looking for a unary operator
@@ -516,7 +529,7 @@ struct FExpressionCompiler
 					// Pop off any pending unary operators
 					while (OperatorStack.Num() > 0 && OperatorStack.Last().Precedence <= 0)
 					{
-						Commands.Add(OperatorStack.Pop(false).Steal());
+						PopOperator();
 					}
 
 					// Make this a post-unary op
@@ -536,7 +549,7 @@ struct FExpressionCompiler
 					// Pop off any pending unary operators
 					while (OperatorStack.Num() > 0 && OperatorStack.Last().Precedence <= 0)
 					{
-						Commands.Add(OperatorStack.Pop(false).Steal());
+						PopOperator();
 					}
 
 					// Make this a post-unary op
@@ -555,11 +568,18 @@ struct FExpressionCompiler
 						// Pop off anything of higher (or equal, if LTR associative) precedence than this one onto the command stack
 						while (OperatorStack.Num() > 0 && CheckPrecedence(OperatorStack.Last().Precedence, OpParms->Precedence))
 						{
-							Commands.Add(OperatorStack.Pop(false).Steal());
+							PopOperator();
+						}
+
+						int32 ShortCircuitIndex = INDEX_NONE;
+						if (OpParms->bCanShortCircuit)
+						{
+							Commands.Add(FCompiledToken(FCompiledToken::ShortCircuit, FExpressionToken(Token.Context, Token.Node.Copy()), Commands.Num()));
+							ShortCircuitIndex = Commands.Num() - 1;
 						}
 
 						// Add the operator itself to the op stack
-						OperatorStack.Emplace(FCompiledToken(FCompiledToken::BinaryOperator, MoveTemp(Token)), OpParms->Precedence);
+						OperatorStack.Emplace(FCompiledToken(FCompiledToken::BinaryOperator, MoveTemp(Token)), OpParms->Precedence, ShortCircuitIndex);
 
 						// Check for a unary op again
 						State = EState::PreUnary;
@@ -587,7 +607,7 @@ struct FExpressionCompiler
 		// Pop everything off the operator stack, onto the command stack
 		while (OperatorStack.Num() > 0)
 		{
-			Commands.Add(OperatorStack.Pop(false).Token);
+			PopOperator();
 		}
 
 		return TOptional<FExpressionError>();
@@ -690,6 +710,13 @@ namespace ExpressionParser
 				OperandStack.Push(Index);
 				continue;
 
+			case FCompiledToken::ShortCircuit:
+				if (OperandStack.Num() >= 1 && Token.ShortCircuitIndex.IsSet() && InEnvironment.ShouldShortCircuit(Token, GetToken(OperandStack.Last())))
+				{
+					Index = Token.ShortCircuitIndex.GetValue();
+				}
+				continue;
+
 			case FCompiledToken::BinaryOperator:
 				if (OperandStack.Num() >= 2)
 				{
@@ -763,6 +790,12 @@ namespace Tests
 PRAGMA_DISABLE_OPTIMIZATION
 	struct FOperator {};
 
+	struct FAnd { static const TCHAR* const Moniker; };
+	struct FOr { static const TCHAR* const Moniker; };
+
+	const TCHAR* const FAnd::Moniker = TEXT("&&");
+	const TCHAR* const FOr::Moniker = TEXT("||");
+
 	struct FMoveableType
 	{
 		static int32* LeakCount;
@@ -923,6 +956,138 @@ PRAGMA_DISABLE_OPTIMIZATION
 DEFINE_EXPRESSION_NODE_TYPE(Tests::FMoveableType, 0xB7F3F127, 0xD5E74833, 0x9EAB754E, 0x6CF3AAC1)
 DEFINE_EXPRESSION_NODE_TYPE(Tests::FHugeType, 0x4A329D81, 0x102343A8, 0xAB95BF45, 0x6578EE54)
 DEFINE_EXPRESSION_NODE_TYPE(Tests::FOperator, 0xC777A5D7, 0x6895456C, 0x9854BFA0, 0xB71B5A8D)
+DEFINE_EXPRESSION_NODE_TYPE(Tests::FAnd, 0x0687f9c5, 0xd8914cb0, 0xae52cc4c, 0x6770f520)
+DEFINE_EXPRESSION_NODE_TYPE(Tests::FOr, 0x81e2b2a3, 0xbcf545d6, 0x95ae2eac, 0xcc5a9ba5)
+
+DEFINE_EXPRESSION_NODE_TYPE(bool, 0x9d5a4b67, 0x8c8f4fc1, 0xb61e561a, 0x5e4114b6)
+
+struct FShortCircuitTestContext
+{
+	FShortCircuitTestContext() : NumOperatorsCalled(0) {}
+	mutable int32 NumOperatorsCalled;
+};
+
+/** A basic math expression evaluator */
+class FShortCircuitParser
+{
+public:
+	/** Constructor that sets up the parser's lexer and compiler */
+	FShortCircuitParser()
+	{
+		using namespace ExpressionParser;
+		using namespace Tests;
+
+		// A || !(B && C)
+		TokenDefinitions.IgnoreWhitespace();
+		TokenDefinitions.DefineToken([](FExpressionTokenConsumer& Consumer) -> TOptional<FExpressionError> {
+			TOptional<FStringToken> Token = Consumer.GetStream().ParseToken(TEXT("true"));
+			if (Token.IsSet())
+			{
+				Consumer.Add(Token.GetValue(), true);
+			}
+
+			Token = Consumer.GetStream().ParseToken(TEXT("false"));
+			if (Token.IsSet())
+			{
+				Consumer.Add(Token.GetValue(), false);
+			}
+
+			return TOptional<FExpressionError>();
+		});
+
+		TokenDefinitions.DefineToken(&ConsumeSymbol<FSubExpressionStart>);
+		TokenDefinitions.DefineToken(&ConsumeSymbol<FSubExpressionEnd>);
+		TokenDefinitions.DefineToken(&ConsumeSymbol<FAnd>);
+		TokenDefinitions.DefineToken(&ConsumeSymbol<FOr>);
+
+		Grammar.DefineGrouping<FSubExpressionStart, FSubExpressionEnd>();
+
+		bool bCanShortCircuit = true;
+		Grammar.DefineBinaryOperator<FAnd>(1, EAssociativity::RightToLeft, bCanShortCircuit);
+		Grammar.DefineBinaryOperator<FOr>(1, EAssociativity::RightToLeft, bCanShortCircuit);
+
+		JumpTable.MapBinary<FAnd>([](bool A, bool B, const FShortCircuitTestContext* Context) {
+			++Context->NumOperatorsCalled;
+			return A && B;
+		});
+
+		JumpTable.MapBinary<FOr>([](bool A, bool B, const FShortCircuitTestContext* Context) {
+			++Context->NumOperatorsCalled;
+			return A || B;
+		});
+
+		JumpTable.MapShortCircuit<FOr>([](bool A) { return A; });
+		JumpTable.MapShortCircuit<FAnd>([](bool A) { return !A; });
+	}
+
+	TValueOrError<bool, FExpressionError> Evaluate(const TCHAR* InExpression, const FShortCircuitTestContext& TestContext) const
+	{
+		using namespace ExpressionParser;
+
+		TValueOrError<TArray<FExpressionToken>, FExpressionError> LexResult = ExpressionParser::Lex(InExpression, TokenDefinitions);
+		if (!LexResult.IsValid())
+		{
+			return MakeError(LexResult.StealError());
+		}
+
+		TValueOrError<TArray<FCompiledToken>, FExpressionError> CompilationResult = ExpressionParser::Compile(LexResult.StealValue(), Grammar);
+		if (!CompilationResult.IsValid())
+		{
+			return MakeError(CompilationResult.StealError());
+		}
+
+		TOperatorEvaluationEnvironment<FShortCircuitTestContext> Env(JumpTable, &TestContext);
+		TValueOrError<FExpressionNode, FExpressionError> Result = ExpressionParser::Evaluate(CompilationResult.GetValue(), Env);
+		if (!Result.IsValid())
+		{
+			return MakeError(Result.GetError());
+		}
+
+		if (const bool* Value = Result.GetValue().Cast<bool>())
+		{
+			return MakeValue(*Value);
+		}
+
+		return MakeError(NSLOCTEXT("Anon", "UnrecognizedResult", "Unrecognized result returned from expression"));
+	}
+
+private:
+
+	FTokenDefinitions TokenDefinitions;
+	FExpressionGrammar Grammar;
+	TOperatorJumpTable<FShortCircuitTestContext> JumpTable;
+
+} TestParser;
+
+IMPLEMENT_SIMPLE_AUTOMATION_TEST(FShortCircuitParserTest, "System.Core.Expression Parser.Short Circuit", EAutomationTestFlags::ApplicationContextMask | EAutomationTestFlags::SmokeFilter | EAutomationTestFlags::HighPriority)
+bool FShortCircuitParserTest::RunTest(const FString& Parameters)
+{
+	struct FExpectedResult
+	{
+		const TCHAR* const Expression;
+		bool Result;
+		int32 NumOperatorsCalled;
+	};
+	const FExpectedResult ExpectedResults[] = {
+		FExpectedResult{ TEXT("true || (true && true)"), true, 0 },
+		FExpectedResult{ TEXT("false && (true)"), false, 0 },
+	};
+
+	for (const FExpectedResult& Expected : ExpectedResults)
+	{
+		FShortCircuitTestContext Context;
+
+		TValueOrError<bool, FExpressionError> Result = TestParser.Evaluate(Expected.Expression, Context);
+		if (ensureAlways(Result.IsValid()))
+		{
+			ensureAlways(Expected.Result == Result.GetValue());
+			ensureAlways(Expected.NumOperatorsCalled == Context.NumOperatorsCalled);
+		}
+	}
+
+	return true;
+}
+
 
 PRAGMA_ENABLE_OPTIMIZATION
 
diff --git a/Engine/Source/Runtime/Core/Private/Windows/WindowsPlatformMemory.cpp b/Engine/Source/Runtime/Core/Private/Windows/WindowsPlatformMemory.cpp
--- a/Engine/Source/Runtime/Core/Private/Windows/WindowsPlatformMemory.cpp
+++ b/Engine/Source/Runtime/Core/Private/Windows/WindowsPlatformMemory.cpp
@@ -110,10 +110,12 @@ FMalloc* FWindowsPlatformMemory::BaseAllocator()
 	{
 		AllocatorToUse = EMemoryAllocatorToUse::Ansi;
 	}
+#if TBB_ALLOCATOR_ALLOWED
 	else if (FCString::Stristr(CommandLine, TEXT("-tbbmalloc")))
 	{
 		AllocatorToUse = EMemoryAllocatorToUse::TBB;
 	}
+#endif
 	else if (FCString::Stristr(CommandLine, TEXT("-binnedmalloc2")))
 	{
 		AllocatorToUse = EMemoryAllocatorToUse::Binned2;
@@ -138,8 +140,10 @@ FMalloc* FWindowsPlatformMemory::BaseAllocator()
 	case EMemoryAllocatorToUse::Stomp:
 		return new FMallocStomp();
 #endif
+#if TBB_ALLOCATOR_ALLOWED
 	case EMemoryAllocatorToUse::TBB:
 		return new FMallocTBB();
+#endif
 	case EMemoryAllocatorToUse::Binned2:
 		return new FMallocBinned2();
 		
diff --git a/Engine/Source/Runtime/Core/Public/HAL/MallocTBB.h b/Engine/Source/Runtime/Core/Public/HAL/MallocTBB.h
--- a/Engine/Source/Runtime/Core/Public/HAL/MallocTBB.h
+++ b/Engine/Source/Runtime/Core/Public/HAL/MallocTBB.h
@@ -6,6 +6,8 @@
 #include "HAL/PlatformMemory.h"
 #include "HAL/MemoryBase.h"
 
+#if PLATFORM_SUPPORTS_TBB && TBB_ALLOCATOR_ALLOWED
+
 /**
  * TBB 64-bit scalable memory allocator.
  */
@@ -52,3 +54,5 @@ private:
 
 	double MemTime;
 };
+
+#endif // PLATFORM_SUPPORTS_TBB && TBB_ALLOCATOR_ALLOWED
\ No newline at end of file
diff --git a/Engine/Source/Runtime/Core/Public/Internationalization/Text.h b/Engine/Source/Runtime/Core/Public/Internationalization/Text.h
--- a/Engine/Source/Runtime/Core/Public/Internationalization/Text.h
+++ b/Engine/Source/Runtime/Core/Public/Internationalization/Text.h
@@ -135,6 +135,14 @@ enum ERoundingMode
 	// Add new enum types at the end only! They are serialized by index.
 };
 
+enum EMemoryUnitStandard
+{
+	/* International Electrotechnical Commission (MiB) 1024-based */
+	IEC,
+	/* International System of Units 1000-based */
+	SI
+};
+
 struct CORE_API FNumberFormattingOptions
 {
 	FNumberFormattingOptions();
@@ -358,7 +366,12 @@ public:
 	/**
 	 * Generate an FText that represents the passed number as a memory size in the current culture
 	 */
-	static FText AsMemory(uint64 NumBytes, const FNumberFormattingOptions* const Options = NULL, const FCulturePtr& TargetCulture = NULL);
+	static FText AsMemory(uint64 NumBytes, const FNumberFormattingOptions* const Options = NULL, const FCulturePtr& TargetCulture = NULL, EMemoryUnitStandard UnitStandard = EMemoryUnitStandard::IEC);
+
+	/**
+	 * Generate an FText that represents the passed number as a memory size in the current culture
+	 */
+	static FText AsMemory(uint64 NumBytes, EMemoryUnitStandard UnitStandard);
 
 	/**
 	 * Attempts to find an existing FText using the representation found in the loc tables for the specified namespace and key
diff --git a/Engine/Source/Runtime/Core/Public/Misc/ExpressionParserTypes.h b/Engine/Source/Runtime/Core/Public/Misc/ExpressionParserTypes.h
--- a/Engine/Source/Runtime/Core/Public/Misc/ExpressionParserTypes.h
+++ b/Engine/Source/Runtime/Core/Public/Misc/ExpressionParserTypes.h
@@ -237,16 +237,17 @@ public:
 struct FCompiledToken : FExpressionToken
 {
 	// Todo: add callable types here?
-	enum EType { Operand, PreUnaryOperator, PostUnaryOperator, BinaryOperator, Benign };
+	enum EType { Operand, PreUnaryOperator, PostUnaryOperator, BinaryOperator, ShortCircuit, Benign };
 
-	FCompiledToken(EType InType, FExpressionToken InToken)
-		: FExpressionToken(MoveTemp(InToken)), Type(InType)
+	FCompiledToken(EType InType, FExpressionToken InToken, TOptional<int32> InShortCircuitIndex = TOptional<int32>())
+		: FExpressionToken(MoveTemp(InToken)), Type(InType), ShortCircuitIndex(InShortCircuitIndex)
 	{}
 
 	FCompiledToken(FCompiledToken&& In) : FExpressionToken(MoveTemp(In)), Type(In.Type) {}
 	FCompiledToken& operator=(FCompiledToken&& In) { FExpressionToken::operator=(MoveTemp(In)); Type = In.Type; return *this; }
 
 	EType Type;
+	TOptional<int32> ShortCircuitIndex;
 };
 
 /** Struct used to identify a function for a specific operator overload */
@@ -280,6 +281,8 @@ struct TOperatorJumpTable
 	FExpressionResult ExecPostUnary(const FExpressionToken& Operator, const FExpressionToken& L, const ContextType* Context) const;
 	/** Execute the specified token as a binary operator, if such an overload exists */
 	FExpressionResult ExecBinary(const FExpressionToken& Operator, const FExpressionToken& L, const FExpressionToken& R, const ContextType* Context) const;
+	/** Check whether we should short circuit the specified operator */
+	bool ShouldShortCircuit(const FExpressionToken& Operator, const FExpressionToken& L, const ContextType* Context) const;
 
 	/**
 	 * Map an expression node to a pre-unary operator with the specified implementation.
@@ -345,17 +348,22 @@ struct TOperatorJumpTable
 	template<typename OperatorType, typename FuncType>
 	void MapBinary(FuncType InFunc);
 
+	template<typename OperatorType, typename FuncType>
+	void MapShortCircuit(FuncType InFunc);
+
 public:
 
 	typedef TFunction<FExpressionResult(const FExpressionNode&, const ContextType* Context)> FUnaryFunction;
 	typedef TFunction<FExpressionResult(const FExpressionNode&, const FExpressionNode&, const ContextType* Context)> FBinaryFunction;
+	typedef TFunction<bool(const FExpressionNode&, const ContextType* Context)> FShortCircuit;
 
 private:
 
 	/** Maps of unary/binary operators */
-	TMap<FOperatorFunctionID, FUnaryFunction> PreUnaryOps;
-	TMap<FOperatorFunctionID, FUnaryFunction> PostUnaryOps;
+	TMap<FOperatorFunctionID, FUnaryFunction>  PreUnaryOps;
+	TMap<FOperatorFunctionID, FUnaryFunction>  PostUnaryOps;
 	TMap<FOperatorFunctionID, FBinaryFunction> BinaryOps;
+	TMap<FOperatorFunctionID, FShortCircuit>   BinaryShortCircuits;
 };
 
 typedef TOperatorJumpTable<> FOperatorJumpTable;
@@ -371,6 +379,8 @@ struct IOperatorEvaluationEnvironment
 	virtual FExpressionResult ExecPostUnary(const FExpressionToken& Operator, const FExpressionToken& L) const = 0;
 	/** Execute the specified token as a binary operator, if such an overload exists */
 	virtual FExpressionResult ExecBinary(const FExpressionToken& Operator, const FExpressionToken& L, const FExpressionToken& R) const = 0;
+	/** Check whether we should short circuit the specified operator */
+	virtual bool ShouldShortCircuit(const FExpressionToken& Operator, const FExpressionToken& L) const = 0;
 };
 template<typename ContextType = void>
 struct TOperatorEvaluationEnvironment : IOperatorEvaluationEnvironment
@@ -391,6 +401,10 @@ struct TOperatorEvaluationEnvironment : IOperatorEvaluationEnvironment
 	{
 		return Operators.ExecBinary(Operator, L, R, Context);
 	}
+	virtual bool ShouldShortCircuit(const FExpressionToken& Operator, const FExpressionToken& L) const override
+	{
+		return Operators.ShouldShortCircuit(Operator, L, Context);
+	}
 
 private:
 	const TOperatorJumpTable<ContextType>& Operators;
@@ -487,9 +501,13 @@ struct FOpParameters
 	/** The associativity of the operator */
 	EAssociativity	Associativity;
 
-	FOpParameters(int32 InPrecedence, EAssociativity InAssociativity)
+	/** Whether this operator can be short circuited or not */
+	bool bCanShortCircuit;
+
+	FOpParameters(int32 InPrecedence, EAssociativity InAssociativity, bool bInCanShortCircuit)
 		: Precedence(InPrecedence)
 		, Associativity(InAssociativity)
+		, bCanShortCircuit(bInCanShortCircuit)
 	{
 	}
 };
@@ -518,7 +536,7 @@ public:
 	 * @param InAssociativity	With operators of the same precedence, determines whether they execute left to right, or right to left
 	 */
 	template<typename TExpressionNode>
-	void DefineBinaryOperator(int32 InPrecedence, EAssociativity InAssociativity=EAssociativity::RightToLeft)
+	void DefineBinaryOperator(int32 InPrecedence, EAssociativity InAssociativity=EAssociativity::RightToLeft, bool bCanShortCircuit = false)
 	{
 #if DO_CHECK
 		for (TMap<FGuid, FOpParameters>::TConstIterator It(BinaryOperators); It; ++It)
@@ -533,7 +551,7 @@ public:
 		}
 #endif
 
-		BinaryOperators.Add(TGetExpressionNodeTypeId<TExpressionNode>::GetTypeId(), FOpParameters(InPrecedence, InAssociativity));
+		BinaryOperators.Add(TGetExpressionNodeTypeId<TExpressionNode>::GetTypeId(), FOpParameters(InPrecedence, InAssociativity, bCanShortCircuit));
 	}
 
 public:
diff --git a/Engine/Source/Runtime/Core/Public/Misc/ExpressionParserTypes.inl b/Engine/Source/Runtime/Core/Public/Misc/ExpressionParserTypes.inl
--- a/Engine/Source/Runtime/Core/Public/Misc/ExpressionParserTypes.inl
+++ b/Engine/Source/Runtime/Core/Public/Misc/ExpressionParserTypes.inl
@@ -143,6 +143,25 @@ namespace Impl
 		};
 	}
 
+	/** Wrapper function for supplied functions of the signature bool(A, const ContextType* Context) */
+	template<typename OperandType, typename ContextType, typename FuncType>
+	static typename TEnableIf<Impl::TCallableInfo<FuncType>::NumArgs == 1, typename TOperatorJumpTable<ContextType>::FShortCircuit>::Type WrapShortCircuitFunction(FuncType In)
+	{
+		// Ignore the context
+		return [=](const FExpressionNode& InOperand, const ContextType* Context) {
+			return In(*InOperand.Cast<OperandType>());
+		};
+	}
+
+	/** Wrapper function for supplied functions of the signature bool(A, const ContextType* Context) */
+	template<typename OperandType, typename ContextType, typename FuncType>
+	static typename TEnableIf<Impl::TCallableInfo<FuncType>::NumArgs == 2, typename TOperatorJumpTable<ContextType>::FShortCircuit>::Type WrapShortCircuitFunction(FuncType In)
+	{
+		// Ignore the context
+		return [=](const FExpressionNode& InOperand, const ContextType* Context) {
+			return In(*InOperand.Cast<OperandType>(), Context);
+		};
+	}
 }
 
 template<typename T>
@@ -181,6 +200,18 @@ FExpressionResult TOperatorJumpTable<ContextType>::ExecBinary(const FExpressionT
 	return MakeError(FText::Format(LOCTEXT("BinaryExecutionError", "Binary operator {0} cannot operate on {1} and {2}"), Args));
 }
 
+template<typename ContextType>
+bool TOperatorJumpTable<ContextType>::ShouldShortCircuit(const FExpressionToken& Operator, const FExpressionToken& L, const ContextType* Context) const
+{
+	FOperatorFunctionID ID = { Operator.Node.GetTypeId(), L.Node.GetTypeId(), FGuid() };
+	if (const auto* Func = BinaryShortCircuits.Find(ID))
+	{
+		return (*Func)(L.Node, Context);
+	}
+
+	return false;
+}
+
 template<typename ContextType>
 FExpressionResult TOperatorJumpTable<ContextType>::ExecPreUnary(const FExpressionToken& Operator, const FExpressionToken& R, const ContextType* Context) const
 {
@@ -257,6 +288,22 @@ void TOperatorJumpTable<ContextType>::MapBinary(FuncType InFunc)
 	BinaryOps.Add(ID, Impl::WrapBinaryFunction<LeftOperandType, RightOperandType, ContextType>(InFunc));
 }
 
+template<typename ContextType>
+template<typename OperatorType, typename FuncType>
+void TOperatorJumpTable<ContextType>::MapShortCircuit(FuncType InFunc)
+{
+	typedef typename TRemoveConst<typename TRemoveReference<typename Impl::TCallableInfo<FuncType>::Arg1>::Type>::Type OperandType;
+
+	FOperatorFunctionID ID = {
+		TGetExpressionNodeTypeId<OperatorType>::GetTypeId(),
+		TGetExpressionNodeTypeId<OperandType>::GetTypeId(),
+		FGuid()
+	};
+
+	BinaryShortCircuits.Add(ID, Impl::WrapShortCircuitFunction<OperandType, ContextType>(InFunc));
+}
+
+
 typedef TOperatorEvaluationEnvironment<> FOperatorEvaluationEnvironment;
 
 #undef LOCTEXT_NAMESPACE
diff --git a/Engine/Source/Runtime/Messaging/Private/MessagingModule.cpp b/Engine/Source/Runtime/Messaging/Private/MessagingModule.cpp
--- a/Engine/Source/Runtime/Messaging/Private/MessagingModule.cpp
+++ b/Engine/Source/Runtime/Messaging/Private/MessagingModule.cpp
@@ -13,7 +13,7 @@
 
 
 #ifndef PLATFORM_SUPPORTS_MESSAGEBUS
-	#define PLATFORM_SUPPORTS_MESSAGEBUS !UE_BUILD_SHIPPING
+	#define PLATFORM_SUPPORTS_MESSAGEBUS !(WITH_SERVER_CODE && UE_BUILD_SHIPPING)
 #endif
 
 
diff --git a/Engine/Source/Runtime/Online/BuildPatchServices/Private/BuildPatchDiffManifests.cpp b/Engine/Source/Runtime/Online/BuildPatchServices/Private/BuildPatchDiffManifests.cpp
--- a/Engine/Source/Runtime/Online/BuildPatchServices/Private/BuildPatchDiffManifests.cpp
+++ b/Engine/Source/Runtime/Online/BuildPatchServices/Private/BuildPatchDiffManifests.cpp
@@ -90,14 +90,14 @@ bool FBuildDiffManifests::DiffManifests(const FString& ManifestFilePathA, const
 
 	int64 NewChunksCount = 0;
 	int64 TotalChunkSize = 0;
-	TSet<FString> TaggedFiles;
+	TSet<FString> TaggedFileSetA;
+	TSet<FString> TaggedFileSetB;
 	TSet<FGuid> ChunkSetA;
 	TSet<FGuid> ChunkSetB;
-	ManifestA->GetTaggedFileList(TagsA, TaggedFiles);
-	ManifestA->GetChunksRequiredForFiles(TaggedFiles, ChunkSetA);
-	TaggedFiles.Reset();
-	ManifestB->GetTaggedFileList(TagsB, TaggedFiles);
-	ManifestB->GetChunksRequiredForFiles(TaggedFiles, ChunkSetB);
+	ManifestA->GetTaggedFileList(TagsA, TaggedFileSetA);
+	ManifestA->GetChunksRequiredForFiles(TaggedFileSetA, ChunkSetA);
+	ManifestB->GetTaggedFileList(TagsB, TaggedFileSetB);
+	ManifestB->GetChunksRequiredForFiles(TaggedFileSetB, ChunkSetB);
 	TArray<FString> NewChunkPaths;
 	for (FGuid& ChunkB : ChunkSetB)
 	{
@@ -114,6 +114,29 @@ bool FBuildDiffManifests::DiffManifests(const FString& ManifestFilePathA, const
 	UE_LOG(LogDiffManifests, Log, TEXT("New chunks:  %lld"), NewChunksCount);
 	UE_LOG(LogDiffManifests, Log, TEXT("Total bytes: %lld"), TotalChunkSize);
 
+	TSet<FString> NewFilePaths = TaggedFileSetB.Difference(TaggedFileSetA);
+	TSet<FString> RemovedFilePaths = TaggedFileSetA.Difference(TaggedFileSetB);
+	TSet<FString> ChangedFilePaths;
+	TSet<FString> UnchangedFilePaths;
+
+	const TSet<FString>& SetToIterate = TaggedFileSetB.Num() > TaggedFileSetA.Num() ? TaggedFileSetA : TaggedFileSetB;
+	for (const FString& TaggedFile : SetToIterate)
+	{
+		FSHAHash FileHashA;
+		FSHAHash FileHashB;
+		if (ManifestA->GetFileHash(TaggedFile, FileHashA) && ManifestB->GetFileHash(TaggedFile, FileHashB))
+		{
+			if (FileHashA == FileHashB)
+			{
+				UnchangedFilePaths.Add(TaggedFile);
+			}
+			else
+			{
+				ChangedFilePaths.Add(TaggedFile);
+			}
+		}
+	}
+
 	// Log download details.
 	FNumberFormattingOptions SizeFormattingOptions;
 	SizeFormattingOptions.MaximumFractionalDigits = 3;
@@ -158,20 +181,20 @@ bool FBuildDiffManifests::DiffManifests(const FString& ManifestFilePathA, const
 	}
 	UE_LOG(LogDiffManifests, Log, TEXT("TagSet: %s"), *TagLogList);
 	UE_LOG(LogDiffManifests, Log, TEXT("%s %s:"), *ManifestA->GetAppName(), *ManifestA->GetVersionString());
-	UE_LOG(LogDiffManifests, Log, TEXT("    Download Size:  %10s"), *FText::AsMemory(DownloadSizeA, &SizeFormattingOptions).ToString());
-	UE_LOG(LogDiffManifests, Log, TEXT("    Build Size:     %10s"), *FText::AsMemory(BuildSizeA, &SizeFormattingOptions).ToString())
+	UE_LOG(LogDiffManifests, Log, TEXT("    Download Size:  %20s bytes (%10s, %11s)"), *FText::AsNumber(DownloadSizeA).ToString(), *FText::AsMemory(DownloadSizeA, &SizeFormattingOptions, nullptr, EMemoryUnitStandard::SI).ToString(), *FText::AsMemory(DownloadSizeA, &SizeFormattingOptions, nullptr, EMemoryUnitStandard::IEC).ToString());
+	UE_LOG(LogDiffManifests, Log, TEXT("    Build Size:     %20s bytes (%10s, %11s)"), *FText::AsNumber(BuildSizeA).ToString(), *FText::AsMemory(BuildSizeA, &SizeFormattingOptions, nullptr, EMemoryUnitStandard::SI).ToString(), *FText::AsMemory(BuildSizeA, &SizeFormattingOptions, nullptr, EMemoryUnitStandard::IEC).ToString())
 	UE_LOG(LogDiffManifests, Log, TEXT("%s %s:"), *ManifestB->GetAppName(), *ManifestB->GetVersionString());
-	UE_LOG(LogDiffManifests, Log, TEXT("    Download Size:  %10s"), *FText::AsMemory(DownloadSizeB, &SizeFormattingOptions).ToString());
-	UE_LOG(LogDiffManifests, Log, TEXT("    Build Size:     %10s"), *FText::AsMemory(BuildSizeB, &SizeFormattingOptions).ToString());
+	UE_LOG(LogDiffManifests, Log, TEXT("    Download Size:  %20s bytes (%10s, %11s)"), *FText::AsNumber(DownloadSizeB).ToString(), *FText::AsMemory(DownloadSizeB, &SizeFormattingOptions, nullptr, EMemoryUnitStandard::SI).ToString(), *FText::AsMemory(DownloadSizeB, &SizeFormattingOptions, nullptr, EMemoryUnitStandard::IEC).ToString());
+	UE_LOG(LogDiffManifests, Log, TEXT("    Build Size:     %20s bytes (%10s, %11s)"), *FText::AsNumber(BuildSizeB).ToString(), *FText::AsMemory(BuildSizeB, &SizeFormattingOptions, nullptr, EMemoryUnitStandard::SI).ToString(), *FText::AsMemory(BuildSizeB, &SizeFormattingOptions, nullptr, EMemoryUnitStandard::IEC).ToString())
 	UE_LOG(LogDiffManifests, Log, TEXT("%s %s -> %s %s:"), *ManifestA->GetAppName(), *ManifestA->GetVersionString(), *ManifestB->GetAppName(), *ManifestB->GetVersionString());
-	UE_LOG(LogDiffManifests, Log, TEXT("    Delta Size:     %10s"), *FText::AsMemory(DeltaDownloadSize, &SizeFormattingOptions).ToString());
+	UE_LOG(LogDiffManifests, Log, TEXT("    Delta Size:     %20s bytes (%10s, %11s)"), *FText::AsNumber(DeltaDownloadSize).ToString(), *FText::AsMemory(DeltaDownloadSize, &SizeFormattingOptions, nullptr, EMemoryUnitStandard::SI).ToString(), *FText::AsMemory(DeltaDownloadSize, &SizeFormattingOptions, nullptr, EMemoryUnitStandard::IEC).ToString())
 	UE_LOG(LogDiffManifests, Log, TEXT(""));
 	for (const FString& Tag : TagArrayB)
 	{
 		UE_LOG(LogDiffManifests, Log, TEXT("%s Impact:"), *(Tag.IsEmpty() ? UntaggedLog : Tag));
-		UE_LOG(LogDiffManifests, Log, TEXT("    Individual Download Size: %10s"), *FText::AsMemory(TagDownloadImpactB[Tag], &SizeFormattingOptions).ToString());
-		UE_LOG(LogDiffManifests, Log, TEXT("    Individual Build Size:    %10s"), *FText::AsMemory(TagBuildImpactB[Tag], &SizeFormattingOptions).ToString());
-		UE_LOG(LogDiffManifests, Log, TEXT("    Individual Delta Size:    %10s"), *FText::AsMemory(TagDeltaImpact[Tag], &SizeFormattingOptions).ToString());
+		UE_LOG(LogDiffManifests, Log, TEXT("    Individual Download Size:  %20s bytes (%10s, %11s)"), *FText::AsNumber(TagDownloadImpactB[Tag]).ToString(), *FText::AsMemory(TagDownloadImpactB[Tag], &SizeFormattingOptions, nullptr, EMemoryUnitStandard::SI).ToString(), *FText::AsMemory(TagDownloadImpactB[Tag], &SizeFormattingOptions, nullptr, EMemoryUnitStandard::IEC).ToString());
+		UE_LOG(LogDiffManifests, Log, TEXT("    Individual Build Size:     %20s bytes (%10s, %11s)"), *FText::AsNumber(TagBuildImpactB[Tag]).ToString(), *FText::AsMemory(TagBuildImpactB[Tag], &SizeFormattingOptions, nullptr, EMemoryUnitStandard::SI).ToString(), *FText::AsMemory(TagBuildImpactB[Tag], &SizeFormattingOptions, nullptr, EMemoryUnitStandard::IEC).ToString())
+		UE_LOG(LogDiffManifests, Log, TEXT("    Individual Delta Size:     %20s bytes (%10s, %11s)"), *FText::AsNumber(TagDeltaImpact[Tag]).ToString(), *FText::AsMemory(TagDeltaImpact[Tag], &SizeFormattingOptions, nullptr, EMemoryUnitStandard::SI).ToString(), *FText::AsMemory(TagDeltaImpact[Tag], &SizeFormattingOptions, nullptr, EMemoryUnitStandard::IEC).ToString())
 	}
 
 	// Save the output.
@@ -225,6 +248,30 @@ bool FBuildDiffManifests::DiffManifests(const FString& ManifestFilePathA, const
 			Writer->WriteObjectEnd();
 			Writer->WriteObjectStart(TEXT("Differential"));
 			{
+				Writer->WriteArrayStart(TEXT("NewFilePaths"));
+				for (const FString& NewFilePath : NewFilePaths)
+				{
+					Writer->WriteValue(NewFilePath);
+				}
+				Writer->WriteArrayEnd();
+				Writer->WriteArrayStart(TEXT("RemovedFilePaths"));
+				for (const FString& RemovedFilePath : RemovedFilePaths)
+				{
+					Writer->WriteValue(RemovedFilePath);
+				}
+				Writer->WriteArrayEnd();
+				Writer->WriteArrayStart(TEXT("ChangedFilePaths"));
+				for (const FString& ChangedFilePath : ChangedFilePaths)
+				{
+					Writer->WriteValue(ChangedFilePath);
+				}
+				Writer->WriteArrayEnd();
+				Writer->WriteArrayStart(TEXT("UnchangedFilePaths"));
+				for (const FString& UnchangedFilePath : UnchangedFilePaths)
+				{
+					Writer->WriteValue(UnchangedFilePath);
+				}
+				Writer->WriteArrayEnd();
 				Writer->WriteArrayStart(TEXT("NewChunkPaths"));
 				for (const FString& NewChunkPath : NewChunkPaths)
 				{
