BuildGraph: Prevent errors being shown when running with -listonly.

#rnx

diff --git a/Engine/Source/Programs/AutomationTool/BuildGraph/BuildGraph.cs b/Engine/Source/Programs/AutomationTool/BuildGraph/BuildGraph.cs
--- a/Engine/Source/Programs/AutomationTool/BuildGraph/BuildGraph.cs
+++ b/Engine/Source/Programs/AutomationTool/BuildGraph/BuildGraph.cs
@@ -387,7 +387,7 @@ namespace AutomationTool
 			}
 
 			// Print out all the diagnostic messages which still apply, unless we're running a step as part of a build system or just listing the contents of the file. 
-			if(SingleNode == null || !bListOnly || bShowDiagnostics)
+			if(SingleNode == null && (!bListOnly || bShowDiagnostics))
 			{
 				IEnumerable<GraphDiagnostic> Diagnostics = Graph.Diagnostics.Where(x => x.EnclosingTrigger == Trigger);
 				foreach(GraphDiagnostic Diagnostic in Diagnostics)
