UAT: fix breakage when cooking Blueprint projects

Yesterday I broke this in CL 4301305. Reworked logic around the -server option was causing

"ERROR: Server target not found. Server target is required with -server and -cook or -cookonthefly"

This CL fixes projects that don't have a Server.Target.cs but want to -cook.

#rnx

diff --git a/Engine/Source/Programs/AutomationTool/AutomationUtils/ProjectParams.cs b/Engine/Source/Programs/AutomationTool/AutomationUtils/ProjectParams.cs
--- a/Engine/Source/Programs/AutomationTool/AutomationUtils/ProjectParams.cs
+++ b/Engine/Source/Programs/AutomationTool/AutomationUtils/ProjectParams.cs
@@ -2147,7 +2147,10 @@ namespace AutomationTool
 			// Compile a server if it was asked for (-server) or we're cooking and require a server
 			if (ServerCookedTargetsList == null)
 			{
-				if (DedicatedServer || Cook || CookOnTheFly)
+				/* Simplified from previous version which makes less sense.
+				   TODO: tease out the actual dependencies between -cook and -server options, fix properly
+				if (DedicatedServer && (Cook || CookOnTheFly || DedicatedServer)) */
+				if (DedicatedServer)
 				{
 					if (String.IsNullOrEmpty(ServerTarget))
 					{
