Improve error message when unable to parse the current CL from a stream.

#rnx

diff --git a/Engine/Source/Programs/AutomationTool/AutomationUtils/P4Environment.cs b/Engine/Source/Programs/AutomationTool/AutomationUtils/P4Environment.cs
--- a/Engine/Source/Programs/AutomationTool/AutomationUtils/P4Environment.cs
+++ b/Engine/Source/Programs/AutomationTool/AutomationUtils/P4Environment.cs
@@ -421,14 +421,22 @@ namespace AutomationTool
 			CommandUtils.LogVerbose("uebp_CL not set, detecting 'have' CL...");
 
 			// Retrieve the current changelist 
-			var P4Result = Connection.P4("changes -m 1 " + CommandUtils.CombinePaths(PathSeparator.Depot, ClientRootPath, "/...#have"), AllowSpew: false);
-			var CLTokens = P4Result.Output.Split(new char[] { ' ' }, StringSplitOptions.RemoveEmptyEntries);
-			var CLString = CLTokens[1];
-			var CL = Int32.Parse(CLString);
-			if (CLString != CL.ToString())
+			IProcessResult P4Result = Connection.P4("changes -m 1 " + CommandUtils.CombinePaths(PathSeparator.Depot, ClientRootPath, "/...#have"), AllowSpew: false);
+
+			string[] CLTokens = P4Result.Output.Split(new char[] { ' ' }, StringSplitOptions.RemoveEmptyEntries);
+			if(CLTokens.Length == 0)
 			{
-				throw new AutomationException("Failed to retrieve current changelist.");
+				throw new AutomationException("Unable to find current changelist (no output from 'p4 changes' command)");
 			}
+			
+			string CLString = CLTokens[1];
+
+			int CL;
+			if(!Int32.TryParse(CLString, out CL) || CLString != CL.ToString())
+			{
+				throw new AutomationException("Unable to parse current changelist from Perforce output:\n{0}", P4Result.Output);
+			}
+
 			return CLString;
 		}
 
