Always pass an explicit file list into BuildPatchTool when chunking a build. Less prone to errors than filtering files out from a build.

diff --git a/Engine/Source/Programs/AutomationTool/AutomationUtils/MCPPublic.cs b/Engine/Source/Programs/AutomationTool/AutomationUtils/MCPPublic.cs
--- a/Engine/Source/Programs/AutomationTool/AutomationUtils/MCPPublic.cs
+++ b/Engine/Source/Programs/AutomationTool/AutomationUtils/MCPPublic.cs
@@ -547,7 +547,11 @@ namespace EpicGames.MCP.Automation
 			/// </summary>
 			public string ManifestFilename;
 			/// <summary>
-			/// A path to a text file containing BuildRoot relative files, separated by \r\n line endings, to not be included in the build.
+			/// A path to a text file containing BuildRoot relative files to be included in the build.
+			/// </summary>
+			public string FileInputList;
+			/// <summary>
+			/// A path to a text file containing BuildRoot relative files to be excluded from the build.
 			/// </summary>
 			public string FileIgnoreList;
 			/// <summary>
