Fix occasional incremental build error in AutomationScripts project. Explicit dependency on Localization output file was invalidating project dependency.

#rnx

diff --git a/Engine/Source/Programs/AutomationTool/Scripts/AutomationScripts.Automation.csproj b/Engine/Source/Programs/AutomationTool/Scripts/AutomationScripts.Automation.csproj
--- a/Engine/Source/Programs/AutomationTool/Scripts/AutomationScripts.Automation.csproj
+++ b/Engine/Source/Programs/AutomationTool/Scripts/AutomationScripts.Automation.csproj
@@ -53,10 +53,6 @@
     <Reference Include="Microsoft.CSharp" />
     <Reference Include="System.Data" />
     <Reference Include="System.Xml" />
-    <Reference Include="Localization.Automation">
-      <HintPath>..\..\..\..\Binaries\DotNET\AutomationScripts\Localization.Automation.dll</HintPath>
-      <Private>False</Private>
-    </Reference>
   </ItemGroup>
   <ItemGroup>
     <Compile Include="..\..\DotNETCommon\MetaData.cs">
@@ -136,6 +132,7 @@
     <ProjectReference Include="..\Localization\Localization.Automation.csproj">
       <Project>{c2f959b1-1727-4598-9c32-3d9acce72750}</Project>
       <Name>Localization.Automation</Name>
+      <Private>False</Private>
     </ProjectReference>
   </ItemGroup>
   <Import Project="$(MSBuildToolsPath)\Microsoft.CSharp.targets" />
