UAT: Remove code for deleting build products and retrying XGE compiles.

#rnx

diff --git a/Engine/Source/Programs/AutomationTool/AutomationUtils/UE4Build.cs b/Engine/Source/Programs/AutomationTool/AutomationUtils/UE4Build.cs
--- a/Engine/Source/Programs/AutomationTool/AutomationUtils/UE4Build.cs
+++ b/Engine/Source/Programs/AutomationTool/AutomationUtils/UE4Build.cs
@@ -24,18 +24,6 @@ namespace AutomationTool
 
 		private String XGEConsoleExePath = "";
 
-		public void PrepareBuildProduct(string File, bool DeleteProduct = true)
-		{
-			// Do not delete these files; they may have been generated by UBT
-			// if (!BuildProductFiles.Contains(File))
-			// {
-			//	if (DeleteProduct && DeleteBuildProducts)
-			//	{
-			//		CommandUtils.DeleteFile(File);
-			//	}
-			// }
-		}
-
 		public bool HasBuildProduct(string InFile)
 		{
 			string File = CommandUtils.CombinePaths(InFile);
@@ -91,10 +79,6 @@ namespace AutomationTool
 				{
 					AddBuildProduct(Item);
 				}
-				else
-				{
-					PrepareBuildProduct(Item);
-				}
 			}
 			return Manifest;
 		}
@@ -243,14 +227,6 @@ namespace AutomationTool
 			}
 		}
 
-		void XGEDeleteBuildProducts(UnrealBuildTool.BuildManifest Manifest)
-		{
-			foreach (string Item in Manifest.BuildProducts)
-			{
-				CommandUtils.DeleteFile(Item);
-			}
-		}
-
 		void CleanWithUBT(string TargetName, UnrealBuildTool.UnrealTargetPlatform Platform, string Config, FileReference UprojectPath, bool ForceMonolithic = false, bool ForceNonUnity = false, bool ForceDebugInfo = false, string InAddArgs = "", bool ForceUnity = false, Dictionary<string, string> EnvVars = null)
 		{
 			string AddArgs = "";
@@ -381,15 +357,6 @@ namespace AutomationTool
             };
 		}
 
-		void PrepareBuildProductsForCSharpProj(string CsProj)
-		{
-			string BaseOutput = CommandUtils.CmdEnv.LocalRoot + @"/Engine/Binaries/DotNET/" + Path.GetFileNameWithoutExtension(CsProj);
-			foreach (var Ext in DotNetProductExtenstions())
-			{
-				PrepareBuildProduct(BaseOutput + Ext, false); // we don't delete these because we are not sure they are actually a build product
-			}
-		}
-
 		void AddBuildProductsForCSharpProj(string CsProj)
 		{
 			string BaseOutput = CommandUtils.CmdEnv.LocalRoot + @"/Engine/Binaries/DotNET/" + Path.GetFileNameWithoutExtension(CsProj);
@@ -590,63 +557,81 @@ namespace AutomationTool
 		[DebuggerDisplay("{TargetName} {Platform} {Config}")]
 		public class BuildTarget
 		{
+			/// <summary>
 			/// Name of the target
+			/// </summary>
 			public string TargetName;
 
+			/// <summary>
 			/// For code-based projects with a .uproject, the TargetName isn't enough for UBT to find the target, this will point UBT to the target
+			/// </summary>
 			public FileReference UprojectPath;
 
+			/// <summary>
 			/// Platform to build
+			/// </summary>
 			public UnrealBuildTool.UnrealTargetPlatform Platform;
 
+			/// <summary>
 			/// Configuration to build
+			/// </summary>
 			public UnrealBuildTool.UnrealTargetConfiguration Config;
 
+			/// <summary>
 			/// Extra UBT args
+			/// </summary>
 			public string UBTArgs;
 
-			/// Whether to clean this target
+			/// <summary>
+			/// Whether to clean this target. If not specified, the target will be cleaned if -Clean is on the command line.
+			/// </summary>
 			public bool? Clean;
 		}
 
 
 		public class BuildAgenda
 		{
+			/// <summary>
 			/// Full .NET solution files that we will compile and include in the build.  Currently we assume the output
 			/// binary file names match the solution file base name, but with various different binary file extensions.
+			/// </summary>
 			public List<string> DotNetSolutions = new List<string>();
 
+			/// <summary>
 			/// .NET .csproj files that will be compiled and included in the build.  Currently we assume the output
 			/// binary file names match the solution file base name, but with various different binary file extensions.
+			/// </summary>
 			public List<string> DotNetProjects = new List<string>();
 
+			/// <summary>
 			/// These are .NET project binary base file names that we want to include and label with the build, but
 			/// we won't be compiling these projects directly ourselves (however, they may be incidentally build or
 			/// published by a different project that we are building.)  We'll look for various .NET binary files with
 			/// this base file name but with different extensions.
+			/// </summary>
 			public List<string> ExtraDotNetFiles = new List<string>();
 
+			/// <summary>
 			/// These are .NET project binary files that are specific to IOS and found in the IOS subdirectory.  We define
 			/// these buildproducts as existing in the DotNET\IOS directory and assume that the output binary file names match
 			/// the solution file base name, but with various binary file extensions
+			/// </summary>
 			public List<string> IOSDotNetProjects = new List<string>();
 
+			/// <summary>
 			/// These are .NET project binary files that are specific to HTML5.  We define these build products as existing in the 
 			/// DotNET directory and assume that the output binary file names match
 			/// the solution file base name, but with various binary file extensions
+			/// </summary>
 			public List<string> HTML5DotNetProjects = new List<string>();
 
 			public string SwarmProject = "";
 
-
+			/// <summary>
 			/// List of targets to build.  These can be various Unreal projects, programs or libraries in various configurations
+			/// </summary>
 			public List<BuildTarget> Targets = new List<BuildTarget>();
 
-			/// Used for special jobs to test things. At the moment, debugging some megaXGE thing
-			public bool SpecialTestFlag = false;
-
-			public bool DoRetries = false;
-
 			/// <summary>
 			/// Adds a target with the specified configuration.
 			/// </summary>
@@ -743,21 +728,14 @@ namespace AutomationTool
 			return XGEConsoleExePath;
 		}
 
-		public bool ProcessXGEItems(List<XGEItem> Actions, string XGETool, string Args, string TaskFilePath, bool DoRetries, bool SpecialTestFlag, bool ShowProgress)
+		public bool ProcessXGEItems(List<XGEItem> Actions, string XGETool, string Args, string TaskFilePath, bool ShowProgress)
 		{
 			TelemetryStopwatch CombineXGEStopwatch = new TelemetryStopwatch("CombineXGEItemFiles.{0}", Path.GetFileNameWithoutExtension(XGETool));
 
 			XmlDocument XGETaskDocument;	
 			if (!CombineXGEItemFiles(Actions, TaskFilePath, out XGETaskDocument))
 			{
-				if (DeleteBuildProducts)
-				{
-					throw new AutomationException("Unable to find xge xmls: " + CommandUtils.CombinePaths(CommandUtils.CmdEnv.LogFolder, "*.xge.xml"));
-				}
-				else
-				{
-					CommandUtils.LogVerbose("Incremental build, apparently everything was up to date, no XGE jobs produced.");
-				}
+				CommandUtils.LogVerbose("Incremental build, apparently everything was up to date, no XGE jobs produced.");
 			}
 			else
 			{
@@ -821,53 +799,33 @@ namespace AutomationTool
 				{
 					using (TelemetryStopwatch ProcessXGEStopwatch = new TelemetryStopwatch("ProcessXGE.{0}", Path.GetFileNameWithoutExtension(XGETool)))
 					{
-						int Retries = DoRetries ? 2 : 1;
 						int ConnectionRetries = 4;
-						for (int i = 0; i < Retries; i++)
+						while (true)
 						{
-							try
+							CommandUtils.LogVerbose("Running {0} *******", XGETool);
+							CommandUtils.PushDir(CommandUtils.CombinePaths(CommandUtils.CmdEnv.LocalRoot, @"\Engine\Source"));
+							int SuccesCode;
+							string LogFile = CommandUtils.GetRunAndLogOnlyName(CommandUtils.CmdEnv, XGETool);
+							string Output = CommandUtils.RunAndLog(XGETool, Args, out SuccesCode, LogFile);
+							bool bOutputContainsProject = Output.Contains("------Project:");
+							CommandUtils.PopDir();
+							if (ConnectionRetries > 0 && (SuccesCode == 4 || SuccesCode == 2) && !bOutputContainsProject && XGETool != null)
 							{
-								while (true)
+								CommandUtils.LogWarning(String.Format("{0} failure on the local connection timeout", XGETool));
+								if (ConnectionRetries < 2)
 								{
-									CommandUtils.LogVerbose("Running {0} *******", XGETool);
-									CommandUtils.PushDir(CommandUtils.CombinePaths(CommandUtils.CmdEnv.LocalRoot, @"\Engine\Source"));
-									int SuccesCode;
-									string LogFile = CommandUtils.GetRunAndLogOnlyName(CommandUtils.CmdEnv, XGETool);
-									string Output = CommandUtils.RunAndLog(XGETool, Args, out SuccesCode, LogFile);
-									bool bOutputContainsProject = Output.Contains("------Project:");
-									CommandUtils.PopDir();
-									if (ConnectionRetries > 0 && (SuccesCode == 4 || SuccesCode == 2) && !bOutputContainsProject && XGETool != null)
-									{
-										CommandUtils.LogWarning(String.Format("{0} failure on the local connection timeout", XGETool));
-										if (ConnectionRetries < 2)
-										{
-											System.Threading.Thread.Sleep(60000);
-										}
-										ConnectionRetries--;
-										continue;
-									}
-									else if (SuccesCode != 0)
-									{
-										throw new AutomationException("Command failed (Result:{3}): {0} {1}. See logfile for details: '{2}' ",
-																		XGETool, Args, Path.GetFileName(LogFile), SuccesCode);
-									}
-									CommandUtils.LogVerbose("{0} {1} Done *******", XGETool, Args);
-									break;
+									System.Threading.Thread.Sleep(60000);
 								}
-								break;
+								ConnectionRetries--;
+								continue;
 							}
-							catch (Exception Ex)
+							else if (SuccesCode != 0)
 							{
-								if (i + 1 >= Retries)
-								{
-									throw;
-								}
-								CommandUtils.LogWarning("{0} failed on try {1}, deleting products to force a rebuild: {2}", XGETool, i + 1, Ex.ToString());
-								foreach (var Item in Actions)
-								{
-									XGEDeleteBuildProducts(Item.Manifest);
-								}
+								throw new AutomationException("Command failed (Result:{3}): {0} {1}. See logfile for details: '{2}' ",
+																XGETool, Args, Path.GetFileName(LogFile), SuccesCode);
 							}
+							CommandUtils.LogVerbose("{0} {1} Done *******", XGETool, Args);
+							break;
 						}
 					}
 				}
@@ -1234,7 +1192,7 @@ namespace AutomationTool
 			{
 				throw new AutomationException("You are attempting to compile on a machine that does not have a supported compiler!");
 			}
-			DeleteBuildProducts = InDeleteBuildProducts.HasValue ? InDeleteBuildProducts.Value : ParseParam("Clean");
+			bool DeleteBuildProducts = InDeleteBuildProducts.HasValue ? InDeleteBuildProducts.Value : ParseParam("Clean");
 			if (InUpdateVersionFiles)
 			{
 				UpdateVersionFiles(ActuallyUpdateVersionFiles: true, ChangelistNumberOverride: InChangelistNumberOverride);
@@ -1252,16 +1210,9 @@ namespace AutomationTool
 				}
 			}
 
-			foreach (var File in Agenda.ExtraDotNetFiles)
-			{
-				PrepareBuildProductsForCSharpProj(Path.Combine(CommandUtils.CmdEnv.LocalRoot, File));
-			}
-
 			if (Agenda.SwarmProject != "")
 			{
 				string SwarmSolution = Path.Combine(CommandUtils.CmdEnv.LocalRoot, Agenda.SwarmProject);
-				PrepareBuildProductsForCSharpProj(SwarmSolution);
-
 				CommandUtils.BuildSolution(CommandUtils.CmdEnv, SwarmSolution, "Development", "Mixed Platforms");
 				AddSwarmBuildProducts();
 			}
@@ -1269,44 +1220,28 @@ namespace AutomationTool
 			foreach (var DotNetSolution in Agenda.DotNetSolutions)
 			{
 				string Solution = Path.Combine(CommandUtils.CmdEnv.LocalRoot, DotNetSolution);
-
-				PrepareBuildProductsForCSharpProj(Solution);
-
 				CommandUtils.BuildSolution(CommandUtils.CmdEnv, Solution, "Development", "Any CPU");
-
 				AddBuildProductsForCSharpProj(Solution);
 			}
 
 			foreach (var DotNetProject in Agenda.DotNetProjects)
 			{
 				string CsProj = Path.Combine(CommandUtils.CmdEnv.LocalRoot, DotNetProject);
-
-				PrepareBuildProductsForCSharpProj(CsProj);
-
 				CommandUtils.BuildCSharpProject(CommandUtils.CmdEnv, CsProj);
-
 				AddBuildProductsForCSharpProj(CsProj);
 			}
 
 			foreach (var IOSDotNetProject in Agenda.IOSDotNetProjects)
 			{
 				string IOSCsProj = Path.Combine(CommandUtils.CmdEnv.LocalRoot, IOSDotNetProject);
-
-				PrepareBuildProductsForCSharpProj(IOSCsProj);
-
 				CommandUtils.BuildCSharpProject(CommandUtils.CmdEnv, IOSCsProj);
-
 				AddIOSBuildProductsForCSharpProj(IOSCsProj);
 			}
 
 			foreach (var HTML5DotNetProject in Agenda.HTML5DotNetProjects)
 			{
 				string CsProj = Path.Combine(CommandUtils.CmdEnv.LocalRoot, HTML5DotNetProject);
-
-				PrepareBuildProductsForCSharpProj(CsProj);
-
 				CommandUtils.BuildCSharpProject(CommandUtils.CmdEnv, CsProj);
-
 				AddBuildProductsForCSharpProj(CsProj);
 			}
 
@@ -1413,7 +1348,7 @@ namespace AutomationTool
 					if (!String.IsNullOrEmpty(XGETool))
 					{
 						CommandUtils.LogSetProgress(InShowProgress, "Building...");
-						if (!ProcessXGEItems(XGEItems, XGETool, Args, TaskFilePath, Agenda.DoRetries, Agenda.SpecialTestFlag, InShowProgress))
+						if (!ProcessXGEItems(XGEItems, XGETool, Args, TaskFilePath, InShowProgress))
 						{
 							throw new AutomationException("BUILD FAILED: {0} failed, retries not enabled:", XGETool);
 						}
@@ -1666,7 +1601,5 @@ namespace AutomationTool
 
 		// List of everything we built so far
 		public readonly List<string> BuildProductFiles = new List<string>();
-
-		private bool DeleteBuildProducts = true;
 	}
 }
diff --git a/Engine/Source/Programs/AutomationTool/Scripts/MegaXGE.Automation.cs b/Engine/Source/Programs/AutomationTool/Scripts/MegaXGE.Automation.cs
--- a/Engine/Source/Programs/AutomationTool/Scripts/MegaXGE.Automation.cs
+++ b/Engine/Source/Programs/AutomationTool/Scripts/MegaXGE.Automation.cs
@@ -159,7 +159,6 @@ class MegaXGE : BuildCommand
 		}
 		LogInformation("*************************");
 
-		Agenda.DoRetries = ParseParam("Retry");
 		UE4Build.Build(Agenda, InUpdateVersionFiles: IsBuildMachine, InUseParallelExecutor: ParseParam("useparallelexecutor"), InShowProgress: ShowProgress);
 
 		// 		if (WorkingCL > 0) // only move UAT files if we intend to check in some build products
