Merging 4279381 to DevBuild
Fix for UE4Build.cs for Swarm

diff --git a/Engine/Source/Programs/AutomationTool/AutomationUtils/UE4Build.cs b/Engine/Source/Programs/AutomationTool/AutomationUtils/UE4Build.cs
--- a/Engine/Source/Programs/AutomationTool/AutomationUtils/UE4Build.cs
+++ b/Engine/Source/Programs/AutomationTool/AutomationUtils/UE4Build.cs
@@ -625,7 +625,8 @@ namespace AutomationTool
 			/// </summary>
 			public List<string> HTML5DotNetProjects = new List<string>();
 
-			public string SwarmProject = "";
+			public string SwarmAgentProject = "";
+			public string SwarmCoordinatorProject = "";
 
 			/// <summary>
 			/// List of targets to build.  These can be various Unreal projects, programs or libraries in various configurations
@@ -1210,20 +1211,27 @@ namespace AutomationTool
 				}
 			}
 
-			if (Agenda.SwarmProject != "")
+			if (Agenda.SwarmAgentProject != "")
 			{
-				string SwarmSolution = Path.Combine(CommandUtils.CmdEnv.LocalRoot, Agenda.SwarmProject);
-				CommandUtils.BuildSolution(CommandUtils.CmdEnv, SwarmSolution, "Development", "Mixed Platforms");
+				string SwarmAgentSolution = Path.Combine(CommandUtils.CmdEnv.LocalRoot, Agenda.SwarmAgentProject);
+				CommandUtils.BuildSolution(CommandUtils.CmdEnv, SwarmAgentSolution, "Development", "Mixed Platforms");
 				AddSwarmBuildProducts();
 			}
 
+			if (Agenda.SwarmCoordinatorProject != "")
+			{
+				string SwarmCoordinatorSolution = Path.Combine(CommandUtils.CmdEnv.LocalRoot, Agenda.SwarmCoordinatorProject);
+				CommandUtils.BuildSolution(CommandUtils.CmdEnv, SwarmCoordinatorSolution, "Development", "Mixed Platforms");
+				AddSwarmBuildProducts();
+			}
+			
 			foreach (var DotNetSolution in Agenda.DotNetSolutions)
 			{
 				string Solution = Path.Combine(CommandUtils.CmdEnv.LocalRoot, DotNetSolution);
 				CommandUtils.BuildSolution(CommandUtils.CmdEnv, Solution, "Development", "Any CPU");
 				AddBuildProductsForCSharpProj(Solution);
 			}
-
+			
 			foreach (var DotNetProject in Agenda.DotNetProjects)
 			{
 				string CsProj = Path.Combine(CommandUtils.CmdEnv.LocalRoot, DotNetProject);
diff --git a/Engine/Source/Programs/AutomationTool/Scripts/BuildCommonTools.Automation.cs b/Engine/Source/Programs/AutomationTool/Scripts/BuildCommonTools.Automation.cs
--- a/Engine/Source/Programs/AutomationTool/Scripts/BuildCommonTools.Automation.cs
+++ b/Engine/Source/Programs/AutomationTool/Scripts/BuildCommonTools.Automation.cs
@@ -85,7 +85,8 @@ public class BuildCommonTools : BuildCommand
 		UE4Build.BuildAgenda Agenda = new UE4Build.BuildAgenda();
 
 		// C# binaries
-		Agenda.SwarmProject = @"Engine\Source\Programs\UnrealSwarm\UnrealSwarm.sln";
+		Agenda.SwarmAgentProject = @"Engine\Source\Programs\UnrealSwarm\SwarmAgent.sln";
+		Agenda.SwarmCoordinatorProject = @"Engine\Source\Programs\UnrealSwarm\SwarmCoordinator.sln";
 		Agenda.DotNetProjects.Add(@"Engine/Source/Editor/SwarmInterface/DotNET/SwarmInterface.csproj");
 		Agenda.DotNetProjects.Add(@"Engine/Source/Programs/DotNETCommon/DotNETUtilities/DotNETUtilities.csproj");
 		Agenda.DotNetProjects.Add(@"Engine/Source/Programs/RPCUtility/RPCUtility.csproj");
