UAT: Use enum for platform and configuration settings in UE4Build.cs.

#rnx

diff --git a/Engine/Source/Programs/AutomationTool/AutomationUtils/UBTUtils.cs b/Engine/Source/Programs/AutomationTool/AutomationUtils/UBTUtils.cs
--- a/Engine/Source/Programs/AutomationTool/AutomationUtils/UBTUtils.cs
+++ b/Engine/Source/Programs/AutomationTool/AutomationUtils/UBTUtils.cs
@@ -23,7 +23,7 @@ namespace AutomationTool
 		/// <param name="Env">Environment to use.</param>
 		/// <param name="CommandLine">Commandline to pass on to UBT.</param>
 		/// <param name="LogName">Optional logfile name.</param>
-        public static void RunUBT(CommandEnvironment Env, string UBTExecutable, string CommandLine, Dictionary<string, string> EnvVars = null)
+        public static void RunUBT(CommandEnvironment Env, string UBTExecutable, string CommandLine)
 		{
 			if (!FileExists(UBTExecutable))
 			{
@@ -69,7 +69,7 @@ namespace AutomationTool
 				}
 			}
 
-			IProcessResult Result = Run(UBTExecutable, CommandLine, Options: ERunOptions.AllowSpew | ERunOptions.NoStdOutCapture, Env: EnvVars);
+			IProcessResult Result = Run(UBTExecutable, CommandLine, Options: ERunOptions.AllowSpew | ERunOptions.NoStdOutCapture);
 			if(Result.ExitCode != 0)
 			{
 				throw new AutomationException((ExitCode)Result.ExitCode, "UnrealBuildTool failed. See log for more details. ({0})", CommandUtils.CombinePaths(Env.FinalLogFolder, LogName));
@@ -84,7 +84,7 @@ namespace AutomationTool
 		/// <param name="Platform">Platform to build for.</param>
 		/// <param name="Config">Configuration to build.</param>
 		/// <param name="AdditionalArgs">Additional arguments to pass on to UBT.</param>
-		public static string UBTCommandline(FileReference Project, string Target, string Platform, string Config, string AdditionalArgs = "")
+		public static string UBTCommandline(FileReference Project, string Target, UnrealTargetPlatform Platform, UnrealTargetConfiguration Config, string AdditionalArgs = "")
 		{
 			string CmdLine;
 			if (Project == null)
@@ -109,9 +109,9 @@ namespace AutomationTool
 		/// <param name="Config">Configuration to build.</param>
 		/// <param name="AdditionalArgs">Additional arguments to pass on to UBT.</param>
 		/// <param name="LogName">Optional logifle name.</param>
-		public static void RunUBT(CommandEnvironment Env, string UBTExecutable, FileReference Project, string Target, string Platform, string Config, string AdditionalArgs = "", Dictionary<string, string> EnvVars = null)
+		public static void RunUBT(CommandEnvironment Env, string UBTExecutable, FileReference Project, string Target, UnrealTargetPlatform Platform, UnrealTargetConfiguration Config, string AdditionalArgs = "")
 		{
-			RunUBT(Env, UBTExecutable, UBTCommandline(Project, Target, Platform, Config, AdditionalArgs), EnvVars);
+			RunUBT(Env, UBTExecutable, UBTCommandline(Project, Target, Platform, Config, AdditionalArgs));
 		}
 
 	}
diff --git a/Engine/Source/Programs/AutomationTool/AutomationUtils/UE4Build.cs b/Engine/Source/Programs/AutomationTool/AutomationUtils/UE4Build.cs
--- a/Engine/Source/Programs/AutomationTool/AutomationUtils/UE4Build.cs
+++ b/Engine/Source/Programs/AutomationTool/AutomationUtils/UE4Build.cs
@@ -118,17 +118,17 @@ namespace AutomationTool
 
 		public class XGEItem
 		{
-			public UnrealBuildTool.BuildManifest Manifest;
+			public BuildManifest Manifest;
 			public string CommandLine;
-			public UnrealBuildTool.UnrealTargetPlatform Platform;
-			public string Config;
+			public UnrealTargetPlatform Platform;
+			public UnrealTargetConfiguration Config;
 			public string TargetName;
 			public FileReference UProjectPath;
 			public List<string> XgeXmlFiles;
 			public string OutputCaption;
 		}
 
-		XGEItem XGEPrepareBuildWithUBT(string TargetName, UnrealBuildTool.UnrealTargetPlatform Platform, string Config, FileReference UprojectPath, bool ForceMonolithic = false, bool ForceNonUnity = false, bool ForceDebugInfo = false, string InAddArgs = "", bool ForceUnity = false, Dictionary<string, string> EnvVars = null)
+		XGEItem XGEPrepareBuildWithUBT(string TargetName, UnrealTargetPlatform Platform, UnrealTargetConfiguration Config, FileReference UprojectPath, bool ForceMonolithic = false, bool ForceNonUnity = false, bool ForceDebugInfo = false, string InAddArgs = "", bool ForceUnity = false)
 		{
 			string AddArgs = "";
 			if (UprojectPath != null)
@@ -165,7 +165,7 @@ namespace AutomationTool
 
 			using(TelemetryStopwatch GenerateManifestStopwatch = new TelemetryStopwatch("GenerateXGEManifest.{0}.{1}.{2}", TargetName, Platform.ToString(), Config))
 			{
-				CommandUtils.RunUBT(CommandUtils.CmdEnv, UBTExecutable: UBTExecutable, Project: UprojectPath, Target: TargetName, Platform: Platform.ToString(), Config: Config, AdditionalArgs: "-generatemanifest -nobuilduht -xgeexport" + AddArgs, EnvVars: EnvVars);
+				CommandUtils.RunUBT(CommandUtils.CmdEnv, UBTExecutable: UBTExecutable, Project: UprojectPath, Target: TargetName, Platform: Platform, Config: Config, AdditionalArgs: "-generatemanifest -nobuilduht -xgeexport" + AddArgs);
 			}
 
 			PrepareManifest(UBTManifest, true);
@@ -178,7 +178,7 @@ namespace AutomationTool
 			Result.OutputCaption = String.Format("{0}-{1}-{2}", TargetName, Platform.ToString(), Config.ToString());
 			CommandUtils.DeleteFile(UBTManifest);
 
-			Result.CommandLine = UBTExecutable + " " + CommandUtils.UBTCommandline(Project: UprojectPath, Target: TargetName, Platform: Platform.ToString(), Config: Config, AdditionalArgs: "-noxge -nobuilduht" + AddArgs);
+			Result.CommandLine = UBTExecutable + " " + CommandUtils.UBTCommandline(Project: UprojectPath, Target: TargetName, Platform: Platform, Config: Config, AdditionalArgs: "-noxge -nobuilduht" + AddArgs);
 
 			Result.XgeXmlFiles = new List<string>();
 			foreach (var XGEFile in FindXGEFiles())
@@ -227,7 +227,7 @@ namespace AutomationTool
 			}
 		}
 
-		void CleanWithUBT(string TargetName, UnrealBuildTool.UnrealTargetPlatform Platform, string Config, FileReference UprojectPath, bool ForceMonolithic = false, bool ForceNonUnity = false, bool ForceDebugInfo = false, string InAddArgs = "", bool ForceUnity = false, Dictionary<string, string> EnvVars = null)
+		void CleanWithUBT(string TargetName, UnrealTargetPlatform Platform, UnrealTargetConfiguration Config, FileReference UprojectPath, bool ForceMonolithic = false, bool ForceNonUnity = false, bool ForceDebugInfo = false, string InAddArgs = "", bool ForceUnity = false)
 		{
 			string AddArgs = "";
 			if (UprojectPath != null)
@@ -260,11 +260,11 @@ namespace AutomationTool
 			PrepareUBT();
 			using(TelemetryStopwatch CleanStopwatch = new TelemetryStopwatch("CleanWithUBT.{0}.{1}.{2}", TargetName, Platform.ToString(), Config))
 			{
-				CommandUtils.RunUBT(CommandUtils.CmdEnv, UBTExecutable: UBTExecutable, Project: UprojectPath, Target: TargetName, Platform: Platform.ToString(), Config: Config, AdditionalArgs: "-clean" + AddArgs, EnvVars: EnvVars);
+				CommandUtils.RunUBT(CommandUtils.CmdEnv, UBTExecutable: UBTExecutable, Project: UprojectPath, Target: TargetName, Platform: Platform, Config: Config, AdditionalArgs: "-clean" + AddArgs);
 			}
         }
 
-		BuildManifest BuildWithUBT(string TargetName, UnrealBuildTool.UnrealTargetPlatform TargetPlatform, string Config, FileReference UprojectPath, bool ForceMonolithic = false, bool ForceNonUnity = false, bool ForceDebugInfo = false, bool ForceFlushMac = false, bool DisableXGE = false, string InAddArgs = "", bool ForceUnity = false, Dictionary<string, string> EnvVars = null)
+		BuildManifest BuildWithUBT(string TargetName, UnrealTargetPlatform TargetPlatform, UnrealTargetConfiguration Config, FileReference UprojectPath, bool ForceMonolithic = false, bool ForceNonUnity = false, bool ForceDebugInfo = false, bool ForceFlushMac = false, bool DisableXGE = false, string InAddArgs = "", bool ForceUnity = false)
 		{
 			string AddArgs = "";
 			if (UprojectPath != null)
@@ -311,14 +311,14 @@ namespace AutomationTool
 				CommandUtils.DeleteFile(UBTManifest);
 				using(TelemetryStopwatch PrepareManifestStopwatch = new TelemetryStopwatch("PrepareUBTManifest.{0}.{1}.{2}", TargetName, TargetPlatform.ToString(), Config))
 				{
-					CommandUtils.RunUBT(CommandUtils.CmdEnv, UBTExecutable: UBTExecutable, Project: UprojectPath, Target: TargetName, Platform: TargetPlatform.ToString(), Config: Config, AdditionalArgs: AddArgs + " -generatemanifest", EnvVars: EnvVars);
+					CommandUtils.RunUBT(CommandUtils.CmdEnv, UBTExecutable: UBTExecutable, Project: UprojectPath, Target: TargetName, Platform: TargetPlatform, Config: Config, AdditionalArgs: AddArgs + " -generatemanifest");
 				}
                 Manifest = PrepareManifest(UBTManifest, false);
 			}
 
 			using(TelemetryStopwatch CompileStopwatch = new TelemetryStopwatch("Compile.{0}.{1}.{2}", TargetName, TargetPlatform.ToString(), Config))
 			{
-				CommandUtils.RunUBT(CommandUtils.CmdEnv, UBTExecutable: UBTExecutable, Project: UprojectPath, Target: TargetName, Platform: TargetPlatform.ToString(), Config: Config, AdditionalArgs: AddArgs, EnvVars: EnvVars);
+				CommandUtils.RunUBT(CommandUtils.CmdEnv, UBTExecutable: UBTExecutable, Project: UprojectPath, Target: TargetName, Platform: TargetPlatform, Config: Config, AdditionalArgs: AddArgs);
 			}
 
 			if (UseManifest)
@@ -1187,7 +1187,7 @@ namespace AutomationTool
 		/// <param name="InUpdateVersionFiles">True if the version files are to be updated </param>
 		/// <param name="InForceNoXGE">If true will force XGE off</param>
 		/// <param name="InUseParallelExecutor">If true AND XGE not present or not being used then use ParallelExecutor</param>
-		public void Build(BuildAgenda Agenda, bool? InDeleteBuildProducts = null, bool InUpdateVersionFiles = true, bool InForceNoXGE = false, bool InUseParallelExecutor = false, bool InForceNonUnity = false, bool InForceUnity = false, bool InShowProgress = false, Dictionary<UnrealBuildTool.UnrealTargetPlatform, Dictionary<string, string>> PlatformEnvVars = null, int? InChangelistNumberOverride = null, Dictionary<BuildTarget, BuildManifest> InTargetToManifest = null)
+		public void Build(BuildAgenda Agenda, bool? InDeleteBuildProducts = null, bool InUpdateVersionFiles = true, bool InForceNoXGE = false, bool InUseParallelExecutor = false, bool InForceNonUnity = false, bool InForceUnity = false, bool InShowProgress = false, int? InChangelistNumberOverride = null, Dictionary<BuildTarget, BuildManifest> InTargetToManifest = null)
 		{
 			if (!CommandUtils.CmdEnv.HasCapabilityToCompile)
 			{
@@ -1290,7 +1290,7 @@ namespace AutomationTool
 					bool bClean = Target.Clean ?? DeleteBuildProducts;
 					if (bClean)
 					{
-						CleanWithUBT(Target.TargetName, Target.Platform, Target.Config.ToString(), Target.UprojectPath, bForceMonolithic, bForceNonUnity, bForceDebugInfo, Target.UBTArgs, bForceUnity);
+						CleanWithUBT(Target.TargetName, Target.Platform, Target.Config, Target.UprojectPath, bForceMonolithic, bForceNonUnity, bForceDebugInfo, Target.UBTArgs, bForceUnity);
 					}
 				}
 				foreach (var Target in Agenda.Targets)
@@ -1300,7 +1300,7 @@ namespace AutomationTool
 						// When building a target for Mac or iOS, use UBT's -flushmac option to clean up the remote builder
 						bool bForceFlushMac = DeleteBuildProducts && (Target.Platform == UnrealBuildTool.UnrealTargetPlatform.Mac || Target.Platform == UnrealBuildTool.UnrealTargetPlatform.IOS);
 						CommandUtils.LogSetProgress(InShowProgress, "Building header tool...");
-						BuildManifest Manifest = BuildWithUBT(Target.TargetName, Target.Platform, Target.Config.ToString(), Target.UprojectPath, bForceMonolithic, bForceNonUnity, bForceDebugInfo, bForceFlushMac, !bCanUseXGE || !CanUseXGE(Target.Platform), Target.UBTArgs, bForceUnity);
+						BuildManifest Manifest = BuildWithUBT(Target.TargetName, Target.Platform, Target.Config, Target.UprojectPath, bForceMonolithic, bForceNonUnity, bForceDebugInfo, bForceFlushMac, !bCanUseXGE || !CanUseXGE(Target.Platform), Target.UBTArgs, bForceUnity);
 						if(InTargetToManifest != null)
 						{
 							InTargetToManifest[Target] = Manifest;
@@ -1314,7 +1314,7 @@ namespace AutomationTool
 				{
 					if (Target.TargetName != "UnrealHeaderTool" && (bCanUseXGE? CanUseXGE(Target.Platform) : CanUseParallelExecutor(Target.Platform)))
 					{
-						XGEItem Item = XGEPrepareBuildWithUBT(Target.TargetName, Target.Platform, Target.Config.ToString(), Target.UprojectPath, bForceMonolithic, bForceNonUnity, bForceDebugInfo, Target.UBTArgs, bForceUnity);
+						XGEItem Item = XGEPrepareBuildWithUBT(Target.TargetName, Target.Platform, Target.Config, Target.UprojectPath, bForceMonolithic, bForceNonUnity, bForceDebugInfo, Target.UBTArgs, bForceUnity);
 						if(InTargetToManifest != null)
 						{
 							InTargetToManifest[Target] = Item.Manifest;
@@ -1379,14 +1379,14 @@ namespace AutomationTool
 					bool bClean = Target.Clean ?? DeleteBuildProducts;
 					if (bClean)
 					{
-						CleanWithUBT(Target.TargetName, Target.Platform, Target.Config.ToString(), Target.UprojectPath, bForceMonolithic, bForceNonUnity, bForceDebugInfo, Target.UBTArgs, bForceUnity);
+						CleanWithUBT(Target.TargetName, Target.Platform, Target.Config, Target.UprojectPath, bForceMonolithic, bForceNonUnity, bForceDebugInfo, Target.UBTArgs, bForceUnity);
 					}
 				}
 				foreach (var Target in Agenda.Targets)
 				{
 					// When building a target for Mac or iOS, use UBT's -flushmac option to clean up the remote builder
 					bool bForceFlushMac = DeleteBuildProducts && (Target.Platform == UnrealBuildTool.UnrealTargetPlatform.Mac || Target.Platform == UnrealBuildTool.UnrealTargetPlatform.IOS);
-					BuildManifest Manifest = BuildWithUBT(Target.TargetName, Target.Platform, Target.Config.ToString(), Target.UprojectPath, bForceMonolithic, bForceNonUnity, bForceDebugInfo, bForceFlushMac, true, Target.UBTArgs, bForceUnity);
+					BuildManifest Manifest = BuildWithUBT(Target.TargetName, Target.Platform, Target.Config, Target.UprojectPath, bForceMonolithic, bForceNonUnity, bForceDebugInfo, bForceFlushMac, true, Target.UBTArgs, bForceUnity);
 					if(InTargetToManifest != null)
 					{
 						InTargetToManifest[Target] = Manifest;
