UAT: Always use the manifest for determining output files. UBT will initialize it to an empty list of files if the arguments deem it necessary.

#rnx

diff --git a/Engine/Source/Programs/AutomationTool/AutomationUtils/Platform.cs b/Engine/Source/Programs/AutomationTool/AutomationUtils/Platform.cs
--- a/Engine/Source/Programs/AutomationTool/AutomationUtils/Platform.cs
+++ b/Engine/Source/Programs/AutomationTool/AutomationUtils/Platform.cs
@@ -553,11 +553,6 @@ namespace AutomationTool
 			return 0;
 		}
 
-		public virtual bool ShouldUseManifestForUBTBuilds(string AddArgs)
-		{
-			return true;
-		}
-
 		/// <summary>
 		/// Determines whether we should stage a UE4CommandLine.txt for this platform
 		/// </summary>
diff --git a/Engine/Source/Programs/AutomationTool/AutomationUtils/UE4Build.cs b/Engine/Source/Programs/AutomationTool/AutomationUtils/UE4Build.cs
--- a/Engine/Source/Programs/AutomationTool/AutomationUtils/UE4Build.cs
+++ b/Engine/Source/Programs/AutomationTool/AutomationUtils/UE4Build.cs
@@ -300,35 +300,22 @@ namespace AutomationTool
 
 			PrepareUBT();
 
-			// let the platform determine when to use the manifest
-            bool UseManifest = Platform.GetPlatform(TargetPlatform).ShouldUseManifestForUBTBuilds(AddArgs);
-
-			BuildManifest Manifest = null;
-			if (UseManifest)
+			string UBTManifest = GetUBTManifest(UprojectPath, AddArgs);
+			CommandUtils.DeleteFile(UBTManifest);
+			using(TelemetryStopwatch PrepareManifestStopwatch = new TelemetryStopwatch("PrepareUBTManifest.{0}.{1}.{2}", TargetName, TargetPlatform.ToString(), Config))
 			{
-                string UBTManifest = GetUBTManifest(UprojectPath, AddArgs);
-
-				CommandUtils.DeleteFile(UBTManifest);
-				using(TelemetryStopwatch PrepareManifestStopwatch = new TelemetryStopwatch("PrepareUBTManifest.{0}.{1}.{2}", TargetName, TargetPlatform.ToString(), Config))
-				{
-					CommandUtils.RunUBT(CommandUtils.CmdEnv, UBTExecutable: UBTExecutable, Project: UprojectPath, Target: TargetName, Platform: TargetPlatform, Config: Config, AdditionalArgs: AddArgs + " -generatemanifest");
-				}
-                Manifest = PrepareManifest(UBTManifest, false);
+				CommandUtils.RunUBT(CommandUtils.CmdEnv, UBTExecutable: UBTExecutable, Project: UprojectPath, Target: TargetName, Platform: TargetPlatform, Config: Config, AdditionalArgs: AddArgs + " -generatemanifest");
 			}
+			BuildManifest Manifest = PrepareManifest(UBTManifest, false);
 
 			using(TelemetryStopwatch CompileStopwatch = new TelemetryStopwatch("Compile.{0}.{1}.{2}", TargetName, TargetPlatform.ToString(), Config))
 			{
 				CommandUtils.RunUBT(CommandUtils.CmdEnv, UBTExecutable: UBTExecutable, Project: UprojectPath, Target: TargetName, Platform: TargetPlatform, Config: Config, AdditionalArgs: AddArgs);
 			}
 
-			if (UseManifest)
-			{
-                string UBTManifest = GetUBTManifest(UprojectPath, AddArgs);
-
-				AddBuildProductsFromManifest(UBTManifest);
+			AddBuildProductsFromManifest(UBTManifest);
+			CommandUtils.DeleteFile(UBTManifest);
 
-				CommandUtils.DeleteFile(UBTManifest);
-			}
 			return Manifest;
 		}
 
diff --git a/Engine/Source/Programs/AutomationTool/Mac/MacPlatform.Automation.cs b/Engine/Source/Programs/AutomationTool/Mac/MacPlatform.Automation.cs
--- a/Engine/Source/Programs/AutomationTool/Mac/MacPlatform.Automation.cs
+++ b/Engine/Source/Programs/AutomationTool/Mac/MacPlatform.Automation.cs
@@ -412,13 +412,6 @@ public class MacPlatform : Platform
 
 	public override bool IsSupported { get { return true; } }
 
-	public override bool ShouldUseManifestForUBTBuilds(string AddArgs)
-	{
-		// don't use the manifest to set up build products if we are compiling Mac under Windows and we aren't going to copy anything back to the PC
-		bool bIsBuildingRemotely = UnrealBuildTool.BuildHostPlatform.Current.Platform != UnrealTargetPlatform.Mac;
-		bool bUseManifest = !bIsBuildingRemotely || AddArgs.IndexOf("-CopyAppBundleBackToDevice", StringComparison.InvariantCultureIgnoreCase) > 0;
-		return bUseManifest;
-	}
 	public override List<string> GetDebugFileExtensions()
 	{
 		return new List<string> { ".dSYM" };
