UAT: Generate the manifest for each compiled target at the same time as exporting its build steps. Prevents module manifests being clobbered due to being overwritten twice.

#rnx

diff --git a/Engine/Source/Programs/AutomationTool/AutomationUtils/CommandUtils.cs b/Engine/Source/Programs/AutomationTool/AutomationUtils/CommandUtils.cs
--- a/Engine/Source/Programs/AutomationTool/AutomationUtils/CommandUtils.cs
+++ b/Engine/Source/Programs/AutomationTool/AutomationUtils/CommandUtils.cs
@@ -486,6 +486,18 @@ namespace AutomationTool
 				throw new AutomationException(String.Format("Failed to delete file '{0}'", NormalizedFilename));
 			}
 		}
+
+		/// <summary>
+		/// Deletes a file(s). 
+		/// If the file does not exist, silently succeeds.
+		/// If the deletion of the file fails, this function throws an Exception.
+		/// </summary>
+		/// <param name="Filenames">Filename</param>
+		public static void DeleteFile(FileReference FileName)
+		{
+			DeleteFile(FileName.FullName);
+		}
+
         /// <summary>
         /// Deletes a file(s). 
         /// If the file does not exist, silently succeeds.
@@ -1469,12 +1481,12 @@ namespace AutomationTool
 		/// <summary>
 		/// Reads a file manifest and returns it
 		/// </summary>
-		/// <param name="ManifestName">ManifestName</param>
+		/// <param name="ManifestFile">ManifestName</param>
 		/// <returns></returns>
-		public static UnrealBuildTool.BuildManifest ReadManifest(string ManifestName)
+		public static BuildManifest ReadManifest(FileReference ManifestFile)
 		{
 			// Create a new default instance of the type
-			UnrealBuildTool.BuildManifest Instance = new UnrealBuildTool.BuildManifest();
+			BuildManifest Instance = new UnrealBuildTool.BuildManifest();
 			XmlReader XmlStream = null;
 			try
 			{
@@ -1484,13 +1496,13 @@ namespace AutomationTool
 				ReaderSettings.IgnoreComments = true;
 
 				// Get the xml data stream to read from
-				XmlStream = XmlReader.Create( ManifestName, ReaderSettings );
+				XmlStream = XmlReader.Create( ManifestFile.FullName, ReaderSettings );
 
 				// Creates an instance of the XmlSerializer class so we can read the settings object
 				XmlSerializer ObjectReader = new XmlSerializer( typeof( UnrealBuildTool.BuildManifest ) );
 
 				// Create an object from the xml data
-				Instance = ( UnrealBuildTool.BuildManifest )ObjectReader.Deserialize( XmlStream );
+				Instance = ( BuildManifest )ObjectReader.Deserialize( XmlStream );
 			}
 			catch( Exception Ex )
 			{
diff --git a/Engine/Source/Programs/AutomationTool/AutomationUtils/UE4Build.cs b/Engine/Source/Programs/AutomationTool/AutomationUtils/UE4Build.cs
--- a/Engine/Source/Programs/AutomationTool/AutomationUtils/UE4Build.cs
+++ b/Engine/Source/Programs/AutomationTool/AutomationUtils/UE4Build.cs
@@ -72,7 +72,7 @@ namespace AutomationTool
 			CommandUtils.LogLog("Copied UBT manifest to {0}", OutFile);
 
 
-			UnrealBuildTool.BuildManifest Manifest = CommandUtils.ReadManifest(ManifestName);
+			UnrealBuildTool.BuildManifest Manifest = CommandUtils.ReadManifest(new FileReference(ManifestName));
 			foreach (string Item in Manifest.BuildProducts)
 			{
 				if(bAddReceipt && IsBuildReceipt(Item))
@@ -91,13 +91,14 @@ namespace AutomationTool
 				|| FileName.EndsWith("buildid.txt", StringComparison.InvariantCultureIgnoreCase);
 		}
 
-		void AddBuildProductsFromManifest(string ManifestName)
+		BuildManifest AddBuildProductsFromManifest(FileReference ManifestFile)
 		{
-			if (CommandUtils.FileExists(ManifestName) == false)
+			if (!FileReference.Exists(ManifestFile))
 			{
-				throw new AutomationException("BUILD FAILED UBT Manifest {0} does not exist.", ManifestName);
+				throw new AutomationException("BUILD FAILED UBT Manifest {0} does not exist.", ManifestFile);
 			}
-			UnrealBuildTool.BuildManifest Manifest = CommandUtils.ReadManifest(ManifestName);
+
+			BuildManifest Manifest = CommandUtils.ReadManifest(ManifestFile);
 			foreach (string Item in Manifest.BuildProducts)
 			{
 				if (!CommandUtils.FileExists_NoExceptions(Item) && !CommandUtils.DirectoryExists_NoExceptions(Item))
@@ -106,6 +107,7 @@ namespace AutomationTool
 				}
 				AddBuildProduct(Item);
 			}
+			return Manifest;
 		}
 		
 		private void PrepareUBT()
@@ -119,7 +121,6 @@ namespace AutomationTool
 		public class XGEItem
 		{
 			public BuildManifest Manifest;
-			public string CommandLine;
 			public UnrealTargetPlatform Platform;
 			public UnrealTargetConfiguration Config;
 			public string TargetName;
@@ -156,29 +157,21 @@ namespace AutomationTool
 
 			PrepareUBT();
 
-            string UBTManifest = GetUBTManifest(UprojectPath, AddArgs);
-
-			CommandUtils.DeleteFile(UBTManifest);
-			XGEItem Result = new XGEItem();
+            FileReference ManifestFile = GetManifestFile(UprojectPath);
+			CommandUtils.DeleteFile(ManifestFile);
 
 			ClearExportedXGEXML();
 
-			using(TelemetryStopwatch GenerateManifestStopwatch = new TelemetryStopwatch("GenerateXGEManifest.{0}.{1}.{2}", TargetName, Platform.ToString(), Config))
-			{
-				CommandUtils.RunUBT(CommandUtils.CmdEnv, UBTExecutable: UBTExecutable, Project: UprojectPath, Target: TargetName, Platform: Platform, Config: Config, AdditionalArgs: "-generatemanifest -nobuilduht -xgeexport" + AddArgs);
-			}
-
-			PrepareManifest(UBTManifest, true);
+			CommandUtils.RunUBT(CommandUtils.CmdEnv, UBTExecutable: UBTExecutable, Project: UprojectPath, Target: TargetName, Platform: Platform, Config: Config, AdditionalArgs: String.Format("-Manifest={0} -nobuilduht -xgeexport {1}", CommandUtils.MakePathSafeToUseWithCommandLine(ManifestFile.FullName), AddArgs));
 
+			XGEItem Result = new XGEItem();
 			Result.Platform = Platform;
 			Result.Config = Config;
 			Result.TargetName = TargetName;
 			Result.UProjectPath = UprojectPath;
-			Result.Manifest = CommandUtils.ReadManifest(UBTManifest);
+			Result.Manifest = CommandUtils.ReadManifest(ManifestFile);
 			Result.OutputCaption = String.Format("{0}-{1}-{2}", TargetName, Platform.ToString(), Config.ToString());
-			CommandUtils.DeleteFile(UBTManifest);
-
-			Result.CommandLine = UBTExecutable + " " + CommandUtils.UBTCommandline(Project: UprojectPath, Target: TargetName, Platform: Platform, Config: Config, AdditionalArgs: "-noxge -nobuilduht" + AddArgs);
+			CommandUtils.DeleteFile(ManifestFile);
 
 			Result.XgeXmlFiles = new List<string>();
 			foreach (var XGEFile in FindXGEFiles())
@@ -208,13 +201,10 @@ namespace AutomationTool
 
 		void XGEFinishBuildWithUBT(XGEItem Item)
 		{
-			if(!Item.CommandLine.Contains("-nolink"))
+			// run the deployment steps, if necessary
+			foreach(string DeployTargetFile in Item.Manifest.DeployTargetFiles)
 			{
-				// run the deployment steps, if necessary
-				foreach(string DeployTargetFile in Item.Manifest.DeployTargetFiles)
-				{
-					CommandUtils.RunUBT(CommandUtils.CmdEnv, UBTExecutable, String.Format("-deploy=\"{0}\" -nomutex", DeployTargetFile));
-				}
+				CommandUtils.RunUBT(CommandUtils.CmdEnv, UBTExecutable, String.Format("-deploy=\"{0}\" -nomutex", DeployTargetFile));
 			}
 
 			foreach (string ManifestItem in Item.Manifest.BuildProducts)
@@ -300,21 +290,13 @@ namespace AutomationTool
 
 			PrepareUBT();
 
-			string UBTManifest = GetUBTManifest(UprojectPath, AddArgs);
-			CommandUtils.DeleteFile(UBTManifest);
-			using(TelemetryStopwatch PrepareManifestStopwatch = new TelemetryStopwatch("PrepareUBTManifest.{0}.{1}.{2}", TargetName, TargetPlatform.ToString(), Config))
-			{
-				CommandUtils.RunUBT(CommandUtils.CmdEnv, UBTExecutable: UBTExecutable, Project: UprojectPath, Target: TargetName, Platform: TargetPlatform, Config: Config, AdditionalArgs: AddArgs + " -generatemanifest");
-			}
-			BuildManifest Manifest = PrepareManifest(UBTManifest, false);
+			FileReference ManifestFile = GetManifestFile(UprojectPath);
+			CommandUtils.DeleteFile(ManifestFile);
 
-			using(TelemetryStopwatch CompileStopwatch = new TelemetryStopwatch("Compile.{0}.{1}.{2}", TargetName, TargetPlatform.ToString(), Config))
-			{
-				CommandUtils.RunUBT(CommandUtils.CmdEnv, UBTExecutable: UBTExecutable, Project: UprojectPath, Target: TargetName, Platform: TargetPlatform, Config: Config, AdditionalArgs: AddArgs);
-			}
+			CommandUtils.RunUBT(CommandUtils.CmdEnv, UBTExecutable: UBTExecutable, Project: UprojectPath, Target: TargetName, Platform: TargetPlatform, Config: Config, AdditionalArgs: String.Format("{0} -Manifest={1}", AddArgs, CommandUtils.MakePathSafeToUseWithCommandLine(ManifestFile.FullName)));
 
-			AddBuildProductsFromManifest(UBTManifest);
-			CommandUtils.DeleteFile(UBTManifest);
+			BuildManifest Manifest = AddBuildProductsFromManifest(ManifestFile);
+			CommandUtils.DeleteFile(ManifestFile);
 
 			return Manifest;
 		}
@@ -1568,16 +1550,16 @@ namespace AutomationTool
 			}
 		}
 
-		string GetUBTManifest(FileReference UProjectPath, string InAddArgs)
+		FileReference GetManifestFile(FileReference ProjectFile)
 		{
 			// Can't write to Engine directory on installed builds
-			if (Automation.IsEngineInstalled() && UProjectPath != null)
+			if (Automation.IsEngineInstalled() && ProjectFile != null)
 			{
-				return Path.Combine(Path.GetDirectoryName(UProjectPath.FullName), "Intermediate/Build/Manifest.xml");				
+				return FileReference.Combine(ProjectFile.Directory, "Intermediate", "Build", "Manifest.xml");
 			}
 			else
 			{
-				return CommandUtils.CombinePaths(CommandUtils.CmdEnv.LocalRoot, @"/Engine/Intermediate/Build/Manifest.xml");
+				return FileReference.Combine(CommandUtils.EngineDirectory, "Intermediate", "Build", "Manifest.xml");
 			}
 		}
 
