UAT: Fix ParallelExecutor not being used when XGE is not installed, and UBT being run twice for each target.

#rnx

diff --git a/Engine/Source/Programs/AutomationTool/AutomationUtils/UE4Build.cs b/Engine/Source/Programs/AutomationTool/AutomationUtils/UE4Build.cs
--- a/Engine/Source/Programs/AutomationTool/AutomationUtils/UE4Build.cs
+++ b/Engine/Source/Programs/AutomationTool/AutomationUtils/UE4Build.cs
@@ -50,39 +50,6 @@ namespace AutomationTool
 			}
 		}
 
-		BuildManifest PrepareManifest(string ManifestName, bool bAddReceipt)
-		{
-			if (CommandUtils.FileExists(ManifestName) == false)
-			{
-				throw new AutomationException("BUILD FAILED UBT Manifest {0} does not exist.", ManifestName);
-			}
-			int FileNum = 0;
-			string OutFile;
-			while (true)
-			{
-				OutFile = CommandUtils.CombinePaths(CommandUtils.CmdEnv.LogFolder, String.Format("UBTManifest.{0}", FileNum) + Path.GetExtension(ManifestName));
-				FileInfo ItemInfo = new FileInfo(OutFile);
-				if (!ItemInfo.Exists)
-				{
-					break;
-				}
-				FileNum++;
-			}
-			CommandUtils.CopyFile(ManifestName, OutFile);
-			CommandUtils.LogLog("Copied UBT manifest to {0}", OutFile);
-
-
-			UnrealBuildTool.BuildManifest Manifest = CommandUtils.ReadManifest(new FileReference(ManifestName));
-			foreach (string Item in Manifest.BuildProducts)
-			{
-				if(bAddReceipt && IsBuildReceipt(Item))
-				{
-					AddBuildProduct(Item);
-				}
-			}
-			return Manifest;
-		}
-
 		static bool IsBuildReceipt(string FileName)
 		{
 			return FileName.EndsWith(".version", StringComparison.InvariantCultureIgnoreCase)
@@ -411,8 +378,7 @@ namespace AutomationTool
 				Version = new BuildVersion();
 			}
 
-
-			var Result = new List<FileReference>();
+			List<FileReference> Result = new List<FileReference>();
 			{
 				if (bDoUpdateVersionFiles)
 				{
@@ -470,59 +436,6 @@ namespace AutomationTool
 			return Result;
 		}
 
-		/// <summary>
-		/// Updates the public key file (decryption)
-		/// </summary>
-		public void UpdatePublicKey(string KeyFilename)
-		{
-			string Modulus = "\"0x0\"";
-			string PublicKeyExponent = "\"0x0\"";
-
-			if (!string.IsNullOrEmpty(KeyFilename))
-			{
-				var Lines = ReadEncryptionKeys(KeyFilename);
-				if (Lines != null && Lines.Length >= 3)
-				{
-					Modulus = String.Format("\"{0}\"", Lines[1]);
-					PublicKeyExponent = String.Format("\"{0}\"", Lines[2]);
-				}
-				else
-				{
-					CommandUtils.LogError("{0} doesn't look like a valid encryption key file or value.", KeyFilename);
-				}
-			}
-
-			FileReference VerFile = FileReference.Combine(CommandUtils.EngineDirectory, "Source", "Runtime", "PakFile", "Private", "PublicKey.inl");
-
-			CommandUtils.LogVerbose("Updating {0} with:", VerFile);
-			CommandUtils.LogVerbose(" #define DECRYPTION_KEY_EXPONENT {0}", PublicKeyExponent);
-			CommandUtils.LogVerbose(" #define DECRYPTION_KEY_MODULUS {0}", Modulus);
-
-			VersionFileUpdater PublicKeyInl = new VersionFileUpdater(VerFile);
-			PublicKeyInl.ReplaceLine("#define DECRYPTION_KEY_EXPONENT ", PublicKeyExponent);
-			PublicKeyInl.ReplaceLine("#define DECRYPTION_KEY_MODULUS ", Modulus);
-			PublicKeyInl.Commit();
-		}
-
-		/// <summary>
-		/// Parses the encryption keys from the command line or loads them from a file.
-		/// </summary>
-		/// <param name="KeyFilenameOrValues">Key values or filename</param>
-		/// <returns>List of three encruption key values: private exponent, modulus and public exponent.</returns>
-		private string[] ReadEncryptionKeys(string KeyFilenameOrValues)
-		{
-			string[] Values;
-			if (KeyFilenameOrValues.StartsWith("0x", StringComparison.InvariantCultureIgnoreCase))
-			{
-				Values = KeyFilenameOrValues.Split(new char[] { '+' }, StringSplitOptions.RemoveEmptyEntries);
-			}
-			else
-			{
-				Values = CommandUtils.ReadAllLines(KeyFilenameOrValues);
-			}
-			return Values;
-		}
-
 		[DebuggerDisplay("{TargetName} {Platform} {Config}")]
 		public class BuildTarget
 		{
@@ -539,12 +452,12 @@ namespace AutomationTool
 			/// <summary>
 			/// Platform to build
 			/// </summary>
-			public UnrealBuildTool.UnrealTargetPlatform Platform;
+			public UnrealTargetPlatform Platform;
 
 			/// <summary>
 			/// Configuration to build
 			/// </summary>
-			public UnrealBuildTool.UnrealTargetConfiguration Config;
+			public UnrealTargetConfiguration Config;
 
 			/// <summary>
 			/// Extra UBT args
@@ -610,7 +523,7 @@ namespace AutomationTool
 			/// <param name="InConfiguration">Configuration</param>
 			/// <param name="InUprojectPath">Path to optional uproject file</param>
 			/// <param name="InAddArgs">Specifies additional arguments for UBT</param>
-			public void AddTarget(string TargetName, UnrealBuildTool.UnrealTargetPlatform InPlatform, UnrealBuildTool.UnrealTargetConfiguration InConfiguration, FileReference InUprojectPath = null, string InAddArgs = "")
+			public void AddTarget(string TargetName, UnrealTargetPlatform InPlatform, UnrealTargetConfiguration InConfiguration, FileReference InUprojectPath = null, string InAddArgs = "")
 			{
 				// Is this platform a compilable target?
 				if (!Platform.GetPlatform(InPlatform).CanBeCompiled())
@@ -636,7 +549,7 @@ namespace AutomationTool
 			/// <param name="InConfiguration">Configuration</param>
 			/// <param name="InUprojectPath">Path to optional uproject file</param>
 			/// <param name="InAddArgs">Specifies additional arguments for UBT</param>
-			public void AddTargets(string[] TargetNames, UnrealBuildTool.UnrealTargetPlatform InPlatform, UnrealBuildTool.UnrealTargetConfiguration InConfiguration, FileReference InUprojectPath = null, string InAddArgs = "")
+			public void AddTargets(string[] TargetNames, UnrealTargetPlatform InPlatform, UnrealTargetConfiguration InConfiguration, FileReference InUprojectPath = null, string InAddArgs = "")
 			{
 				// Is this platform a compilable target?
 				if (!Platform.GetPlatform(InPlatform).CanBeCompiled())
@@ -683,7 +596,7 @@ namespace AutomationTool
 				{
 					try
 					{
-						string FullPath = Path.Combine(PathDir, "xgConsole" + Platform.GetExeExtension(UnrealBuildTool.BuildHostPlatform.Current.Platform));
+						string FullPath = Path.Combine(PathDir, "xgConsole" + Platform.GetExeExtension(BuildHostPlatform.Current.Platform));
 						if (CommandUtils.FileExists(FullPath))
 						{
 							XGEConsoleExePath = FullPath;
@@ -1123,14 +1036,14 @@ namespace AutomationTool
 			}
 		}
 
-		public bool CanUseXGE(UnrealBuildTool.UnrealTargetPlatform Platform)
+		public bool CanUseXGE(UnrealTargetPlatform Platform)
 		{
-			return UnrealBuildTool.PlatformExports.CanUseXGE(Platform);
+			return PlatformExports.CanUseXGE(Platform);
 		}
 
-		public bool CanUseParallelExecutor(UnrealBuildTool.UnrealTargetPlatform Platform)
+		public bool CanUseParallelExecutor(UnrealTargetPlatform Platform)
 		{
-			return UnrealBuildTool.PlatformExports.CanUseParallelExecutor(Platform);
+			return PlatformExports.CanUseParallelExecutor(Platform);
 		}
 
 		private bool ParseParam(string Name)
@@ -1171,7 +1084,7 @@ namespace AutomationTool
 			//////////////////////////////////////
 
 			// make a set of unique platforms involved
-			var UniquePlatforms = new List<UnrealBuildTool.UnrealTargetPlatform>();
+			var UniquePlatforms = new List<UnrealTargetPlatform>();
 			foreach (var Target in Agenda.Targets)
 			{
 				if (!UniquePlatforms.Contains(Target.Platform))
@@ -1231,14 +1144,8 @@ namespace AutomationTool
 			bool bForceNonUnity = ParseParam("ForceNonUnity") || InForceNonUnity;
 			bool bForceUnity = ParseParam("ForceUnity") || InForceUnity;
 			bool bForceDebugInfo = ParseParam("ForceDebugInfo");
-			bool bUsedXGE = false;
-			bool bCanUseXGE = !ParseParam("NoXGE") && !InForceNoXGE;
 			string XGEConsole = XGEConsoleExe();
-			if (bCanUseXGE && string.IsNullOrEmpty(XGEConsole))
-			{
-				CommandUtils.LogLog("XGE was requested, but is unavailable, so we won't use it.");
-				bCanUseXGE = false;
-			}
+			bool bCanUseXGE = !ParseParam("NoXGE") && !InForceNoXGE && !String.IsNullOrEmpty(XGEConsole);
 
 			// only run ParallelExecutor if not running XGE (and we've requested ParallelExecutor and it exists)
 			bool bCanUseParallelExecutor = !bCanUseXGE && InUseParallelExecutor && (HostPlatform.Current.HostEditorPlatform == UnrealTargetPlatform.Win64);
@@ -1249,118 +1156,112 @@ namespace AutomationTool
 			CommandUtils.LogLog("************************* UseXGE: {0}", bCanUseXGE);
 			CommandUtils.LogLog("************************* UseParallelExecutor: {0}", bCanUseParallelExecutor);
 
-			// process XGE files
-			if (bCanUseXGE || bCanUseParallelExecutor)
+			// Clean all the targets
+			foreach (BuildTarget Target in Agenda.Targets)
 			{
-				string TaskFilePath = CommandUtils.CombinePaths(CommandUtils.CmdEnv.LogFolder, @"UAT_XGE.xml");
+				bool bClean = Target.Clean ?? DeleteBuildProducts;
+				if (bClean)
+				{
+					CleanWithUBT(Target.TargetName, Target.Platform, Target.Config, Target.UprojectPath, bForceMonolithic, bForceNonUnity, bForceDebugInfo, Target.UBTArgs, bForceUnity);
+				}
+			}
 
-				foreach (var Target in Agenda.Targets)
+			// Filter the targets into those which can be built in parallel, vs those that must be executed serially
+			List<BuildTarget> NonParallelTargets = new List<BuildTarget>();
+			List<BuildTarget> ParallelTargets = new List<BuildTarget>();
+			foreach (BuildTarget Target in Agenda.Targets)
+			{
+				if(Target.TargetName == "UnrealHeaderTool")
 				{
-					bool bClean = Target.Clean ?? DeleteBuildProducts;
-					if (bClean)
-					{
-						CleanWithUBT(Target.TargetName, Target.Platform, Target.Config, Target.UprojectPath, bForceMonolithic, bForceNonUnity, bForceDebugInfo, Target.UBTArgs, bForceUnity);
-					}
+					NonParallelTargets.Insert(0, Target);
 				}
-				foreach (var Target in Agenda.Targets)
+				else if(bCanUseXGE)
 				{
-					if (Target.TargetName == "UnrealHeaderTool" || !(bCanUseXGE? CanUseXGE(Target.Platform) : CanUseParallelExecutor(Target.Platform)))
+					if(CanUseXGE(Target.Platform))
 					{
-						// When building a target for Mac or iOS, use UBT's -flushmac option to clean up the remote builder
-						bool bForceFlushMac = DeleteBuildProducts && (Target.Platform == UnrealBuildTool.UnrealTargetPlatform.Mac || Target.Platform == UnrealBuildTool.UnrealTargetPlatform.IOS);
-						CommandUtils.LogSetProgress(InShowProgress, "Building header tool...");
-						BuildManifest Manifest = BuildWithUBT(Target.TargetName, Target.Platform, Target.Config, Target.UprojectPath, bForceMonolithic, bForceNonUnity, bForceDebugInfo, bForceFlushMac, !bCanUseXGE || !CanUseXGE(Target.Platform), Target.UBTArgs, bForceUnity);
-						if(InTargetToManifest != null)
-						{
-							InTargetToManifest[Target] = Manifest;
-						}
+						ParallelTargets.Add(Target);
 					}
-				}
-
-				List<XGEItem> XGEItems = new List<XGEItem>();
-				CommandUtils.LogSetProgress(InShowProgress, "Generating headers...");
-				foreach (var Target in Agenda.Targets)
-				{
-					if (Target.TargetName != "UnrealHeaderTool" && (bCanUseXGE? CanUseXGE(Target.Platform) : CanUseParallelExecutor(Target.Platform)))
+					else
 					{
-						XGEItem Item = XGEPrepareBuildWithUBT(Target.TargetName, Target.Platform, Target.Config, Target.UprojectPath, bForceMonolithic, bForceNonUnity, bForceDebugInfo, Target.UBTArgs, bForceUnity);
-						if(InTargetToManifest != null)
-						{
-							InTargetToManifest[Target] = Item.Manifest;
-						}
-						XGEItems.Add(Item);
+						NonParallelTargets.Add(Target);
 					}
 				}
-				if (XGEItems.Count > 0)
+				else if(bCanUseParallelExecutor)
 				{
-					string XGETool = null;
-					string Args = null;
-					if (bCanUseXGE) 
+					if(CanUseParallelExecutor(Target.Platform))
 					{
-						XGETool = XGEConsoleExePath;
-						Args = "\"" + TaskFilePath + "\" /Rebuild /MaxCPUS=200";
-						if (ParseParam("StopOnErrors"))
-						{
-							Args += " /StopOnErrors";
-						}
-
-						// A bug in the UCRT can cause XGE to hang on VS2015 builds. Figure out if this hang is likely to effect this build and workaround it if able.
-						string XGEVersion = Microsoft.Win32.Registry.GetValue(@"HKEY_LOCAL_MACHINE\SOFTWARE\Wow6432Node\Xoreax\IncrediBuild\Builder", "Version", null) as string;
-						if (XGEVersion != null)
-						{
-							// Per Xoreax support, subtract 1001000 from the registry value to get the build number of the installed XGE.
-							int XGEBuildNumber;
-							if (Int32.TryParse(XGEVersion, out XGEBuildNumber) && XGEBuildNumber - 1001000 >= 1659)
-							{
-								Args += " /no_watchdog_thread";
-							}
-						}
+						ParallelTargets.Add(Target);
 					}
-
-					if (!bCanUseParallelExecutor && String.IsNullOrEmpty(XGETool))
+					else
 					{
-						throw new AutomationException("BUILD FAILED: no tool present to process XGE files");
-					}
-
-					if (!String.IsNullOrEmpty(XGETool))
-					{
-						CommandUtils.LogSetProgress(InShowProgress, "Building...");
-						if (!ProcessXGEItems(XGEItems, XGETool, Args, TaskFilePath, InShowProgress))
-						{
-							throw new AutomationException("BUILD FAILED: {0} failed, retries not enabled:", XGETool);
-						}
-						else
-						{
-							bUsedXGE = true;
-						}
+						NonParallelTargets.Add(Target);
 					}
 				}
 				else
 				{
-					bUsedXGE = true; // if there was nothing to build, we still consider this a success
+					NonParallelTargets.Add(Target);
 				}
 			}
 
-			if (!bUsedXGE)
+			// Execute all the serial targets
+			foreach(BuildTarget Target in NonParallelTargets)
 			{
-				foreach (var Target in Agenda.Targets)
+				// When building a target for Mac or iOS, use UBT's -flushmac option to clean up the remote builder
+				bool bForceFlushMac = DeleteBuildProducts && (Target.Platform == UnrealTargetPlatform.Mac || Target.Platform == UnrealTargetPlatform.IOS);
+				BuildManifest Manifest = BuildWithUBT(Target.TargetName, Target.Platform, Target.Config, Target.UprojectPath, bForceMonolithic, bForceNonUnity, bForceDebugInfo, bForceFlushMac, true, Target.UBTArgs, bForceUnity);
+				if(InTargetToManifest != null)
 				{
-					bool bClean = Target.Clean ?? DeleteBuildProducts;
-					if (bClean)
+					InTargetToManifest[Target] = Manifest;
+				}
+			}
+
+			// Execute all the parallel targets
+			if(ParallelTargets.Count > 0)
+			{
+				string TaskFilePath = CommandUtils.CombinePaths(CommandUtils.CmdEnv.LogFolder, @"UAT_XGE.xml");
+
+				CommandUtils.LogSetProgress(InShowProgress, "Generating headers...");
+
+				List<XGEItem> XGEItems = new List<XGEItem>();
+				foreach (BuildTarget Target in ParallelTargets)
+				{
+					XGEItem Item = XGEPrepareBuildWithUBT(Target.TargetName, Target.Platform, Target.Config, Target.UprojectPath, bForceMonolithic, bForceNonUnity, bForceDebugInfo, Target.UBTArgs, bForceUnity);
+					if(InTargetToManifest != null)
 					{
-						CleanWithUBT(Target.TargetName, Target.Platform, Target.Config, Target.UprojectPath, bForceMonolithic, bForceNonUnity, bForceDebugInfo, Target.UBTArgs, bForceUnity);
+						InTargetToManifest[Target] = Item.Manifest;
 					}
+					XGEItems.Add(Item);
 				}
-				foreach (var Target in Agenda.Targets)
+
+				string XGETool = null;
+				string Args = null;
+				if (bCanUseXGE) 
 				{
-					// When building a target for Mac or iOS, use UBT's -flushmac option to clean up the remote builder
-					bool bForceFlushMac = DeleteBuildProducts && (Target.Platform == UnrealBuildTool.UnrealTargetPlatform.Mac || Target.Platform == UnrealBuildTool.UnrealTargetPlatform.IOS);
-					BuildManifest Manifest = BuildWithUBT(Target.TargetName, Target.Platform, Target.Config, Target.UprojectPath, bForceMonolithic, bForceNonUnity, bForceDebugInfo, bForceFlushMac, true, Target.UBTArgs, bForceUnity);
-					if(InTargetToManifest != null)
+					XGETool = XGEConsoleExePath;
+					Args = "\"" + TaskFilePath + "\" /Rebuild /MaxCPUS=200";
+					if (ParseParam("StopOnErrors"))
+					{
+						Args += " /StopOnErrors";
+					}
+
+					// A bug in the UCRT can cause XGE to hang on VS2015 builds. Figure out if this hang is likely to effect this build and workaround it if able.
+					string XGEVersion = Microsoft.Win32.Registry.GetValue(@"HKEY_LOCAL_MACHINE\SOFTWARE\Wow6432Node\Xoreax\IncrediBuild\Builder", "Version", null) as string;
+					if (XGEVersion != null)
 					{
-						InTargetToManifest[Target] = Manifest;
+						// Per Xoreax support, subtract 1001000 from the registry value to get the build number of the installed XGE.
+						int XGEBuildNumber;
+						if (Int32.TryParse(XGEVersion, out XGEBuildNumber) && XGEBuildNumber - 1001000 >= 1659)
+						{
+							Args += " /no_watchdog_thread";
+						}
 					}
 				}
+
+				CommandUtils.LogSetProgress(InShowProgress, "Building...");
+				if (!ProcessXGEItems(XGEItems, XGETool, Args, TaskFilePath, InShowProgress))
+				{
+					throw new AutomationException("BUILD FAILED: {0} failed, retries not enabled:", XGETool);
+				}
 			}
 
 			// NOTE: OK, we're done building projects.  You can build additional targets after this.  Call FinalizebuildProducts() when done.
