fix for IOS launch on deployment problem
reverting out the change made in Seal which was causing issues

diff --git a/Engine/Source/Programs/AutomationTool/IOS/IOSPlatform.Automation.cs b/Engine/Source/Programs/AutomationTool/IOS/IOSPlatform.Automation.cs
--- a/Engine/Source/Programs/AutomationTool/IOS/IOSPlatform.Automation.cs
+++ b/Engine/Source/Programs/AutomationTool/IOS/IOSPlatform.Automation.cs
@@ -1237,15 +1237,8 @@ public class IOSPlatform : Platform
 		}
 		if (Params.IterativeDeploy)
 		{
-			string SrcSanitizedDeltaFileName = SC.GetUFSDeploymentDeltaPath(Params.Devices.Count == 0 ? "" : Params.DeviceNames[0]);
-
-			string NoPathDstSanitizedDeltaFileName = DeploymentContext.UFSDeployDeltaFileName + SC.GetSanitizedDeviceName((Params.Devices.Count == 0 ? "" : Params.DeviceNames[0]));
-			string DstLocatedDeltaFileName = CombinePaths(Params.BaseStageDirectory, PlatformName, NoPathDstSanitizedDeltaFileName);
-
-			InternalUtils.SafeCopyFile(SrcSanitizedDeltaFileName, DstLocatedDeltaFileName, true);
-
-			// push over the changed files
-			RunAndLog(CmdEnv, DeployServer, "Deploy -manifest \"" + DstLocatedDeltaFileName + "\"" + (Params.Devices.Count == 0 ? "" : " -device " + Params.DeviceNames[0]) + AdditionalCommandline + " -bundle " + BundleIdentifier);
+            // push over the changed files
+            RunAndLog(CmdEnv, DeployServer, "Deploy -manifest \"" + CombinePaths(Params.BaseStageDirectory, PlatformName, DeploymentContext.UFSDeployDeltaFileName + (Params.Devices.Count == 0 ? "" : Params.DeviceNames[0])) + "\"" + (Params.Devices.Count == 0 ? "" : " -device " + Params.DeviceNames[0]) + AdditionalCommandline + " -bundle " + BundleIdentifier);
 		}
 		Directory.SetCurrentDirectory (CurrentDir);
         PrintRunTime();
