fix for PC deploy for IOS

diff --git a/Engine/Source/Programs/AutomationTool/IOS/IOSPlatform.Automation.cs b/Engine/Source/Programs/AutomationTool/IOS/IOSPlatform.Automation.cs
--- a/Engine/Source/Programs/AutomationTool/IOS/IOSPlatform.Automation.cs
+++ b/Engine/Source/Programs/AutomationTool/IOS/IOSPlatform.Automation.cs
@@ -1237,8 +1237,9 @@ public class IOSPlatform : Platform
 		}
 		if (Params.IterativeDeploy)
 		{
-            // push over the changed files
-            RunAndLog(CmdEnv, DeployServer, "Deploy -manifest \"" + CombinePaths(Params.BaseStageDirectory, PlatformName, DeploymentContext.UFSDeployDeltaFileName + (Params.Devices.Count == 0 ? "" : Params.DeviceNames[0])) + "\"" + (Params.Devices.Count == 0 ? "" : " -device " + Params.DeviceNames[0]) + AdditionalCommandline + " -bundle " + BundleIdentifier);
+			// push over the changed files
+			string SanitizedDeviceName = SC.GetSanitizedDeviceName((Params.Devices.Count == 0 ? "" : Params.DeviceNames[0]));
+            RunAndLog(CmdEnv, DeployServer, "Deploy -manifest \"" + CombinePaths(Params.BaseStageDirectory, PlatformName, DeploymentContext.UFSDeployDeltaFileName + SanitizedDeviceName) + "\"" + (Params.Devices.Count == 0 ? "" : " -device " + Params.DeviceNames[0]) + AdditionalCommandline + " -bundle " + BundleIdentifier);
 		}
 		Directory.SetCurrentDirectory (CurrentDir);
         PrintRunTime();
