Linux: Dont assume DISPlAY=:0

diff --git a/Engine/Source/Programs/AutomationTool/Linux/LinuxPlatform.Automation.cs b/Engine/Source/Programs/AutomationTool/Linux/LinuxPlatform.Automation.cs
--- a/Engine/Source/Programs/AutomationTool/Linux/LinuxPlatform.Automation.cs
+++ b/Engine/Source/Programs/AutomationTool/Linux/LinuxPlatform.Automation.cs
@@ -209,8 +209,21 @@ public abstract class BaseLinuxPlatform : Platform
 				}
 
 				// stage a shell script that makes running easy
+				// Starting from 0, finds the first valid X11 connection xset is able to query
+				// if it fails to find one it'll default to 0 and fail to connect.
 				string Script = String.Format(@"#!/bin/bash
-export DISPLAY=:0
+VALID_DISPLAY=0
+
+for i in `seq 0 16`; do
+    DISPLAY=:$i xset -q > /dev/null 2>&1
+    if [ $? -eq 0 ]
+    then
+        VALID_DISPLAY=$i
+        break
+    fi
+done
+
+export DISPLAY=:$VALID_DISPLAY
 chmod +x {0}
 {0} $@
 ", BinaryName, BinaryName);
