UAT: Add a script that checks that the following macro pairs are correctly balanced within every source file in the engine:

* PRAGMA_DISABLE_OPTIMIZATION and PRAGMA_ENABLE_OPTIMIZATION
* PRAGMA_DISABLE_DEPRECATION_WARNINGS and PRAGMA_ENABLE_DEPRECATION_WARNINGS
* THIRD_PARTY_INCLUDES_START and THIRD_PARTY_INCLUDES_END
* PRAGMA_DISABLE_SHADOW_VARIABLE_WARNINGS and PRAGMA_ENABLE_SHADOW_VARIABLE_WARNINGS
* PRAGMA_DISABLE_UNDEFINED_IDENTIFIER_WARNINGS and PRAGMA_ENABLE_UNDEFINED_IDENTIFIER_WARNINGS
* PRAGMA_DISABLE_MISSING_VIRTUAL_DESTRUCTOR_WARNINGS and PRAGMA_ENABLE_MISSING_VIRTUAL_DESTRUCTOR_WARNINGS"
* BEGIN_FUNCTION_BUILD_OPTIMIZATION and END_FUNCTION_BUILD_OPTIMIZATION
* BEGIN_FUNCTION_BUILD_OPTIMIZATION and END_FUNCTION_BUILD_OPTIMIZATION
* BEGIN_SLATE_FUNCTION_BUILD_OPTIMIZATION and END_SLATE_FUNCTION_BUILD_OPTIMIZATION

Additional projects to scan may be specified with the -Project=<FileName> argument.

diff --git a/Engine/Source/Programs/AutomationTool/Scripts/AutomationScripts.Automation.csproj b/Engine/Source/Programs/AutomationTool/Scripts/AutomationScripts.Automation.csproj
--- a/Engine/Source/Programs/AutomationTool/Scripts/AutomationScripts.Automation.csproj
+++ b/Engine/Source/Programs/AutomationTool/Scripts/AutomationScripts.Automation.csproj
@@ -69,6 +69,7 @@
     <Compile Include="BuildPluginCommand.Automation.cs" />
     <Compile Include="BuildProjectCommand.Automation.cs" />
     <Compile Include="BuildThirdPartyLibs.Automation.cs" />
+    <Compile Include="CheckBalancedMacros.cs" />
     <Compile Include="CheckCsprojDotNetVersion.cs" />
     <Compile Include="CheckForHacks.cs" />
     <Compile Include="CheckRestrictedFolders.cs" />
diff --git a/Engine/Source/Programs/AutomationTool/Scripts/CheckBalancedMacros.cs b/Engine/Source/Programs/AutomationTool/Scripts/CheckBalancedMacros.cs
new file mode 100644
--- /dev/null
+++ b/Engine/Source/Programs/AutomationTool/Scripts/CheckBalancedMacros.cs
@@ -0,0 +1,339 @@
+// Copyright 1998-2018 Epic Games, Inc. All Rights Reserved.
+
+using System;
+using System.Collections.Generic;
+using System.IO;
+using System.Linq;
+using System.Text;
+using System.Threading.Tasks;
+using Tools.DotNETCommon;
+using UnrealBuildTool;
+
+namespace AutomationTool
+{
+	[Help("Checks that all source files have balanced macros for enabling/disabling optimization, warnings, etc...")]
+	[Help("Project=<Path>", "Path to an additional project file to consider")]
+	[Help("File=<Path>", "Path to a file to parse in isolation, for testing")]
+	[Help("Ignore=<Name>", "File name (without path) to exclude from testing")]
+	class CheckBalancedMacros : BuildCommand
+	{
+		/// <summary>
+		/// List of macros that should be paired up
+		/// </summary>
+		static readonly string[,] MacroPairs = new string[,]
+		{
+			{
+				"PRAGMA_DISABLE_OPTIMIZATION",
+				"PRAGMA_ENABLE_OPTIMIZATION"
+			},
+			{
+				"PRAGMA_DISABLE_DEPRECATION_WARNINGS",
+				"PRAGMA_ENABLE_DEPRECATION_WARNINGS"
+			},
+			{
+				"THIRD_PARTY_INCLUDES_START",
+				"THIRD_PARTY_INCLUDES_END"
+			},
+			{
+				"PRAGMA_DISABLE_SHADOW_VARIABLE_WARNINGS",
+				"PRAGMA_ENABLE_SHADOW_VARIABLE_WARNINGS"
+			},
+			{
+				"PRAGMA_DISABLE_UNDEFINED_IDENTIFIER_WARNINGS",
+				"PRAGMA_ENABLE_UNDEFINED_IDENTIFIER_WARNINGS"
+			},
+			{
+				"PRAGMA_DISABLE_MISSING_VIRTUAL_DESTRUCTOR_WARNINGS",
+				"PRAGMA_ENABLE_MISSING_VIRTUAL_DESTRUCTOR_WARNINGS"
+			},
+			{
+				"BEGIN_FUNCTION_BUILD_OPTIMIZATION",
+				"END_FUNCTION_BUILD_OPTIMIZATION"
+			},
+			{
+				"BEGIN_FUNCTION_BUILD_OPTIMIZATION",
+				"END_FUNCTION_BUILD_OPTIMIZATION"
+			},
+			{
+				"BEGIN_SLATE_FUNCTION_BUILD_OPTIMIZATION",
+				"END_SLATE_FUNCTION_BUILD_OPTIMIZATION"
+			},
+		};
+
+		/// <summary>
+		/// List of files to ignore for balanced macros. Additional filenames may be specified on the command line via -Ignore=...
+		/// </summary>
+		HashSet<string> IgnoreFileNames = new HashSet<string>(StringComparer.OrdinalIgnoreCase)
+		{
+			"PreWindowsApi.h",
+			"PostWindowsApi.h",
+		};
+
+		/// <summary>
+		/// Main entry point for the command
+		/// </summary>
+		public override void ExecuteBuild()
+		{
+			// Build a lookup of flags to set and clear for each identifier
+			Dictionary<string, int> IdentifierToIndex = new Dictionary<string, int>();
+			for(int Idx = 0; Idx < MacroPairs.GetLength(0); Idx++)
+			{
+				IdentifierToIndex[MacroPairs[Idx, 0]] = Idx;
+				IdentifierToIndex[MacroPairs[Idx, 1]] = ~Idx;
+			}
+
+			// Check if we want to just parse a single file
+			string FileParam = ParseParamValue("File");
+			if(FileParam != null)
+			{
+				// Check the file exists
+				FileReference File = new FileReference(FileParam);
+				if(!FileReference.Exists(File))
+				{
+					throw new AutomationException("File '{0}' does not exist", File);
+				}
+				CheckSourceFile(File, IdentifierToIndex, new object());
+			}
+			else
+			{
+				// Add the additional files to be ignored
+				foreach(string IgnoreFileName in ParseParamValues("Ignore"))
+				{
+					IgnoreFileNames.Add(IgnoreFileName);
+				}
+
+				// Create a list of all the root directories
+				HashSet<DirectoryReference> RootDirs = new HashSet<DirectoryReference>();
+				RootDirs.Add(EngineDirectory);
+
+				// Add the enterprise directory
+				DirectoryReference EnterpriseDirectory = DirectoryReference.Combine(RootDirectory, "Enterprise");
+				if(DirectoryReference.Exists(EnterpriseDirectory))
+				{
+					RootDirs.Add(EnterpriseDirectory);
+				}
+
+				// Add the project directories
+				string[] ProjectParams = ParseParamValues("Project");
+				foreach(string ProjectParam in ProjectParams)
+				{
+					FileReference ProjectFile = new FileReference(ProjectParam);
+					if(!FileReference.Exists(ProjectFile))
+					{
+						throw new AutomationException("Unable to find project '{0}'", ProjectFile);
+					}
+					RootDirs.Add(ProjectFile.Directory);
+				}
+
+				// Recurse through the tree
+				LogInformation("Finding source files...");
+				List<FileReference> SourceFiles = new List<FileReference>();
+				using(ThreadPoolWorkQueue Queue = new ThreadPoolWorkQueue())
+				{
+					foreach(DirectoryReference RootDir in RootDirs)
+					{
+						DirectoryInfo PluginsDir = new DirectoryInfo(Path.Combine(RootDir.FullName, "Plugins"));
+						if(PluginsDir.Exists)
+						{
+							Queue.Enqueue(() => FindSourceFiles(PluginsDir, SourceFiles, Queue));
+						}
+
+						DirectoryInfo SourceDir = new DirectoryInfo(Path.Combine(RootDir.FullName, "Source"));
+						if(SourceDir.Exists)
+						{
+							Queue.Enqueue(() => FindSourceFiles(SourceDir, SourceFiles, Queue));
+						}
+					}
+					Queue.Wait();
+				}
+
+				// Loop through all the source files
+				using(ThreadPoolWorkQueue Queue = new ThreadPoolWorkQueue())
+				{
+					object LogLock = new object();
+					foreach(FileReference SourceFile in SourceFiles)
+					{
+						Queue.Enqueue(() => CheckSourceFile(SourceFile, IdentifierToIndex, LogLock));
+					}
+
+					using(LogStatusScope Scope = new LogStatusScope("Checking source files..."))
+					{
+						while(!Queue.Wait(10 * 1000))
+						{
+							Scope.SetProgress("{0}/{1}", SourceFiles.Count - Queue.NumRemaining, SourceFiles.Count);
+						}
+					}
+				}
+			}
+		}
+
+		/// <summary>
+		/// Finds all the source files under a given directory
+		/// </summary>
+		/// <param name="BaseDir">Directory to search</param>
+		/// <param name="SourceFiles">List to receive the files found. A lock will be taken on this object to ensure multiple threads do not add to it simultaneously.</param>
+		/// <param name="Queue">Queue for additional tasks to be added to</param>
+		void FindSourceFiles(DirectoryInfo BaseDir, List<FileReference> SourceFiles, ThreadPoolWorkQueue Queue)
+		{
+			foreach(DirectoryInfo SubDir in BaseDir.EnumerateDirectories())
+			{
+				if(!SubDir.Name.Equals("Intermediate", StringComparison.OrdinalIgnoreCase))
+				{
+					Queue.Enqueue(() => FindSourceFiles(SubDir, SourceFiles, Queue));
+				}
+			}
+
+			foreach(FileInfo File in BaseDir.EnumerateFiles())
+			{
+				if(File.Name.EndsWith(".h", StringComparison.OrdinalIgnoreCase) || File.Name.EndsWith(".cpp", StringComparison.OrdinalIgnoreCase))
+				{
+					if(!IgnoreFileNames.Contains(File.Name))
+					{
+						lock(SourceFiles)
+						{
+							SourceFiles.Add(new FileReference(File));
+						}
+					}
+				}
+			}
+		}
+
+		/// <summary>
+		/// Checks whether macros in the given source file are matched
+		/// </summary>
+		/// <param name="SourceFile"></param>
+		/// <param name="IdentifierToIndex">Map of macro identifier to bit index. The complement of an index is used to indicate the end of the pair.</param>
+		/// <param name="LogLock">Object used to marshal access to the global log instance</param>
+		void CheckSourceFile(FileReference SourceFile, Dictionary<string, int> IdentifierToIndex, object LogLock)
+		{
+			// Read the text
+			string Text = FileReference.ReadAllText(SourceFile);
+
+			// Scan through the file token by token. Each bit in the Flags array indicates an index into the MacroPairs array that is currently active.
+			int Flags = 0;
+			for(int Idx = 0; Idx < Text.Length; )
+			{
+				int StartIdx = Idx++;
+				if((Text[StartIdx] >= 'a' && Text[StartIdx] <= 'z') || (Text[StartIdx] >= 'A' && Text[StartIdx] <= 'Z') || Text[StartIdx] == '_')
+				{
+					// Identifier
+					while(Idx < Text.Length && ((Text[Idx] >= 'a' && Text[Idx] <= 'z') || (Text[Idx] >= 'A' && Text[Idx] <= 'Z') || (Text[Idx] >= '0' && Text[Idx] <= '9') || Text[Idx] == '_'))
+					{
+						Idx++;
+					}
+
+					// Extract the identifier
+					string Identifier = Text.Substring(StartIdx, Idx - StartIdx);
+
+					// Find the matching flag
+					int Index;
+					if(IdentifierToIndex.TryGetValue(Identifier, out Index))
+					{
+						if(Index >= 0)
+						{
+							// Set the flag (should not already be set)
+							int Flag = 1 << Index;
+							if((Flags & Flag) != 0)
+							{
+								Tools.DotNETCommon.Log.TraceWarning(SourceFile, GetLineNumber(Text, StartIdx), "{0} macro appears a second time without matching {1} macro", Identifier, MacroPairs[Index, 1]);
+							}
+							Flags |= Flag;
+						}
+						else
+						{
+							// Clear the flag (should already be set)
+							int Flag = 1 << ~Index;
+							if((Flags & Flag) == 0)
+							{
+								Tools.DotNETCommon.Log.TraceWarning(SourceFile, GetLineNumber(Text, StartIdx), "{0} macro appears without matching {1} macro", Identifier, MacroPairs[~Index, 0]);
+							}
+							Flags &= ~Flag;
+						}
+					}
+				}
+				else if(Text[StartIdx] == '/' && Idx < Text.Length)
+				{
+					if(Text[Idx] == '/')
+					{
+						// Single-line comment
+						while(Idx < Text.Length && Text[Idx] != '\n')
+						{
+							Idx++;
+						}
+					}
+					else if(Text[Idx] == '*')
+					{
+						// Multi-line comment
+						Idx++;
+						for(; Idx < Text.Length; Idx++)
+						{
+							if(Idx + 2 < Text.Length && Text[Idx] == '*' && Text[Idx + 1] == '/')
+							{
+								Idx += 2;
+								break;
+							}
+						}
+					}
+				}
+				else if(Text[StartIdx] == '"' || Text[StartIdx] == '\'')
+				{
+					// String
+					for(; Idx < Text.Length; Idx++)
+					{
+						if(Text[Idx] == Text[StartIdx])
+						{
+							Idx++;
+							break;
+						}
+						if(Text[Idx] == '\\')
+						{
+							Idx++;
+						}
+					}
+				}
+				else if(Text[StartIdx] == '#')
+				{
+					// Preprocessor directive (eg. #define)
+					for(; Idx < Text.Length && Text[Idx] != '\n'; Idx++)
+					{
+						if(Text[Idx] == '\\')
+						{
+							Idx++;
+						}
+					}
+				}
+			}
+
+			// Check if there's anything left over
+			if(Flags != 0)
+			{
+				for(int Idx = 0; Idx < MacroPairs.GetLength(0); Idx++)
+				{
+					if((Flags & (1 << Idx)) != 0)
+					{
+						Tools.DotNETCommon.Log.TraceWarning(SourceFile, "{0} macro does not have matching {1} macro", MacroPairs[Idx, 0], MacroPairs[Idx, 1]);
+					}
+				}
+			}
+		}
+
+		/// <summary>
+		/// Converts an offset within a text buffer into a line number
+		/// </summary>
+		/// <param name="Text">Text to search</param>
+		/// <param name="Offset">Offset within the text</param>
+		/// <returns>Line number corresponding to the given offset. Starts from one.</returns>
+		int GetLineNumber(string Text, int Offset)
+		{
+			int LineNumber = 1;
+			for(int Idx = 0; Idx < Offset; Idx++)
+			{
+				if(Text[Idx] == '\n')
+				{
+					LineNumber++;
+				}
+			}
+			return LineNumber;
+		}
+	}
+}
