UAT: Add a command for stashing and unstashing a target. Useful for A/B testing. Example syntax:

Stashing:

RunUAT StashTarget -Name=UE4Editor -To=D:\TempEditor
RunUAT StashTarget -Name=UE4Editor -Platform=Win64 -Configuration=Debug -To=D:\TempEditor
RunUAT StashTarget -Name=ShooterGameEditor -Platform=Win64 -Configuration=Debug -Project=D:\P4\Samples\ShooterGame\ShooterGame.uproject -To=D:\TempEditor

Unstashing:

RunUAT UnstashTarget -From=D:\TempEditor

diff --git a/Engine/Source/Programs/AutomationTool/AutomationUtils/BuildCommand.cs b/Engine/Source/Programs/AutomationTool/AutomationUtils/BuildCommand.cs
--- a/Engine/Source/Programs/AutomationTool/AutomationUtils/BuildCommand.cs
+++ b/Engine/Source/Programs/AutomationTool/AutomationUtils/BuildCommand.cs
@@ -1,8 +1,9 @@
-﻿// Copyright 1998-2018 Epic Games, Inc. All Rights Reserved.
+// Copyright 1998-2018 Epic Games, Inc. All Rights Reserved.
 using System;
 using System.Collections.Generic;
 using System.Linq;
 using System.Text;
+using Tools.DotNETCommon;
 
 namespace AutomationTool
 {
@@ -42,6 +43,135 @@ namespace AutomationTool
 			return ParseParamValue(Params, Param, Default);
 		}
 
+		/// <summary>
+		/// Parses an argument.
+		/// </summary>
+		/// <param name="Param"></param>
+		/// <returns></returns>
+		public string ParseOptionalStringParam(string Param)
+		{
+			return ParseParamValue(Param, null);
+		}
+
+		/// <summary>
+		/// Parses an argument. Throws an exception if the parameter is not specified.
+		/// </summary>
+		/// <param name="Param">Name of the argument</param>
+		/// <returns>Value of the argument</returns>
+		public string ParseRequiredStringParam(string Param)
+		{
+			string Value = ParseOptionalStringParam(Param);
+			if(Value == null)
+			{
+				throw new AutomationException("Missing -{0}=... parameter");
+			}
+			return Value;
+		}
+
+		/// <summary>
+		/// Parses an file reference argument.
+		/// </summary>
+		/// <param name="Param">Name of the argument</param>
+		/// <returns>Value of the argument</returns>
+		public FileReference ParseOptionalFileReferenceParam(string Param)
+		{
+			string StringValue = ParseParamValue(Param);
+			if(StringValue == null)
+			{
+				return null;
+			}
+			else
+			{
+				return new FileReference(StringValue);
+			}
+		}
+
+		/// <summary>
+		/// Parses an file reference argument. Throws an exception if the parameter is not specified.
+		/// </summary>
+		/// <param name="Param">Name of the argument</param>
+		/// <returns>Value of the argument</returns>
+		public FileReference ParseRequiredFileReferenceParam(string Param)
+		{
+			FileReference Value = ParseOptionalFileReferenceParam(Param);
+			if(Value == null)
+			{
+				throw new AutomationException("Missing -{0}=... parameter");
+			}
+			return Value;
+		}
+
+		/// <summary>
+		/// Parses a directory reference argument.
+		/// </summary>
+		/// <param name="Param">Name of the argument</param>
+		/// <returns>Value of the argument</returns>
+		public DirectoryReference ParseOptionalDirectoryReferenceParam(string Param)
+		{
+			string StringValue = ParseOptionalStringParam(Param);
+			if(StringValue == null)
+			{
+				return null;
+			}
+			else
+			{
+				return new DirectoryReference(StringValue);
+			}
+		}
+
+		/// <summary>
+		/// Parses a directory reference argument. Throws an exception if the parameter is not specified.
+		/// </summary>
+		/// <param name="Param">Name of the argument</param>
+		/// <returns>Value of the argument</returns>
+		public DirectoryReference ParseRequiredDirectoryReferenceParam(string Param)
+		{
+			DirectoryReference Value = ParseOptionalDirectoryReferenceParam(Param);
+			if(Value == null)
+			{
+				throw new AutomationException("Missing -{0}=... parameter");
+			}
+			return Value;
+		}
+
+		/// <summary>
+		/// Parses an argument as an enum.
+		/// </summary>
+		/// <param name="Param">Name of the parameter to read.</param>
+		/// <returns>Returns the value that was parsed.
+		public Nullable<T> ParseOptionalEnumParam<T>(string Param) where T : struct
+		{
+			string ValueString = ParseParamValue(Param);
+			if(ValueString == null)
+			{
+				return null;
+			}
+			else
+			{
+				T Value;
+				if(!Enum.TryParse<T>(ValueString, out Value))
+				{
+					throw new AutomationException("'{0}' is not a valid value for {1}", ValueString, typeof(T).Name);
+				}
+				return Value;
+			}
+		}
+
+		/// <summary>
+		/// Parses an argument as an enum. Throws an exception if the parameter is not specified.
+		/// </summary>
+		/// <param name="Param">Name of the parameter to read.</param>
+		/// <returns>Returns the value that was parsed.
+		public T ParseRequiredEnumParamEnum<T>(string Param) where T : struct
+		{
+			Nullable<T> Value = ParseOptionalEnumParam<T>(Param);
+			if(!Value.HasValue)
+			{
+				throw new AutomationException("Missing -{0}=... parameter");
+			}
+			return Value.Value;
+		}
+
 		/// <summary>
 		/// Parses the argument list for any number of parameters.
 		/// </summary>
diff --git a/Engine/Source/Programs/AutomationTool/Scripts/AutomationScripts.Automation.csproj b/Engine/Source/Programs/AutomationTool/Scripts/AutomationScripts.Automation.csproj
--- a/Engine/Source/Programs/AutomationTool/Scripts/AutomationScripts.Automation.csproj
+++ b/Engine/Source/Programs/AutomationTool/Scripts/AutomationScripts.Automation.csproj
@@ -60,6 +60,7 @@
     </Compile>
     <Compile Include="AnalyzeThirdPartyLibs.Automation.cs" />
     <Compile Include="ArchiveCommand.Automation.cs" />
+    <Compile Include="StashTarget.cs" />
     <Compile Include="BlameKeyword.Automation.cs" />
     <Compile Include="BuildCommonTools.Automation.cs" />
     <Compile Include="BuildCookRun.Automation.cs" />
diff --git a/Engine/Source/Programs/AutomationTool/Scripts/StashTarget.cs b/Engine/Source/Programs/AutomationTool/Scripts/StashTarget.cs
new file mode 100644
--- /dev/null
+++ b/Engine/Source/Programs/AutomationTool/Scripts/StashTarget.cs
@@ -0,0 +1,79 @@
+using System;
+using System.Collections.Generic;
+using System.IO;
+using System.Linq;
+using System.Text;
+using System.Threading.Tasks;
+using Tools.DotNETCommon;
+using UnrealBuildTool;
+
+namespace AutomationTool
+{
+	[Help("Copy all the binaries for a target into a different folder. Can be restored using the UnstashTarget command. Useful for A/B testing.")]
+	[Help("-Name", "Name of the target")]
+	[Help("-Platform", "Platform that the target was built for")]
+	[Help("-Configuration", "Architecture that the target was built for")]
+	[Help("-Architecture", "Architecture that the target was built for")]
+	[Help("-Project", "Project file for the target")]
+	[Help("-To", "Output directory to store the stashed binaries")]
+	public class StashTarget : BuildCommand
+	{
+		public override void ExecuteBuild()
+		{
+			// Parse all the arguments
+			string TargetName = ParseRequiredStringParam("Name");
+			UnrealTargetPlatform Platform = ParseOptionalEnumParam<UnrealTargetPlatform>("Platform") ?? HostPlatform.Current.HostEditorPlatform;
+			UnrealTargetConfiguration Configuration = ParseOptionalEnumParam<UnrealTargetConfiguration>("Configuration") ?? UnrealTargetConfiguration.Development;
+			string Architecture = ParseOptionalStringParam("Architecture");
+			FileReference ProjectFile = ParseOptionalFileReferenceParam("Project");
+			DirectoryReference ToDir = ParseRequiredDirectoryReferenceParam("To");
+
+			// Read the receipt
+			FileReference ReceiptFile = TargetReceipt.GetDefaultPath(DirectoryReference.FromFile(ProjectFile) ?? EngineDirectory, TargetName, Platform, Configuration, Architecture);
+			if(!FileReference.Exists(ReceiptFile))
+			{
+				throw new AutomationException("Unable to find '{0}'", ReceiptFile);
+			}
+
+			TargetReceipt Receipt = TargetReceipt.Read(ReceiptFile, EngineDirectory, DirectoryReference.FromFile(ProjectFile));
+
+			// Enumerate all the files we want to move
+			List<FileReference> FilesToMove = new List<FileReference>();
+			FilesToMove.Add(ReceiptFile);
+			FilesToMove.AddRange(Receipt.BuildProducts.Select(x => x.Path));
+
+			// Move all the files to the output folder
+			DirectoryReference.CreateDirectory(ToDir);
+			CommandUtils.DeleteDirectoryContents(ToDir.FullName);
+			foreach(FileReference SourceFile in FilesToMove)
+			{
+				FileReference TargetFile = FileReference.Combine(ToDir, SourceFile.MakeRelativeTo(RootDirectory));
+				LogInformation("Copying {0} to {1}", SourceFile, TargetFile);
+				CommandUtils.CopyFile(SourceFile.FullName, TargetFile.FullName);
+			}
+		}
+	}
+
+	[Help("Copy all the binaries from a target back into the root directory. Use in combination with the StashTarget command.")]
+	[Help("-From", "Directory to copy from")]
+	public class UnstashTarget : BuildCommand
+	{
+		public override void ExecuteBuild()
+		{
+			// Parse the arguments
+			DirectoryReference FromDir = ParseRequiredDirectoryReferenceParam("From");
+			if(!DirectoryReference.Exists(FromDir))
+			{
+				throw new AutomationException("Source directory '{0}' does not exist", FromDir);
+			}
+
+			// Just copy all the files into place
+			foreach(FileReference SourceFile in DirectoryReference.EnumerateFiles(FromDir, "*", SearchOption.AllDirectories))
+			{
+				FileReference TargetFile = FileReference.Combine(RootDirectory, SourceFile.MakeRelativeTo(FromDir));
+				LogInformation("Copying {0} to {1}", SourceFile, TargetFile);
+				CopyFile(SourceFile.FullName, TargetFile.FullName);
+			}
+		}
+	}
+}
