UAT: Show timing/agent info in UAT the same way as in UBT, to make it easier to track down badly configured agents.

diff --git a/Engine/Source/Programs/AutomationTool/AutomationUtils/UE4Build.cs b/Engine/Source/Programs/AutomationTool/AutomationUtils/UE4Build.cs
--- a/Engine/Source/Programs/AutomationTool/AutomationUtils/UE4Build.cs
+++ b/Engine/Source/Programs/AutomationTool/AutomationUtils/UE4Build.cs
@@ -1145,7 +1145,8 @@ namespace AutomationTool
 			bool bForceUnity = ParseParam("ForceUnity") || InForceUnity;
 			bool bForceDebugInfo = ParseParam("ForceDebugInfo");
 			string XGEConsole = XGEConsoleExe();
-			bool bCanUseXGE = !ParseParam("NoXGE") && !InForceNoXGE && !String.IsNullOrEmpty(XGEConsole);
+			bool bDisableXGE = ParseParam("NoXGE") || InForceNoXGE;
+			bool bCanUseXGE = !bDisableXGE && !String.IsNullOrEmpty(XGEConsole);
 
 			// only run ParallelExecutor if not running XGE (and we've requested ParallelExecutor and it exists)
 			bool bCanUseParallelExecutor = !bCanUseXGE && InUseParallelExecutor && (HostPlatform.Current.HostEditorPlatform == UnrealTargetPlatform.Win64);
@@ -1208,7 +1209,7 @@ namespace AutomationTool
 			{
 				// When building a target for Mac or iOS, use UBT's -flushmac option to clean up the remote builder
 				bool bForceFlushMac = DeleteBuildProducts && (Target.Platform == UnrealTargetPlatform.Mac || Target.Platform == UnrealTargetPlatform.IOS);
-				BuildManifest Manifest = BuildWithUBT(Target.TargetName, Target.Platform, Target.Config, Target.UprojectPath, bForceMonolithic, bForceNonUnity, bForceDebugInfo, bForceFlushMac, true, Target.UBTArgs, bForceUnity);
+				BuildManifest Manifest = BuildWithUBT(Target.TargetName, Target.Platform, Target.Config, Target.UprojectPath, bForceMonolithic, bForceNonUnity, bForceDebugInfo, bForceFlushMac, bDisableXGE, Target.UBTArgs, bForceUnity);
 				if(InTargetToManifest != null)
 				{
 					InTargetToManifest[Target] = Manifest;
@@ -1238,7 +1239,7 @@ namespace AutomationTool
 				if (bCanUseXGE) 
 				{
 					XGETool = XGEConsoleExePath;
-					Args = "\"" + TaskFilePath + "\" /Rebuild /MaxCPUS=200";
+					Args = "\"" + TaskFilePath + "\" /Rebuild /NoLogo /ShowAgent /ShowTime";
 					if (ParseParam("StopOnErrors"))
 					{
 						Args += " /StopOnErrors";
