merging over pakfile fix from dev-core

diff --git a/Engine/Source/Programs/AutomationTool/Scripts/CopyBuildToStagingDirectory.Automation.cs b/Engine/Source/Programs/AutomationTool/Scripts/CopyBuildToStagingDirectory.Automation.cs
--- a/Engine/Source/Programs/AutomationTool/Scripts/CopyBuildToStagingDirectory.Automation.cs
+++ b/Engine/Source/Programs/AutomationTool/Scripts/CopyBuildToStagingDirectory.Automation.cs
@@ -1479,10 +1479,11 @@ public partial class Project : CommandUtils
 			string[] OrderFileNames = new string[] { "GameOpenOrder.log", "CookerOpenOrder.log", "EditorOpenOrder.log" };
 			foreach (string OrderFileName in OrderFileNames)
 			{
-				PakOrderFileLocation = FileReference.Combine(PakOrderFileLocationBase, OrderFileName);
+				FileReference FileLocation = FileReference.Combine(PakOrderFileLocationBase, OrderFileName);
 
-				if (FileExists_NoExceptions(PakOrderFileLocation.FullName))
+				if (FileExists_NoExceptions(FileLocation.FullName))
 				{
+					PakOrderFileLocation = FileLocation;
 					break;
 				}
 			}
