Added logging if XLOC fails to find a PO file to use, and now compare the PO file names case-insensitively

diff --git a/Engine/Source/Programs/AutomationTool/XLocLocalization/XLocLocalizationProvider.cs b/Engine/Source/Programs/AutomationTool/XLocLocalization/XLocLocalizationProvider.cs
--- a/Engine/Source/Programs/AutomationTool/XLocLocalization/XLocLocalizationProvider.cs
+++ b/Engine/Source/Programs/AutomationTool/XLocLocalization/XLocLocalizationProvider.cs
@@ -126,7 +126,7 @@ public abstract class XLocLocalizationProvider : LocalizationProvider
 					if (GameFileElem != null)
 					{
 						var GameFileNameElem = GameFileElem["Name"];
-						if (GameFileNameElem != null && GameFileNameElem.InnerText == XLocFilename)
+						if (GameFileNameElem != null && GameFileNameElem.InnerText.Equals(XLocFilename, StringComparison.OrdinalIgnoreCase))
 						{
 							IsCorrectFile = true;
 						}
@@ -145,7 +145,11 @@ public abstract class XLocLocalizationProvider : LocalizationProvider
 			}
 		}
 
-		if (!String.IsNullOrEmpty(POFileUri))
+		if (String.IsNullOrEmpty(POFileUri))
+		{
+			Console.WriteLine("[IGNORED] '{0}' was not found in the build data ({1})", XLocFilename, Culture);
+		}
+		else
 		{
 			var DestinationDirectory = new DirectoryInfo(CommandUtils.CombinePaths(RootWorkingDirectory, ProjectImportInfo.DestinationPath));
 			var CultureDirectory = (ProjectImportInfo.bUseCultureDirectory) ? new DirectoryInfo(Path.Combine(DestinationDirectory.FullName, Culture)) : DestinationDirectory;
