HTML5 - fix merge stomping and add in missing SC.Metadatadir (null) object to last {if else} block

#rnx

diff --git a/Engine/Source/Programs/AutomationTool/Scripts/CopyBuildToStagingDirectory.Automation.cs b/Engine/Source/Programs/AutomationTool/Scripts/CopyBuildToStagingDirectory.Automation.cs
--- a/Engine/Source/Programs/AutomationTool/Scripts/CopyBuildToStagingDirectory.Automation.cs
+++ b/Engine/Source/Programs/AutomationTool/Scripts/CopyBuildToStagingDirectory.Automation.cs
@@ -433,7 +433,6 @@ public partial class Project : CommandUtils
 			if (!Params.CookOnTheFly && !Params.SkipCookOnTheFly) // only stage the UFS files if we are not using cook on the fly
 			{
 
-
 				// Initialize internationalization preset.
 				string InternationalizationPreset = Params.InternationalizationPreset;
 
@@ -621,7 +620,6 @@ public partial class Project : CommandUtils
 				{
 					SC.StageCrashReporterFiles(StagedFileType.UFS, DirectoryReference.Combine(SC.EngineRoot, "Content", "Slate"), StageFilesSearch.AllDirectories);
 					SC.StageCrashReporterFiles(StagedFileType.UFS, DirectoryReference.Combine(SC.EngineRoot, "Shaders", "StandaloneRenderer"), StageFilesSearch.AllDirectories);
-
 					SC.StageCrashReporterFiles(StagedFileType.UFS, DirectoryReference.Combine(SC.EngineRoot, "Content", "Internationalization", InternationalizationPreset), StageFilesSearch.AllDirectories, new StagedDirectoryReference("Engine/Content/Internationalization"));
 
 					// Get the architecture in use
@@ -658,7 +656,6 @@ public partial class Project : CommandUtils
 					// Add config files.
 					SC.StageCrashReporterFiles(StagedFileType.UFS, DirectoryReference.Combine(SC.EngineRoot, "Programs", "CrashReportClient", "Config"), StageFilesSearch.AllDirectories);
 				}
-
 				// check if the game will be verifying ssl connections - if not, we can skip staging files that won't be needed
 				bool bStageSSLCertificates = false;
 				ConfigHierarchy PlatformEngineConfig = ConfigCache.ReadHierarchy(ConfigHierarchyType.Engine, DirectoryReference.FromFile(Params.RawProjectPath), SC.StageTargetPlatform.IniPlatformType);
@@ -730,6 +727,22 @@ public partial class Project : CommandUtils
 						}
 					}
 				}
+
+				// UE-58423
+				DirectoryReference CookOutputDir;
+				if (!String.IsNullOrEmpty(Params.CookOutputDir))
+				{
+					CookOutputDir = DirectoryReference.Combine(new DirectoryReference(Params.CookOutputDir), SC.CookPlatform);
+				}
+				else if (Params.CookInEditor)
+				{
+					CookOutputDir = DirectoryReference.Combine(SC.ProjectRoot, "Saved", "EditorCooked", SC.CookPlatform);
+				}
+				else
+				{
+					CookOutputDir = DirectoryReference.Combine(SC.ProjectRoot, "Saved", "Cooked", SC.CookPlatform);
+				}
+				SC.MetadataDir = DirectoryReference.Combine(CookOutputDir, SC.ShortProjectName, "Metadata");
 			}
 		}
 
@@ -2125,8 +2138,8 @@ public partial class Project : CommandUtils
 				CommandLine += " -iterative";
 			}
 			File.WriteAllText(IntermediateCmdLineFile.FullName, CommandLine);
-		}
-		else
+        }
+        else
 		{
 			String ProjectFile = String.Format("{0} ", SC.ProjectArgForCommandLines);
 			if (SC.StageTargetPlatform.PlatformType == UnrealTargetPlatform.Mac || SC.StageTargetPlatform.PlatformType == UnrealTargetPlatform.Win64 || SC.StageTargetPlatform.PlatformType == UnrealTargetPlatform.Win32 || SC.StageTargetPlatform.PlatformType == UnrealTargetPlatform.Linux)
@@ -2135,8 +2148,8 @@ public partial class Project : CommandUtils
 			}
 			DirectoryReference.CreateDirectory(GetIntermediateCommandlineDir(SC));
 			File.WriteAllText(IntermediateCmdLineFile.FullName, ProjectFile);
-		}
-	}
+        }
+    }
 
 	private static void WriteStageCommandline(ProjectParams Params, DeploymentContext SC)
 	{
@@ -2320,7 +2333,7 @@ public partial class Project : CommandUtils
 
 			string PlatformName = StagePlatform.ToString();
 			string StageArchitecture = !String.IsNullOrEmpty(Params.SpecifiedArchitecture) ? Params.SpecifiedArchitecture : "";
-			foreach (var Target in ListToProcess)
+            foreach (var Target in ListToProcess)
 			{
 				foreach (var Config in ConfigsToProcess)
 				{
@@ -2330,12 +2343,12 @@ public partial class Project : CommandUtils
 						Exe = Target + "-" + PlatformName + "-" + Config.ToString() + StageArchitecture;
 					}
 					ExecutablesToStage.Add(Exe);
-				}
+                }
 			}
 
 			string StageDirectory = ((ShouldCreatePak(Params) || (Params.Stage)) || !String.IsNullOrEmpty(Params.StageDirectoryParam)) ? Params.BaseStageDirectory : "";
 			string ArchiveDirectory = (Params.Archive || !String.IsNullOrEmpty(Params.ArchiveDirectoryParam)) ? Params.BaseArchiveDirectory : "";
-			DirectoryReference EngineDir = DirectoryReference.Combine(CommandUtils.RootDirectory, "Engine");
+            DirectoryReference EngineDir = DirectoryReference.Combine(CommandUtils.RootDirectory, "Engine");
 			DirectoryReference ProjectDir = DirectoryReference.FromFile(Params.RawProjectPath);
 
 			List<StageTarget> TargetsToStage = new List<StageTarget>();
