UE-56076 Android Launch On "Running..." toast intermittently doesn't appear
#Android
#4.21

wait before getting the process list with "adb shell ps" - on some devices the list is not yet ready

diff --git a/Engine/Source/Programs/AutomationTool/Android/AndroidPlatform.Automation.cs b/Engine/Source/Programs/AutomationTool/Android/AndroidPlatform.Automation.cs
--- a/Engine/Source/Programs/AutomationTool/Android/AndroidPlatform.Automation.cs
+++ b/Engine/Source/Programs/AutomationTool/Android/AndroidPlatform.Automation.cs
@@ -1639,6 +1639,10 @@ public class AndroidPlatform : Platform
 		DateTime StartTime = DateTime.Now;
 		int TimeOutSeconds = Params.RunTimeoutSeconds;
 
+		// wait before getting the process list with "adb shell ps"
+		// on some devices the list is not yet ready
+		Thread.Sleep(2000);
+
 		while (DeviceNames.Count > 0)
 		{
 			for(int DeviceIndex = 0; DeviceIndex < DeviceNames.Count; DeviceIndex++)
