Fix warning seen by CIS

diff --git a/Engine/Source/Programs/AutomationTool/HTML5/HTML5Platform.Automation.cs b/Engine/Source/Programs/AutomationTool/HTML5/HTML5Platform.Automation.cs
--- a/Engine/Source/Programs/AutomationTool/HTML5/HTML5Platform.Automation.cs
+++ b/Engine/Source/Programs/AutomationTool/HTML5/HTML5Platform.Automation.cs
@@ -118,7 +118,7 @@ public class HTML5Platform : Platform
 		// ensure the ue4game binary exists, if applicable
 		if (!FileExists_NoExceptions(SrcUE4GameBasename + ".js"))
 		{
-			Log("Failed to find game application " + SrcUE4GameBasename + ".js");
+			LogInformation("Failed to find game application " + SrcUE4GameBasename + ".js");
 			throw new AutomationException(ExitCode.Error_MissingExecutable, "Stage Failed. Could not find application {0}. You may need to build the UE4 project with your target configuration and platform.", SrcUE4GameBasename + ".js");
 		}
 
