UE-63224 Android DLC profile in Project Launcher fails if there is a space in the path for the destination directory

#4.21
#Android

Avoid enclosing path arguments in double quotes (""<path>""). MakePathSafeToUseWithCommandLine(<path>) already adds quotes when needed.

diff --git a/Engine/Source/Programs/AutomationTool/Scripts/CopyBuildToStagingDirectory.Automation.cs b/Engine/Source/Programs/AutomationTool/Scripts/CopyBuildToStagingDirectory.Automation.cs
--- a/Engine/Source/Programs/AutomationTool/Scripts/CopyBuildToStagingDirectory.Automation.cs
+++ b/Engine/Source/Programs/AutomationTool/Scripts/CopyBuildToStagingDirectory.Automation.cs
@@ -1645,7 +1645,7 @@ public partial class Project : CommandUtils
 				string DestManifestPath = CombinePaths(ManifestDir, ManifestFilename);
 				InternalUtils.SafeCreateDirectory(ManifestDir, true);
 
-				string CmdLine = String.Format("-BuildRoot=\"{0}\" -CloudDir=\"{1}\" -AppID={2} -AppName=\"{3}\" -BuildVersion=\"{4}\" -AppLaunch=\"{5}\"", BuildRoot, CloudDir, AppID, AppName, VersionString, AppLaunch);
+				string CmdLine = String.Format("-BuildRoot={0} -CloudDir={1} -AppID={2} -AppName=\"{3}\" -BuildVersion=\"{4}\" -AppLaunch=\"{5}\"", BuildRoot, CloudDir, AppID, AppName, VersionString, AppLaunch);
 				CmdLine += " -AppArgs=\"\"";
 				CmdLine += " -custom=\"bIsPatch=false\"";
 				CmdLine += String.Format(" -customint=\"ChunkID={0}\"", ChunkID);
