iOS - iPhonePackager ExportCertificate command failing 100%

diff --git a/Engine/Source/Programs/AutomationTool/AutomationUtils/DeploymentContext.cs b/Engine/Source/Programs/AutomationTool/AutomationUtils/DeploymentContext.cs
--- a/Engine/Source/Programs/AutomationTool/AutomationUtils/DeploymentContext.cs
+++ b/Engine/Source/Programs/AutomationTool/AutomationUtils/DeploymentContext.cs
@@ -909,28 +909,28 @@ public class DeploymentContext //: ProjectParams
 	{
 		string SanitizedDeviceName = GetSanitizedDeviceName(DeviceName);
 
-		return Path.Combine(StageDirectory.FullName, UFSDeployDeltaFileName + SanitizedDeviceName);
+		return Path.Combine(StageDirectory.FullName, UFSDeployDeltaFileName + " " + SanitizedDeviceName);
 	}
 
 	public String GetNonUFSDeploymentDeltaPath(string DeviceName)
 	{
 		string SanitizedDeviceName = GetSanitizedDeviceName(DeviceName);
 
-		return Path.Combine(StageDirectory.FullName, NonUFSDeployDeltaFileName + SanitizedDeviceName);
+		return Path.Combine(StageDirectory.FullName, NonUFSDeployDeltaFileName + " " + SanitizedDeviceName);
 	}
 
 	public String GetUFSDeploymentObsoletePath(string DeviceName)
 	{
 		string SanitizedDeviceName = GetSanitizedDeviceName(DeviceName);
 
-		return Path.Combine(StageDirectory.FullName, UFSDeployObsoleteFileName + SanitizedDeviceName);
+		return Path.Combine(StageDirectory.FullName, UFSDeployObsoleteFileName + " " + SanitizedDeviceName);
 	}
 
 	public String GetNonUFSDeploymentObsoletePath(string DeviceName)
 	{
 		string SanitizedDeviceName = GetSanitizedDeviceName(DeviceName);
 
-		return Path.Combine(StageDirectory.FullName, NonUFSDeployObsoleteFileName + SanitizedDeviceName);
+		return Path.Combine(StageDirectory.FullName, NonUFSDeployObsoleteFileName + " " + SanitizedDeviceName);
 	}
 
 	public string UFSDeployedManifestFileName
diff --git a/Engine/Source/Programs/IOS/iPhonePackager/CompileTime.cs b/Engine/Source/Programs/IOS/iPhonePackager/CompileTime.cs
--- a/Engine/Source/Programs/IOS/iPhonePackager/CompileTime.cs
+++ b/Engine/Source/Programs/IOS/iPhonePackager/CompileTime.cs
@@ -177,10 +177,10 @@ namespace iPhonePackager
 			}
 		}
 
-		static bool FindMobileProvision(string BundleIdentifier, out string OutFileName)
+		static bool FindMobileProvision(string BundleIdentifier, out string OutFileName, bool bCheckCert = true)
 		{
 			bool bNameMatch;
-			string ProvisionWithPrefix = MobileProvision.FindCompatibleProvision(BundleIdentifier, out bNameMatch, true, true, false);
+			string ProvisionWithPrefix = MobileProvision.FindCompatibleProvision(BundleIdentifier, out bNameMatch, bCheckCert, true, false);
 			if (!File.Exists(ProvisionWithPrefix))
 			{
 				ProvisionWithPrefix = FileOperations.FindPrefixedFile(Config.BuildDirectory, Program.GameName + ".mobileprovision");
@@ -213,7 +213,8 @@ namespace iPhonePackager
 		static public void ExportCertificate()
 		{
 			string ProvisionWithPrefix;
-			if(!FindMobileProvision("", out ProvisionWithPrefix))
+			// this needs the provisioningUUID or it will fail
+			if(!FindMobileProvision("", out ProvisionWithPrefix, false))
 			{
 				Program.Error("Missing provision");
 				return;
diff --git a/Engine/Source/Programs/UnrealBuildTool/ToolChain/RemoteMac.cs b/Engine/Source/Programs/UnrealBuildTool/ToolChain/RemoteMac.cs
--- a/Engine/Source/Programs/UnrealBuildTool/ToolChain/RemoteMac.cs
+++ b/Engine/Source/Programs/UnrealBuildTool/ToolChain/RemoteMac.cs
@@ -418,6 +418,7 @@ namespace UnrealBuildTool
 						Arguments.AppendFormat(" \"{0}\"", TargetDesc.ProjectFile.Directory);
 					}
 					Arguments.AppendFormat(" -certificate \"{0}\"", CertificateFile);
+					Arguments.AppendFormat(" -provisioninguuid {0}", ProvisioningData.MobileProvisionUUID);
 
 					ProcessStartInfo StartInfo = new ProcessStartInfo();
 					StartInfo.FileName = FileReference.Combine(UnrealBuildTool.EngineDirectory, "Binaries", "DotNET", "IOS", "IPhonePackager.exe").FullName;
