BuildGraph: Add support for FileReference and DirectoryReference parameters, without having to explicitly parse them.

diff --git a/Engine/Source/Programs/AutomationTool/BuildGraph/Schema.cs b/Engine/Source/Programs/AutomationTool/BuildGraph/Schema.cs
--- a/Engine/Source/Programs/AutomationTool/BuildGraph/Schema.cs
+++ b/Engine/Source/Programs/AutomationTool/BuildGraph/Schema.cs
@@ -1,4 +1,4 @@
-﻿// Copyright 1998-2018 Epic Games, Inc. All Rights Reserved.
+// Copyright 1998-2018 Epic Games, Inc. All Rights Reserved.
 
 using System;
 using System.Collections.Generic;
@@ -252,6 +252,8 @@ namespace AutomationTool
 			TypeToSchemaTypeName.Add(typeof(String), GetQualifiedTypeName(ScriptSchemaStandardType.BalancedString));
 			TypeToSchemaTypeName.Add(typeof(Boolean), GetQualifiedTypeName(ScriptSchemaStandardType.Boolean));
 			TypeToSchemaTypeName.Add(typeof(Int32), GetQualifiedTypeName(ScriptSchemaStandardType.Integer));
+			TypeToSchemaTypeName.Add(typeof(FileReference), GetQualifiedTypeName(ScriptSchemaStandardType.BalancedString));
+			TypeToSchemaTypeName.Add(typeof(DirectoryReference), GetQualifiedTypeName(ScriptSchemaStandardType.BalancedString));
 			
 			// Create all the custom user types, and add them to the qualified name lookup
 			List<XmlSchemaType> UserTypes = new List<XmlSchemaType>();
@@ -395,18 +397,8 @@ namespace AutomationTool
 		{
 			switch(Type)
 			{
-				case TaskParameterValidationType.Name:
-					return GetQualifiedTypeName(ScriptSchemaStandardType.Name);
-				case TaskParameterValidationType.NameList:
-					return GetQualifiedTypeName(ScriptSchemaStandardType.NameList);
-				case TaskParameterValidationType.Tag:
-					return GetQualifiedTypeName(ScriptSchemaStandardType.Tag);
 				case TaskParameterValidationType.TagList:
 					return GetQualifiedTypeName(ScriptSchemaStandardType.TagList);
-				case TaskParameterValidationType.Target:
-					return GetQualifiedTypeName(ScriptSchemaStandardType.NameOrTag);
-				case TaskParameterValidationType.TargetList:
-					return GetQualifiedTypeName(ScriptSchemaStandardType.NameOrTagList);
 			}
 			return null;
 		}
diff --git a/Engine/Source/Programs/AutomationTool/BuildGraph/Script.cs b/Engine/Source/Programs/AutomationTool/BuildGraph/Script.cs
--- a/Engine/Source/Programs/AutomationTool/BuildGraph/Script.cs
+++ b/Engine/Source/Programs/AutomationTool/BuildGraph/Script.cs
@@ -1268,6 +1268,14 @@ namespace AutomationTool
 						{
 							Value = Condition.Evaluate(ExpandedValue);
 						}
+						else if (Parameter.ValueType == typeof(FileReference))
+						{
+							Value = CustomTask.ResolveFile(ExpandedValue);
+						}
+						else if (Parameter.ValueType == typeof(DirectoryReference))
+						{
+							Value = CustomTask.ResolveDirectory(ExpandedValue);
+						}
 						else
 						{
 							Value = Convert.ChangeType(ExpandedValue, Parameter.ValueType);
diff --git a/Engine/Source/Programs/AutomationTool/BuildGraph/Task.cs b/Engine/Source/Programs/AutomationTool/BuildGraph/Task.cs
--- a/Engine/Source/Programs/AutomationTool/BuildGraph/Task.cs
+++ b/Engine/Source/Programs/AutomationTool/BuildGraph/Task.cs
@@ -1,4 +1,4 @@
-﻿// Copyright 1998-2018 Epic Games, Inc. All Rights Reserved.
+// Copyright 1998-2018 Epic Games, Inc. All Rights Reserved.
 
 using System;
 using System.Collections.Generic;
@@ -23,50 +23,15 @@ namespace AutomationTool
 		/// </summary>
 		Default,
 
-		/// <summary>
-		/// A standard name; alphanumeric characters, plus underscore and space. Spaces at the start or end, or more than one in a row are prohibited.
-		/// </summary>
-		Name,
-
-		/// <summary>
-		/// A list of names separated by semicolons
-		/// </summary>
-		NameList,
-
-		/// <summary>
-		/// A tag name (a regular name with '#' prefix)
-		/// </summary>
-		Tag,
-
 		/// <summary>
 		/// A list of tag names separated by semicolons
 		/// </summary>
 		TagList,
 
-		/// <summary>
-		/// A standard node/aggregate/agent name or tag name
-		/// </summary>
-		Target,
-
-		/// <summary>
-		/// A list of standard name or tag names separated by semicolons
-		/// </summary>
-		TargetList,
-
 		/// <summary>
 		/// A file specification, which may contain tags and wildcards.
 		/// </summary>
 		FileSpec,
-
-		/// <summary>
-		/// A single file name
-		/// </summary>
-		FileName,
-
-		/// <summary>
-		/// A single directory name
-		/// </summary>
-		DirectoryName,
 	}
 
 	/// <summary>
diff --git a/Engine/Source/Programs/AutomationTool/BuildGraph/Tasks/CommandletTask.cs b/Engine/Source/Programs/AutomationTool/BuildGraph/Tasks/CommandletTask.cs
--- a/Engine/Source/Programs/AutomationTool/BuildGraph/Tasks/CommandletTask.cs
+++ b/Engine/Source/Programs/AutomationTool/BuildGraph/Tasks/CommandletTask.cs
@@ -1,4 +1,4 @@
-﻿// Copyright 1998-2018 Epic Games, Inc. All Rights Reserved.
+// Copyright 1998-2018 Epic Games, Inc. All Rights Reserved.
 
 using System;
 using System.Collections.Generic;
@@ -37,8 +37,8 @@ namespace AutomationTool.Tasks
 		/// <summary>
 		/// The editor executable to use. Defaults to the development UE4Editor executable for the current platform.
 		/// </summary>
-		[TaskParameter(Optional = true, ValidationType = TaskParameterValidationType.FileName)]
-		public string EditorExe;
+		[TaskParameter(Optional = true)]
+		public FileReference EditorExe;
 	}
 
 	/// <summary>
@@ -78,13 +78,13 @@ namespace AutomationTool.Tasks
 
 			// Get the path to the editor, and check it exists
 			FileReference EditorExe;
-			if(String.IsNullOrEmpty(Parameters.EditorExe))
+			if(Parameters.EditorExe == null)
 			{
                 EditorExe = new FileReference(HostPlatform.Current.GetUE4ExePath("UE4Editor-Cmd.exe"));
 			}
 			else
 			{
-				EditorExe = ResolveFile(Parameters.EditorExe);
+				EditorExe = Parameters.EditorExe;
 			}
 
 			// Make sure the editor exists
diff --git a/Engine/Source/Programs/AutomationTool/BuildGraph/Tasks/PakFileTask.cs b/Engine/Source/Programs/AutomationTool/BuildGraph/Tasks/PakFileTask.cs
--- a/Engine/Source/Programs/AutomationTool/BuildGraph/Tasks/PakFileTask.cs
+++ b/Engine/Source/Programs/AutomationTool/BuildGraph/Tasks/PakFileTask.cs
@@ -1,4 +1,4 @@
-﻿// Copyright 1998-2018 Epic Games, Inc. All Rights Reserved.
+// Copyright 1998-2018 Epic Games, Inc. All Rights Reserved.
 
 using AutomationTool;
 using System;
@@ -27,32 +27,32 @@ namespace BuildGraph.Tasks
 		/// <summary>
 		/// PAK file to output
 		/// </summary>
-		[TaskParameter(ValidationType = TaskParameterValidationType.FileName)]
-		public string Output;
+		[TaskParameter]
+		public FileReference Output;
 
 		/// <summary>
 		/// Path to a Response File that contains a list of files to add to the pak file, instead of specifying them individually
 		/// </summary>
-		[TaskParameter(Optional = true, ValidationType = TaskParameterValidationType.FileName)]
-		public string ResponseFile;
+		[TaskParameter(Optional = true)]
+		public FileReference ResponseFile;
 
 		/// <summary>
 		/// Directories to rebase the files relative to. If specified, the shortest path under a listed directory will be used for each file.
 		/// </summary>
-		[TaskParameter(Optional = true, ValidationType = TaskParameterValidationType.DirectoryName)]
+		[TaskParameter(Optional = true)]
 		public string RebaseDir;
 
 		/// <summary>
 		/// Script which gives the order of files
 		/// </summary>
-		[TaskParameter(Optional = true, ValidationType = TaskParameterValidationType.FileName)]
-		public string Order;
+		[TaskParameter(Optional = true)]
+		public FileReference Order;
 
 		/// <summary>
 		/// Encryption keys for this pak file
 		/// </summary>
 		[TaskParameter(Optional = true)]
-		public string Sign;
+		public FileReference Sign;
 
 		/// <summary>
 		/// Whether to compress files
@@ -109,15 +109,10 @@ namespace BuildGraph.Tasks
 			}
 
 			// Get the output parameter
-			FileReference OutputFile = ResolveFile(Parameters.Output);
+			FileReference OutputFile = Parameters.Output;
 
 			// Check for a ResponseFile parameter
-			FileReference ResponseFile = null;
-			if (!String.IsNullOrEmpty(Parameters.ResponseFile))
-			{
-				ResponseFile = ResolveFile(Parameters.ResponseFile);
-			}
-
+			FileReference ResponseFile = Parameters.ResponseFile;
 			if (ResponseFile == null)
 			{
 				// Get a unique filename for the response file
@@ -148,11 +143,11 @@ namespace BuildGraph.Tasks
 			CommandLine.AppendFormat("{0} -create={1}", CommandUtils.MakePathSafeToUseWithCommandLine(OutputFile.FullName), CommandUtils.MakePathSafeToUseWithCommandLine(ResponseFile.FullName));
 			if(Parameters.Sign != null)
 			{
-				CommandLine.AppendFormat(" -sign={0}", CommandUtils.MakePathSafeToUseWithCommandLine(ResolveFile(Parameters.Sign).FullName));
+				CommandLine.AppendFormat(" -sign={0}", CommandUtils.MakePathSafeToUseWithCommandLine(Parameters.Sign.FullName));
 			}
 			if(Parameters.Order != null)
 			{
-				CommandLine.AppendFormat(" -order={0}", CommandUtils.MakePathSafeToUseWithCommandLine(ResolveFile(Parameters.Order).FullName));
+				CommandLine.AppendFormat(" -order={0}", CommandUtils.MakePathSafeToUseWithCommandLine(Parameters.Order.FullName));
 			}
 			if (GlobalCommandLine.Installed)
 			{
diff --git a/Engine/Source/Programs/AutomationTool/BuildGraph/Tasks/RenameTask.cs b/Engine/Source/Programs/AutomationTool/BuildGraph/Tasks/RenameTask.cs
--- a/Engine/Source/Programs/AutomationTool/BuildGraph/Tasks/RenameTask.cs
+++ b/Engine/Source/Programs/AutomationTool/BuildGraph/Tasks/RenameTask.cs
@@ -28,13 +28,13 @@ namespace BuildGraph.Tasks
 		/// <summary>
 		/// The current file name, or pattern to match (eg. *.txt). Should not include any path separators.
 		/// </summary>
-		[TaskParameter(Optional = true, ValidationType = TaskParameterValidationType.DirectoryName)]
+		[TaskParameter(Optional = true)]
 		public string From;
 
 		/// <summary>
 		/// The new name for the file(s). Should not include any path separators.
 		/// </summary>
-		[TaskParameter(ValidationType = TaskParameterValidationType.DirectoryName)]
+		[TaskParameter]
 		public string To;
 
 		/// <summary>
diff --git a/Engine/Source/Programs/AutomationTool/BuildGraph/Tasks/SpawnTask.cs b/Engine/Source/Programs/AutomationTool/BuildGraph/Tasks/SpawnTask.cs
--- a/Engine/Source/Programs/AutomationTool/BuildGraph/Tasks/SpawnTask.cs
+++ b/Engine/Source/Programs/AutomationTool/BuildGraph/Tasks/SpawnTask.cs
@@ -1,4 +1,4 @@
-﻿// Copyright 1998-2018 Epic Games, Inc. All Rights Reserved.
+// Copyright 1998-2018 Epic Games, Inc. All Rights Reserved.
 
 using AutomationTool;
 using System;
@@ -21,8 +21,8 @@ namespace BuildGraph.Tasks
 		/// <summary>
 		/// Executable to spawn
 		/// </summary>
-		[TaskParameter(ValidationType = TaskParameterValidationType.FileName)]
-		public string Exe;
+		[TaskParameter]
+		public FileReference Exe;
 
 		/// <summary>
 		/// Arguments for the newly created process
@@ -65,10 +65,10 @@ namespace BuildGraph.Tasks
 		/// <param name="TagNameToFileSet">Mapping from tag names to the set of files they include</param>
 		public override void Execute(JobContext Job, HashSet<FileReference> BuildProducts, Dictionary<string, HashSet<FileReference>> TagNameToFileSet)
 		{
-			IProcessResult Result = CommandUtils.Run(Parameters.Exe, Parameters.Arguments);
+			IProcessResult Result = CommandUtils.Run(Parameters.Exe.FullName, Parameters.Arguments);
 			if(Result.ExitCode < 0 || Result.ExitCode >= Parameters.ErrorLevel)
 			{
-				throw new AutomationException("{0} terminated with an exit code indicating an error ({1})", Path.GetFileName(Parameters.Exe), Result.ExitCode);
+				throw new AutomationException("{0} terminated with an exit code indicating an error ({1})", Path.GetFileName(Parameters.Exe.FullName), Result.ExitCode);
 			}
 		}
 
diff --git a/Engine/Source/Programs/AutomationTool/BuildGraph/Tasks/StageTask.cs b/Engine/Source/Programs/AutomationTool/BuildGraph/Tasks/StageTask.cs
--- a/Engine/Source/Programs/AutomationTool/BuildGraph/Tasks/StageTask.cs
+++ b/Engine/Source/Programs/AutomationTool/BuildGraph/Tasks/StageTask.cs
@@ -20,8 +20,8 @@ namespace BuildGraph.Tasks
 		/// <summary>
 		/// The project that this target belongs to
 		/// </summary>
-		[TaskParameter(Optional = true, ValidationType = TaskParameterValidationType.FileName)]
-		public string Project;
+		[TaskParameter(Optional = true)]
+		public FileReference Project;
 
 		/// <summary>
 		/// Name of the target to stage
@@ -51,7 +51,7 @@ namespace BuildGraph.Tasks
 		/// Directory the receipt files should be staged to
 		/// </summary>
 		[TaskParameter]
-		public string ToDir;
+		public DirectoryReference ToDir;
 
 		/// <summary>
 		/// Whether to overwrite existing files
@@ -95,14 +95,10 @@ namespace BuildGraph.Tasks
 		public override void Execute(JobContext Job, HashSet<FileReference> BuildProducts, Dictionary<string, HashSet<FileReference>> TagNameToFileSet)
 		{
 			// Get the project path, and check it exists
-			FileReference ProjectFile = null;
-			if(Parameters.Project != null)
+			FileReference ProjectFile = Parameters.Project;
+			if(Parameters.Project != null && !FileReference.Exists(ProjectFile))
 			{
-				ProjectFile = ResolveFile(Parameters.Project);
-				if(!FileReference.Exists(ProjectFile))
-				{
-					throw new AutomationException("Couldn't find project '{0}'", ProjectFile.FullName);
-				}
+				throw new AutomationException("Couldn't find project '{0}'", ProjectFile.FullName);
 			}
 
 			// Get the directories used for staging this project
@@ -110,7 +106,7 @@ namespace BuildGraph.Tasks
 			DirectoryReference SourceProjectDir = (ProjectFile == null)? SourceEngineDir : ProjectFile.Directory;
 
 			// Get the output directories. We flatten the directory structure on output.
-			DirectoryReference TargetDir = ResolveDirectory(Parameters.ToDir);
+			DirectoryReference TargetDir = Parameters.ToDir;
 			DirectoryReference TargetEngineDir = DirectoryReference.Combine(TargetDir, "Engine");
 			DirectoryReference TargetProjectDir = DirectoryReference.Combine(TargetDir, ProjectFile.GetFileNameWithoutExtension());
 
diff --git a/Engine/Source/Programs/AutomationTool/BuildGraph/Tasks/StripTask.cs b/Engine/Source/Programs/AutomationTool/BuildGraph/Tasks/StripTask.cs
--- a/Engine/Source/Programs/AutomationTool/BuildGraph/Tasks/StripTask.cs
+++ b/Engine/Source/Programs/AutomationTool/BuildGraph/Tasks/StripTask.cs
@@ -1,4 +1,4 @@
-﻿// Copyright 1998-2018 Epic Games, Inc. All Rights Reserved.
+// Copyright 1998-2018 Epic Games, Inc. All Rights Reserved.
 
 using AutomationTool;
 using System;
@@ -26,8 +26,8 @@ namespace BuildGraph.Tasks
 		/// <summary>
 		/// The directory to find files in
 		/// </summary>
-		[TaskParameter(Optional = true, ValidationType = TaskParameterValidationType.DirectoryName)]
-		public string BaseDir;
+		[TaskParameter(Optional = true)]
+		public DirectoryReference BaseDir;
 
 		/// <summary>
 		/// List of file specifications separated by semicolons (eg. Engine/.../*.pdb), or the name of a tag set
@@ -38,8 +38,8 @@ namespace BuildGraph.Tasks
 		/// <summary>
 		/// Output directory for the stripped files. Defaults to the input path (overwriting the input files).
 		/// </summary>
-		[TaskParameter(Optional = true, ValidationType = TaskParameterValidationType.DirectoryName)]
-		public string OutputDir;
+		[TaskParameter(Optional = true)]
+		public DirectoryReference OutputDir;
 
 		/// <summary>
 		/// Tag to be applied to build products of this task
@@ -77,10 +77,10 @@ namespace BuildGraph.Tasks
 		public override void Execute(JobContext Job, HashSet<FileReference> BuildProducts, Dictionary<string, HashSet<FileReference>> TagNameToFileSet)
 		{
 			// Get the base directory
-			DirectoryReference BaseDir = ResolveDirectory(Parameters.BaseDir);
+			DirectoryReference BaseDir = Parameters.BaseDir;
 
 			// Get the output directory
-			DirectoryReference OutputDir = ResolveDirectory(Parameters.OutputDir);
+			DirectoryReference OutputDir = Parameters.OutputDir;
 
 			// Find the matching files
 			FileReference[] SourceFiles = ResolveFilespec(BaseDir, Parameters.Files, TagNameToFileSet).OrderBy(x => x.FullName).ToArray();
diff --git a/Engine/Source/Programs/AutomationTool/BuildGraph/Tasks/SubmitTask.cs b/Engine/Source/Programs/AutomationTool/BuildGraph/Tasks/SubmitTask.cs
--- a/Engine/Source/Programs/AutomationTool/BuildGraph/Tasks/SubmitTask.cs
+++ b/Engine/Source/Programs/AutomationTool/BuildGraph/Tasks/SubmitTask.cs
@@ -49,8 +49,8 @@ namespace AutomationTool.Tasks
 		/// <summary>
 		/// Root directory for the stream. If not specified, defaults to the current root directory.
 		/// </summary>
-		[TaskParameter(Optional = true, ValidationType = TaskParameterValidationType.DirectoryName)]
-		public string RootDir;
+		[TaskParameter(Optional = true)]
+		public DirectoryReference RootDir;
 
 		/// <summary>
 		/// Whether to revert unchanged files before attempting to submit
@@ -107,7 +107,7 @@ namespace AutomationTool.Tasks
 					Client.Owner = CommandUtils.P4Env.User;
 					Client.Host = Environment.MachineName;
 					Client.Stream = Parameters.Stream ?? CommandUtils.P4Env.Branch;
-					Client.RootPath = Parameters.RootDir ?? CommandUtils.RootDirectory.FullName;
+					Client.RootPath = Parameters.RootDir.FullName ?? CommandUtils.RootDirectory.FullName;
 					Client.Name = Parameters.Workspace;
 					Client.Options = P4ClientOption.NoAllWrite | P4ClientOption.Clobber | P4ClientOption.NoCompress | P4ClientOption.Unlocked | P4ClientOption.NoModTime | P4ClientOption.RmDir;
 					Client.LineEnd = P4LineEnd.Local;
diff --git a/Engine/Source/Programs/AutomationTool/BuildGraph/Tasks/TagReceiptTask.cs b/Engine/Source/Programs/AutomationTool/BuildGraph/Tasks/TagReceiptTask.cs
--- a/Engine/Source/Programs/AutomationTool/BuildGraph/Tasks/TagReceiptTask.cs
+++ b/Engine/Source/Programs/AutomationTool/BuildGraph/Tasks/TagReceiptTask.cs
@@ -26,14 +26,14 @@ namespace AutomationTool.Tasks
 		/// <summary>
 		/// Path to the Engine folder, used to expand $(EngineDir) properties in receipt files. Defaults to the Engine directory for the current workspace.
 		/// </summary>
-		[TaskParameter(Optional = true, ValidationType = TaskParameterValidationType.DirectoryName)]
-		public string EngineDir;
+		[TaskParameter(Optional = true)]
+		public DirectoryReference EngineDir;
 
 		/// <summary>
 		/// Path to the project folder, used to expand $(ProjectDir) properties in receipt files. Defaults to the Engine directory for the current workspace.
 		/// </summary>
-		[TaskParameter(Optional = true, ValidationType = TaskParameterValidationType.DirectoryName)]
-		public string ProjectDir;
+		[TaskParameter(Optional = true)]
+		public DirectoryReference ProjectDir;
 
 		/// <summary>
 		/// Whether to tag the Build Products listed in receipts
@@ -114,18 +114,10 @@ namespace AutomationTool.Tasks
 		public override void Execute(JobContext Job, HashSet<FileReference> BuildProducts, Dictionary<string, HashSet<FileReference>> TagNameToFileSet)
 		{
 			// Set the Engine directory
-			DirectoryReference EngineDir = DirectoryReference.Combine(CommandUtils.RootDirectory, "Engine");
-			if (!String.IsNullOrEmpty(Parameters.EngineDir))
-			{
-				EngineDir = DirectoryReference.Combine(CommandUtils.RootDirectory, Parameters.EngineDir);
-			}
+			DirectoryReference EngineDir = Parameters.EngineDir ?? CommandUtils.EngineDirectory;
 
 			// Set the Project directory
-			DirectoryReference ProjectDir = DirectoryReference.Combine(CommandUtils.RootDirectory, "Engine");
-			if (!String.IsNullOrEmpty(Parameters.ProjectDir))
-			{
-				ProjectDir = DirectoryReference.Combine(CommandUtils.RootDirectory, Parameters.ProjectDir);
-			}
+			DirectoryReference ProjectDir = Parameters.ProjectDir ?? EngineDir;
 
 			// Resolve the input list
 			IEnumerable<FileReference> TargetFiles = ResolveFilespec(CommandUtils.RootDirectory, Parameters.Files, TagNameToFileSet);
diff --git a/Engine/Source/Programs/AutomationTool/BuildGraph/Tasks/TagTask.cs b/Engine/Source/Programs/AutomationTool/BuildGraph/Tasks/TagTask.cs
--- a/Engine/Source/Programs/AutomationTool/BuildGraph/Tasks/TagTask.cs
+++ b/Engine/Source/Programs/AutomationTool/BuildGraph/Tasks/TagTask.cs
@@ -20,8 +20,8 @@ namespace AutomationTool.Tasks
 		/// <summary>
 		/// Set the base directory to resolve relative paths and patterns against. If set, any absolute patterns (eg. /Engine/Build/...) are taken to be relative to this path. If not, they are taken to be truly absolute.
 		/// </summary>
-		[TaskParameter(Optional = true, ValidationType = TaskParameterValidationType.DirectoryName)]
-		public string BaseDir;
+		[TaskParameter(Optional = true)]
+		public DirectoryReference BaseDir;
 
 		/// <summary>
 		/// Set of files to work from, including wildcards and tag names, separated by semicolons. Resolved relative to BaseDir if set, otherwise to the branch root directory.
@@ -84,7 +84,7 @@ namespace AutomationTool.Tasks
 		public override void Execute(JobContext Job, HashSet<FileReference> BuildProducts, Dictionary<string, HashSet<FileReference>> TagNameToFileSet)
 		{
 			// Get the base directory
-			DirectoryReference BaseDir = ResolveDirectory(Parameters.BaseDir);
+			DirectoryReference BaseDir = Parameters.BaseDir;
 
 			// Parse all the exclude rules
 			List<string> ExcludeRules = ParseRules(BaseDir, Parameters.Except ?? "", TagNameToFileSet);
diff --git a/Engine/Source/Programs/AutomationTool/BuildGraph/Tasks/UnzipTask.cs b/Engine/Source/Programs/AutomationTool/BuildGraph/Tasks/UnzipTask.cs
--- a/Engine/Source/Programs/AutomationTool/BuildGraph/Tasks/UnzipTask.cs
+++ b/Engine/Source/Programs/AutomationTool/BuildGraph/Tasks/UnzipTask.cs
@@ -20,14 +20,14 @@ namespace BuildGraph.Tasks
 		/// <summary>
 		/// Path to the zip file to extract
 		/// </summary>
-		[TaskParameter(ValidationType = TaskParameterValidationType.FileName)]
+		[TaskParameter(ValidationType = TaskParameterValidationType.FileSpec)]
 		public string ZipFile;
 
 		/// <summary>
 		/// Output directory for the extracted files
 		/// </summary>
-		[TaskParameter(ValidationType = TaskParameterValidationType.DirectoryName)]
-		public string ToDir;
+		[TaskParameter]
+		public DirectoryReference ToDir;
 
 		/// <summary>
 		/// Tag to be applied to the extracted files
@@ -64,7 +64,7 @@ namespace BuildGraph.Tasks
 		/// <param name="TagNameToFileSet">Mapping from tag names \to the set of files they include</param>
 		public override void Execute(JobContext Job, HashSet<FileReference> BuildProducts, Dictionary<string, HashSet<FileReference>> TagNameToFileSet)
 		{
-			DirectoryReference ToDir = ResolveDirectory(Parameters.ToDir);
+			DirectoryReference ToDir = Parameters.ToDir;
 
 			// Find all the zip files
 			IEnumerable<FileReference> ZipFiles = ResolveFilespec(CommandUtils.RootDirectory, Parameters.ZipFile, TagNameToFileSet);
diff --git a/Engine/Source/Programs/AutomationTool/BuildGraph/Tasks/ZipTask.cs b/Engine/Source/Programs/AutomationTool/BuildGraph/Tasks/ZipTask.cs
--- a/Engine/Source/Programs/AutomationTool/BuildGraph/Tasks/ZipTask.cs
+++ b/Engine/Source/Programs/AutomationTool/BuildGraph/Tasks/ZipTask.cs
@@ -1,4 +1,4 @@
-﻿// Copyright 1998-2018 Epic Games, Inc. All Rights Reserved.
+// Copyright 1998-2018 Epic Games, Inc. All Rights Reserved.
 
 using AutomationTool;
 using System;
@@ -20,8 +20,8 @@ namespace BuildGraph.Tasks
 		/// <summary>
 		/// The directory to read compressed files from
 		/// </summary>
-		[TaskParameter(ValidationType = TaskParameterValidationType.DirectoryName)]
-		public string FromDir;
+		[TaskParameter]
+		public DirectoryReference FromDir;
 
 		/// <summary>
 		/// List of file specifications separated by semicolons (eg. *.cpp;Engine/.../*.bat), or the name of a tag set. Relative paths are taken from FromDir.
@@ -32,8 +32,8 @@ namespace BuildGraph.Tasks
 		/// <summary>
 		/// The zip file to create
 		/// </summary>
-		[TaskParameter(ValidationType = TaskParameterValidationType.FileName)]
-		public string ZipFile;
+		[TaskParameter]
+		public FileReference ZipFile;
 
 		/// <summary>
 		/// Tag to be applied to the created zip file
@@ -70,31 +70,28 @@ namespace BuildGraph.Tasks
 		/// <param name="TagNameToFileSet">Mapping from tag names to the set of files they include</param>
 		public override void Execute(JobContext Job, HashSet<FileReference> BuildProducts, Dictionary<string, HashSet<FileReference>> TagNameToFileSet)
 		{
-			DirectoryReference FromDir = ResolveDirectory(Parameters.FromDir);
-
 			// Find all the input files
 			IEnumerable<FileReference> Files;
 			if(Parameters.Files == null)
 			{
-				Files = DirectoryReference.EnumerateFiles(FromDir, "*", System.IO.SearchOption.AllDirectories);
+				Files = DirectoryReference.EnumerateFiles(Parameters.FromDir, "*", System.IO.SearchOption.AllDirectories);
 			}
 			else
 			{
-				Files = ResolveFilespec(FromDir, Parameters.Files, TagNameToFileSet);
+				Files = ResolveFilespec(Parameters.FromDir, Parameters.Files, TagNameToFileSet);
 			}
 
 			// Create the zip file
-			FileReference ArchiveFile = ResolveFile(Parameters.ZipFile);
-			CommandUtils.ZipFiles(ArchiveFile, FromDir, Files);
+			CommandUtils.ZipFiles(Parameters.ZipFile, Parameters.FromDir, Files);
 
 			// Apply the optional tag to the produced archive
 			foreach(string TagName in FindTagNamesFromList(Parameters.Tag))
 			{
-				FindOrAddTagSet(TagNameToFileSet, TagName).Add(ArchiveFile);
+				FindOrAddTagSet(TagNameToFileSet, TagName).Add(Parameters.ZipFile);
 			}
 
 			// Add the archive to the set of build products
-			BuildProducts.Add(ArchiveFile);
+			BuildProducts.Add(Parameters.ZipFile);
 		}
 
 		/// <summary>
