BuildGraph: Add support for automatic parsing of collection types in BuildGraph parameters.

diff --git a/Engine/Source/Programs/AutomationTool/BuildGraph/Schema.cs b/Engine/Source/Programs/AutomationTool/BuildGraph/Schema.cs
--- a/Engine/Source/Programs/AutomationTool/BuildGraph/Schema.cs
+++ b/Engine/Source/Programs/AutomationTool/BuildGraph/Schema.cs
@@ -35,6 +35,11 @@ namespace AutomationTool
 		/// </summary>
 		public Type ValueType;
 
+		/// <summary>
+		/// The ICollection interface for this type
+		/// </summary>
+		public Type CollectionType;
+
 		/// <summary>
 		/// Validation type for this field
 		/// </summary>
@@ -61,6 +66,22 @@ namespace AutomationTool
 				ValueType = ValueType.GetGenericArguments()[0];
 				bOptional = true;
 			}
+
+			if(ValueType.IsClass)
+			{
+				foreach(Type InterfaceType in ValueType.GetInterfaces())
+				{
+					if(InterfaceType.IsGenericType)
+					{
+						Type GenericInterfaceType = InterfaceType.GetGenericTypeDefinition();
+						if(GenericInterfaceType == typeof(ICollection<>))
+						{
+							CollectionType = InterfaceType;
+							ValueType = InterfaceType.GetGenericArguments()[0];
+						}
+					}
+				}
+			}
 		}
 	}
 
@@ -261,10 +282,17 @@ namespace AutomationTool
 			{
 				if(!TypeToSchemaTypeName.ContainsKey(Type))
 				{
-					string Name = Type.Name + "UserType";
-					XmlSchemaType SchemaType = CreateUserType(Name, Type);
-					UserTypes.Add(SchemaType);
-					TypeToSchemaTypeName.Add(Type, new XmlQualifiedName(Name, NamespaceURI));
+					if(Type.IsClass && Type.GetInterfaces().Any(x => x.GetGenericTypeDefinition() == typeof(ICollection<>)))
+					{
+						TypeToSchemaTypeName.Add(Type, GetQualifiedTypeName(ScriptSchemaStandardType.BalancedString));
+					}
+					else
+					{
+						string Name = Type.Name + "UserType";
+						XmlSchemaType SchemaType = CreateUserType(Name, Type);
+						UserTypes.Add(SchemaType);
+						TypeToSchemaTypeName.Add(Type, new XmlQualifiedName(Name, NamespaceURI));
+					}
 				}
 			}
 
diff --git a/Engine/Source/Programs/AutomationTool/BuildGraph/Script.cs b/Engine/Source/Programs/AutomationTool/BuildGraph/Script.cs
--- a/Engine/Source/Programs/AutomationTool/BuildGraph/Script.cs
+++ b/Engine/Source/Programs/AutomationTool/BuildGraph/Script.cs
@@ -1258,29 +1258,31 @@ namespace AutomationTool
 						// Expand variables in the value
 						string ExpandedValue = ExpandProperties(Element, Attribute.Value);
 
-						// Parse it and assign it to the parameters object
-						object Value;
-						if (Parameter.ValueType.IsEnum)
+						// If it's a collection type, split it into separate values
+						if(Parameter.CollectionType == null)
 						{
-							Value = Enum.Parse(Parameter.ValueType, ExpandedValue);
-						}
-						else if (Parameter.ValueType == typeof(Boolean))
-						{
-							Value = Condition.Evaluate(ExpandedValue);
-						}
-						else if (Parameter.ValueType == typeof(FileReference))
-						{
-							Value = CustomTask.ResolveFile(ExpandedValue);
-						}
-						else if (Parameter.ValueType == typeof(DirectoryReference))
-						{
-							Value = CustomTask.ResolveDirectory(ExpandedValue);
+							// Parse it and assign it to the parameters object
+							object Value = ParseValue(ExpandedValue, Parameter.ValueType);
+							Parameter.FieldInfo.SetValue(ParametersObject, Value);
 						}
 						else
 						{
-							Value = Convert.ChangeType(ExpandedValue, Parameter.ValueType);
+							// Get the collection, or create one if necessary
+							object CollectionValue = Parameter.FieldInfo.GetValue(ParametersObject);
+							if(CollectionValue == null)
+							{
+								CollectionValue = Activator.CreateInstance(Parameter.FieldInfo.FieldType);
+								Parameter.FieldInfo.SetValue(ParametersObject, CollectionValue);
+							}
+
+							// Parse the values and add them to the collection
+							List<string> ValueStrings = CustomTask.SplitDelimitedList(ExpandedValue);
+							foreach(string ValueString in ValueStrings)
+							{
+								object Value = ParseValue(ValueString, Parameter.ValueType);
+								Parameter.CollectionType.InvokeMember("Add", BindingFlags.InvokeMethod | BindingFlags.Instance | BindingFlags.Public, null, CollectionValue, new object[]{ Value });
+							}
 						}
-						Parameter.FieldInfo.SetValue(ParametersObject, Value);
 					}
 				}
 
@@ -1325,6 +1327,37 @@ namespace AutomationTool
 			}
 		}
 
+		/// <summary>
+		/// Parse a value of the given type
+		/// </summary>
+		/// <param name="ValueText">The text to parse</param>
+		/// <param name="ValueType">Type of the value to parse</param>
+		/// <returns>Value that was parsed</returns>
+		object ParseValue(string ValueText, Type ValueType)
+		{
+			// Parse it and assign it to the parameters object
+			if (ValueType.IsEnum)
+			{
+				return Enum.Parse(ValueType, ValueText);
+			}
+			else if (ValueType == typeof(Boolean))
+			{
+				return Condition.Evaluate(ValueText);
+			}
+			else if (ValueType == typeof(FileReference))
+			{
+				return CustomTask.ResolveFile(ValueText);
+			}
+			else if (ValueType == typeof(DirectoryReference))
+			{
+				return CustomTask.ResolveDirectory(ValueText);
+			}
+			else
+			{
+				return Convert.ChangeType(ValueText, ValueType);
+			}
+		}
+
 		/// <summary>
 		/// Reads the definition for an email notifier
 		/// </summary>
diff --git a/Engine/Source/Programs/AutomationTool/BuildGraph/Tasks/PakFileTask.cs b/Engine/Source/Programs/AutomationTool/BuildGraph/Tasks/PakFileTask.cs
--- a/Engine/Source/Programs/AutomationTool/BuildGraph/Tasks/PakFileTask.cs
+++ b/Engine/Source/Programs/AutomationTool/BuildGraph/Tasks/PakFileTask.cs
@@ -40,7 +40,7 @@ namespace BuildGraph.Tasks
 		/// Directories to rebase the files relative to. If specified, the shortest path under a listed directory will be used for each file.
 		/// </summary>
 		[TaskParameter(Optional = true)]
-		public string RebaseDir;
+		public HashSet<DirectoryReference> RebaseDir;
 
 		/// <summary>
 		/// Script which gives the order of files
@@ -105,7 +105,7 @@ namespace BuildGraph.Tasks
 			HashSet<DirectoryReference> RebaseDirs = new HashSet<DirectoryReference>{ CommandUtils.RootDirectory };
 			if(Parameters.RebaseDir != null)
 			{
-				RebaseDirs.UnionWith(SplitDelimitedList(Parameters.RebaseDir).Select(x => ResolveDirectory(x)));
+				RebaseDirs.UnionWith(Parameters.RebaseDir);
 			}
 
 			// Get the output parameter
