UBT: Make BuildException internal to UBT.

#rnx

diff --git a/Engine/Source/Programs/AutomationTool/BuildGraph/Tasks/TagTask.cs b/Engine/Source/Programs/AutomationTool/BuildGraph/Tasks/TagTask.cs
--- a/Engine/Source/Programs/AutomationTool/BuildGraph/Tasks/TagTask.cs
+++ b/Engine/Source/Programs/AutomationTool/BuildGraph/Tasks/TagTask.cs
@@ -104,7 +104,7 @@ namespace AutomationTool.Tasks
 				{
 					if(!FileReference.Exists(FileList))
 					{
-						throw new BuildException("Specified file list '{0}' does not exist", FileList);
+						throw new AutomationException("Specified file list '{0}' does not exist", FileList);
 					}
 
 					string[] Lines = FileReference.ReadAllLines(FileList);
diff --git a/Engine/Source/Programs/AutomationTool/Scripts/BuildPhysX.Automation.cs b/Engine/Source/Programs/AutomationTool/Scripts/BuildPhysX.Automation.cs
--- a/Engine/Source/Programs/AutomationTool/Scripts/BuildPhysX.Automation.cs
+++ b/Engine/Source/Programs/AutomationTool/Scripts/BuildPhysX.Automation.cs
@@ -234,7 +234,7 @@ class BuildPhysX : BuildCommand
 						// don't register if we don't have an NDKROOT specified
 						if (String.IsNullOrEmpty(NDKDirectory))
 						{
-							throw new BuildException("NDKROOT is not specified; cannot build Android.");
+							throw new AutomationException("NDKROOT is not specified; cannot build Android.");
 						}
 
 						NDKDirectory = NDKDirectory.Replace("\"", "");
@@ -571,7 +571,7 @@ class BuildPhysX : BuildCommand
 			}
 			catch (Exception ex)
 			{
-				throw new BuildException(ex, "Failed to start local process for action (\"{0}\"): {1} {2}", ex.Message, LocalProcess.StartInfo.FileName, LocalProcess.StartInfo.Arguments);
+				throw new AutomationException(ex, "Failed to start local process for action (\"{0}\"): {1} {2}", ex.Message, LocalProcess.StartInfo.FileName, LocalProcess.StartInfo.Arguments);
 			}
 		}
 
@@ -1001,7 +1001,7 @@ class BuildPhysX : BuildCommand
 					// don't register if we don't have an NDKROOT specified
 					if (String.IsNullOrEmpty(NDKDirectory))
 					{
-						throw new BuildException("NDKROOT is not specified; cannot build Android.");
+						throw new AutomationException("NDKROOT is not specified; cannot build Android.");
 					}
 		
 					NDKDirectory = NDKDirectory.Replace("\"", "");
diff --git a/Engine/Source/Programs/UnrealBuildTool/System/BuildException.cs b/Engine/Source/Programs/UnrealBuildTool/System/BuildException.cs
--- a/Engine/Source/Programs/UnrealBuildTool/System/BuildException.cs
+++ b/Engine/Source/Programs/UnrealBuildTool/System/BuildException.cs
@@ -9,7 +9,7 @@ namespace UnrealBuildTool
 	/// <summary>
 	/// Base class for exceptions thrown by UBT
 	/// </summary>
-	public class BuildException : Exception
+	class BuildException : Exception
 	{
 		/// <summary>
 		/// Constructor
@@ -26,7 +26,7 @@ namespace UnrealBuildTool
 		/// <param name="InnerException">An inner exception to wrap</param>
 		/// <param name="Message">The error message to display.</param>
 		public BuildException(Exception InnerException, string Message)
-			: base(Message)
+			: base(Message, InnerException)
 		{
 		}
 
