Prevent null dereference when tagging a set of files.

#rnx

diff --git a/Engine/Source/Programs/AutomationTool/BuildGraph/Task.cs b/Engine/Source/Programs/AutomationTool/BuildGraph/Task.cs
--- a/Engine/Source/Programs/AutomationTool/BuildGraph/Task.cs
+++ b/Engine/Source/Programs/AutomationTool/BuildGraph/Task.cs
@@ -391,13 +391,16 @@ namespace AutomationTool
 				{
 					FileFilter Filter = new FileFilter();
 					Filter.AddRule(IncludePattern, FileFilterType.Include);
-					Filter.AddRules(ExcludePatterns, FileFilterType.Exclude);
+					if(ExcludePatterns != null && ExcludePatterns.Count > 0)
+					{
+						Filter.AddRules(ExcludePatterns, FileFilterType.Exclude);
+					}
 					Files.UnionWith(Filter.ApplyToDirectory(BaseDir, BaseDir.FullName, true));
 				}
 			}
 
 			// If we have exclude rules, create and run a filter against all the output files to catch things that weren't added from an include
-			if(ExcludePatterns.Count > 0)
+			if(ExcludePatterns != null && ExcludePatterns.Count > 0)
 			{
 				FileFilter Filter = new FileFilter(FileFilterType.Include);
 				Filter.AddRules(ExcludePatterns, FileFilterType.Exclude);
