UAT: iterative deploy no longer prevents cleaning of staging directory

Partial fix for

Might need additional work to reduce file copies from Build to Stage.

diff --git a/Engine/Source/Programs/AutomationTool/Scripts/CopyBuildToStagingDirectory.Automation.cs b/Engine/Source/Programs/AutomationTool/Scripts/CopyBuildToStagingDirectory.Automation.cs
--- a/Engine/Source/Programs/AutomationTool/Scripts/CopyBuildToStagingDirectory.Automation.cs
+++ b/Engine/Source/Programs/AutomationTool/Scripts/CopyBuildToStagingDirectory.Automation.cs
@@ -2138,9 +2138,9 @@ public partial class Project : CommandUtils
 
 	public static void CleanStagingDirectory(ProjectParams Params, DeploymentContext SC)
 	{
-		LogInformation("Cleaning Stage Directory: {0}", SC.StageDirectory.FullName);
-		if (SC.Stage && !Params.NoCleanStage && !Params.SkipStage && !Params.IterativeDeploy)
+		if (SC.Stage && !Params.NoCleanStage && !Params.SkipStage)
 		{
+			LogInformation("Cleaning stage directory: {0}", SC.StageDirectory.FullName);
 			try
 			{
 				DeleteDirectory(SC.StageDirectory.FullName);
@@ -2153,6 +2153,7 @@ public partial class Project : CommandUtils
 		}
 		else
 		{
+			LogInformation("Cleaning PAK files in stage directory: {0}", SC.StageDirectory.FullName);
 			try
 			{
 				// delete old pak files
