UAT: fix Stage & Iterative Deploy leaving behind old files

For iterative deploy on Android, files & directories not in staging will be zapped from the device.
Tested with Samsung Galaxy S6+ (Android 6.0) and Google Pixel 2 (Android 9.0)

diff --git a/Engine/Source/Programs/AutomationTool/Android/AndroidPlatform.Automation.cs b/Engine/Source/Programs/AutomationTool/Android/AndroidPlatform.Automation.cs
--- a/Engine/Source/Programs/AutomationTool/Android/AndroidPlatform.Automation.cs
+++ b/Engine/Source/Programs/AutomationTool/Android/AndroidPlatform.Automation.cs
@@ -991,6 +991,22 @@ public class AndroidPlatform : Platform
 		}
 	}
 
+	// Returns a filename from "adb shell ls -RF" output
+	// or null if the input line is a directory.
+	private string GetFileNameFromListing(string SingleLine)
+	{
+		if (SingleLine.StartsWith("- ")) // file on Samsung
+			return SingleLine.Substring(2);
+		else if (SingleLine.StartsWith("d ")) // directory on Samsung
+			return null;
+		else if (SingleLine.EndsWith("/")) // directory on Google
+			return null;
+		else // undecorated = file on Google
+		{
+			return SingleLine;
+		}
+	}
+
     public override void Deploy(ProjectParams Params, DeploymentContext SC)
     {
 		var AppArchitectures = AndroidExports.CreateToolChain(Params.RawProjectPath).GetAllArchitectures();
@@ -1155,6 +1171,137 @@ public class AndroidPlatform : Platform
                         EntriesToDeploy.TrimExcess();
                         EntriesToDeploy.Add(SC.StageDirectory.FullName);
                     }
+					else
+					{
+						// Discover & remove any files on device that are not in staging
+
+						// get listing of remote directory from device
+						string Commandline = "shell ls -RF1 " + RemoteDir;
+						var CommandResult = RunAdbCommand(Params, DeviceName, Commandline, null, ERunOptions.AppMustExist);
+						// CommandResult.ExitCode is adb shell's exit code, not ls exit code, which is what we need.
+						// Check output for error message instead.
+						if (CommandResult.Output.StartsWith("ls: "))
+						{
+							// list command failed, try simpler options
+							Commandline = "shell ls -RF " + RemoteDir;
+							CommandResult = RunAdbCommand(Params, DeviceName, Commandline, null, ERunOptions.AppMustExist);
+						}
+
+						if (CommandResult.Output.StartsWith("ls: "))
+						{
+							// list command failed, so clean the remote dir instead of selectively deleting files
+							RunAdbCommand(Params, DeviceName, "shell rm -r " + RemoteDir);
+						}
+						else
+						{
+						// listing output is of the form
+						// [Samsung]                 [Google]
+						//
+						// RemoteDir/RestOfPath:     RemoteDir/RestOfPath:
+						// - File1.png               File1.png
+						// - File2.txt               File2.txt
+						// d SubDir1                 SubDir1/
+						// d SubDir2                 Subdir2/
+						//
+						// RemoteDir/RestOfPath/SubDir1:
+
+						HashSet<string> DirsToDeleteFromDevice = new HashSet<string>();
+						List<string> FilesToDeleteFromDevice = new List<string>();
+
+						using (var reader = new StringReader(CommandResult.Output))
+						{
+							string ProjectSaved = Params.ShortProjectName + "/Saved";
+							string ProjectConfig = Params.ShortProjectName + "/Config";
+							const string EngineSaved = "Engine/Saved"; // is this safe to use, or should we use SC.EngineRoot.GetDirectoryName()?
+							const string EngineConfig = "Engine/Config";
+							LogWarning("Excluding {0} {1} {2} {3} from clean during deployment.", ProjectSaved, ProjectConfig, EngineSaved, EngineConfig);
+
+							string CurrentDir = "";
+							bool SkipFiles = false;
+							for (string Line = reader.ReadLine(); Line != null; Line = reader.ReadLine())
+							{
+								if (String.IsNullOrWhiteSpace(Line))
+								{
+									continue; // ignore blank lines
+								}
+
+								if (Line.EndsWith(":"))
+								{
+									// RemoteDir/RestOfPath:
+									//      keep ^--------^
+									CurrentDir = Line.Substring(RemoteDir.Length + 1, Math.Max(0, Line.Length - RemoteDir.Length - 2));
+									// Max is there for the case of base "RemoteDir:" --> ""
+
+									// We want to keep config & logs between deployments.
+									if (CurrentDir.StartsWith(ProjectSaved) || CurrentDir.StartsWith(ProjectConfig) || CurrentDir.StartsWith(EngineSaved) || CurrentDir.StartsWith(EngineConfig))
+									{
+										SkipFiles = true;
+										continue;
+									}
+
+									bool DirExistsInStagingArea = Directory.Exists(Path.Combine(SC.StageDirectory.FullName, CurrentDir));
+									if (DirExistsInStagingArea)
+									{
+										SkipFiles = false;
+									}
+									else
+									{
+										// delete directory from device
+										SkipFiles = true;
+										DirsToDeleteFromDevice.Add(CurrentDir);
+									}
+								}
+								else
+								{
+									if (SkipFiles)
+									{
+										continue;
+									}
+
+									string FileName = GetFileNameFromListing(Line);
+									if (FileName != null)
+									{
+										bool FileExistsInStagingArea = File.Exists(Path.Combine(SC.StageDirectory.FullName, CurrentDir, FileName));
+										if (FileExistsInStagingArea)
+										{
+											// keep or overwrite
+										}
+										else
+										{
+											// delete file from device
+											string FilePath = CurrentDir.Length == 0 ? FileName : (CurrentDir + "/" + FileName); // use / for Android target, no matter the development system
+											LogWarning("Deleting {0} from device; not found in staging area", FilePath);
+											FilesToDeleteFromDevice.Add(FilePath);
+										}
+									}
+									// We ignore subdirs here as each will have its own "RemoteDir/CurrentDir/SubDir:" entry.
+								}
+							}
+						}
+
+						// delete directories
+						foreach (var DirToDelete in DirsToDeleteFromDevice)
+						{
+							// if a whole tree is to be deleted, don't spend extra commands deleting its branches
+							int FinalSlash = DirToDelete.LastIndexOf('/');
+							string ParentDir = FinalSlash >= 0 ? DirToDelete.Substring(0, FinalSlash) : "";
+							bool ParentMarkedForDeletion = DirsToDeleteFromDevice.Contains(ParentDir);
+							if (!ParentMarkedForDeletion)
+							{
+								LogWarning("Deleting {0} and its contents from device; not found in staging area", DirToDelete);
+								RunAdbCommand(Params, DeviceName, "shell rm -r " + RemoteDir + "/" + DirToDelete);
+							}
+						}
+
+						// delete loose files
+						if (FilesToDeleteFromDevice.Count > 0)
+						{
+							// delete all stray files with one command
+							Commandline = String.Format("shell cd {0}; rm ", RemoteDir);
+							RunAdbCommand(Params, DeviceName, Commandline + String.Join(" ", FilesToDeleteFromDevice));
+						}
+						}
+					}
                 }
                 else
                 {
