Revert change to use FileReference for executable to use in the <Spawn> task, to ensure we can use programs in the PATH.

#rnx

diff --git a/Engine/Source/Programs/AutomationTool/BuildGraph/Tasks/SpawnTask.cs b/Engine/Source/Programs/AutomationTool/BuildGraph/Tasks/SpawnTask.cs
--- a/Engine/Source/Programs/AutomationTool/BuildGraph/Tasks/SpawnTask.cs
+++ b/Engine/Source/Programs/AutomationTool/BuildGraph/Tasks/SpawnTask.cs
@@ -22,7 +22,7 @@ namespace BuildGraph.Tasks
 		/// Executable to spawn
 		/// </summary>
 		[TaskParameter]
-		public FileReference Exe;
+		public string Exe;
 
 		/// <summary>
 		/// Arguments for the newly created process
@@ -65,10 +65,10 @@ namespace BuildGraph.Tasks
 		/// <param name="TagNameToFileSet">Mapping from tag names to the set of files they include</param>
 		public override void Execute(JobContext Job, HashSet<FileReference> BuildProducts, Dictionary<string, HashSet<FileReference>> TagNameToFileSet)
 		{
-			IProcessResult Result = CommandUtils.Run(Parameters.Exe.FullName, Parameters.Arguments);
+			IProcessResult Result = CommandUtils.Run(Parameters.Exe, Parameters.Arguments);
 			if(Result.ExitCode < 0 || Result.ExitCode >= Parameters.ErrorLevel)
 			{
-				throw new AutomationException("{0} terminated with an exit code indicating an error ({1})", Path.GetFileName(Parameters.Exe.FullName), Result.ExitCode);
+				throw new AutomationException("{0} terminated with an exit code indicating an error ({1})", Path.GetFileName(Parameters.Exe), Result.ExitCode);
 			}
 		}
 
