Fix null reference exception when resolving paths.

#rnx

diff --git a/Engine/Source/Programs/AutomationTool/BuildGraph/Tasks/TagTask.cs b/Engine/Source/Programs/AutomationTool/BuildGraph/Tasks/TagTask.cs
--- a/Engine/Source/Programs/AutomationTool/BuildGraph/Tasks/TagTask.cs
+++ b/Engine/Source/Programs/AutomationTool/BuildGraph/Tasks/TagTask.cs
@@ -84,7 +84,7 @@ namespace AutomationTool.Tasks
 		public override void Execute(JobContext Job, HashSet<FileReference> BuildProducts, Dictionary<string, HashSet<FileReference>> TagNameToFileSet)
 		{
 			// Get the base directory
-			DirectoryReference BaseDir = Parameters.BaseDir;
+			DirectoryReference BaseDir = Parameters.BaseDir ?? CommandUtils.RootDirectory;
 
 			// Parse all the exclude rules
 			List<string> ExcludeRules = ParseRules(BaseDir, Parameters.Except ?? "", TagNameToFileSet);
