UAT: Allow overriding the stream name in the BuildForUGS command.

diff --git a/Engine/Source/Programs/AutomationTool/Scripts/BuildForUGS.Automation.cs b/Engine/Source/Programs/AutomationTool/Scripts/BuildForUGS.Automation.cs
--- a/Engine/Source/Programs/AutomationTool/Scripts/BuildForUGS.Automation.cs
+++ b/Engine/Source/Programs/AutomationTool/Scripts/BuildForUGS.Automation.cs
@@ -1,4 +1,4 @@
-﻿// Copyright 1998-2018 Epic Games, Inc. All Rights Reserved.
+// Copyright 1998-2018 Epic Games, Inc. All Rights Reserved.
 
 using System;
 using System.Collections.Generic;
@@ -30,6 +30,13 @@ namespace AutomationTool
 				throw new AutomationException("Archive path is not a valid depot filename");
 			}
 
+			// Parse the stream name
+			string StreamName = ParseParamValue("Stream");
+			if(StreamName == null && ArchivePath != null)
+			{
+				StreamName = ArchivePath.Substring(0, ArchivePath.IndexOf('/', ArchivePath.IndexOf('/', 2) + 1));
+			}
+
 			// Prepare the build agenda
 			UE4Build.BuildAgenda Agenda = new UE4Build.BuildAgenda();
 			foreach(string Target in Targets)
@@ -112,7 +119,7 @@ namespace AutomationTool
 					P4ClientInfo Client = new P4ClientInfo();
 					Client.Owner = CommandUtils.P4Env.User;
 					Client.Host = Environment.MachineName;
-					Client.Stream = ArchivePath.Substring(0, ArchivePath.IndexOf('/', ArchivePath.IndexOf('/', 2) + 1));
+					Client.Stream = StreamName;
 					Client.RootPath = Path.Combine(OutputFolder, "Perforce");
 					Client.Name = ClientName;
 					Client.Options = P4ClientOption.NoAllWrite | P4ClientOption.NoClobber | P4ClientOption.NoCompress | P4ClientOption.Unlocked | P4ClientOption.NoModTime | P4ClientOption.RmDir;
