Add "Regex" tag, useful for stripping bits out of existing properties or options.

diff --git a/Engine/Source/Programs/AutomationTool/BuildGraph/Schema.cs b/Engine/Source/Programs/AutomationTool/BuildGraph/Schema.cs
--- a/Engine/Source/Programs/AutomationTool/BuildGraph/Schema.cs
+++ b/Engine/Source/Programs/AutomationTool/BuildGraph/Schema.cs
@@ -158,6 +158,7 @@ namespace AutomationTool
 		Option,
 		EnvVar,
 		Property,
+		Regex,
 		Macro,
 		Expand,
 		Trace,
@@ -335,6 +336,7 @@ namespace AutomationTool
 			NewSchema.Items.Add(CreateOptionType());
 			NewSchema.Items.Add(CreateEnvVarType());
 			NewSchema.Items.Add(CreatePropertyType());
+			NewSchema.Items.Add(CreateRegexType());
 			NewSchema.Items.Add(CreateMacroType());
 			NewSchema.Items.Add(CreateExpandType());
 			NewSchema.Items.Add(CreateDiagnosticType(ScriptSchemaStandardType.Trace));
@@ -444,6 +446,7 @@ namespace AutomationTool
 			GraphChoice.Items.Add(CreateSchemaElement("Option", ScriptSchemaStandardType.Option));
 			GraphChoice.Items.Add(CreateSchemaElement("EnvVar", ScriptSchemaStandardType.EnvVar));
 			GraphChoice.Items.Add(CreateSchemaElement("Property", ScriptSchemaStandardType.Property));
+			GraphChoice.Items.Add(CreateSchemaElement("Regex", ScriptSchemaStandardType.Regex));
 			GraphChoice.Items.Add(CreateSchemaElement("Macro", ScriptSchemaStandardType.Macro));
 			GraphChoice.Items.Add(CreateSchemaElement("Agent", ScriptSchemaStandardType.Agent));
 			GraphChoice.Items.Add(CreateSchemaElement("Trigger", ScriptSchemaStandardType.Trigger));
@@ -494,6 +497,7 @@ namespace AutomationTool
 			TriggerChoice.MinOccurs = 0;
 			TriggerChoice.MaxOccursString = "unbounded";
 			TriggerChoice.Items.Add(CreateSchemaElement("Property", ScriptSchemaStandardType.Property));
+			TriggerChoice.Items.Add(CreateSchemaElement("Regex", ScriptSchemaStandardType.Regex));
 			TriggerChoice.Items.Add(CreateSchemaElement("EnvVar", ScriptSchemaStandardType.EnvVar));
 			TriggerChoice.Items.Add(CreateSchemaElement("Agent", ScriptSchemaStandardType.Agent));
 			TriggerChoice.Items.Add(CreateSchemaElement("Aggregate", ScriptSchemaStandardType.Aggregate));
@@ -541,6 +545,7 @@ namespace AutomationTool
 			AgentChoice.MinOccurs = 0;
 			AgentChoice.MaxOccursString = "unbounded";
 			AgentChoice.Items.Add(CreateSchemaElement("Property", ScriptSchemaStandardType.Property));
+			AgentChoice.Items.Add(CreateSchemaElement("Regex", ScriptSchemaStandardType.Regex));
 			AgentChoice.Items.Add(CreateSchemaElement("EnvVar", ScriptSchemaStandardType.EnvVar));
 			AgentChoice.Items.Add(CreateSchemaElement("Node", ScriptSchemaStandardType.Node));
 			AgentChoice.Items.Add(CreateSchemaElement("Trace", ScriptSchemaStandardType.Trace));
@@ -592,6 +597,7 @@ namespace AutomationTool
 			NodeChoice.MinOccurs = 0;
 			NodeChoice.MaxOccursString = "unbounded";
 			NodeChoice.Items.Add(CreateSchemaElement("Property", ScriptSchemaStandardType.Property));
+			NodeChoice.Items.Add(CreateSchemaElement("Regex", ScriptSchemaStandardType.Regex));
 			NodeChoice.Items.Add(CreateSchemaElement("EnvVar", ScriptSchemaStandardType.EnvVar));
 			NodeChoice.Items.Add(CreateSchemaElement("Trace", ScriptSchemaStandardType.Trace));
 			NodeChoice.Items.Add(CreateSchemaElement("Warning", ScriptSchemaStandardType.Warning));
@@ -739,6 +745,29 @@ namespace AutomationTool
 			return PropertyType;
 		}
 
+		/// <summary>
+		/// Creates the schema type representing a regex type
+		/// </summary>
+		/// <returns>Type definition for a regex</returns>
+		static XmlSchemaType CreateRegexType()
+		{
+			XmlSchemaSimpleContentExtension Extension = new XmlSchemaSimpleContentExtension();
+			Extension.BaseTypeName = StringTypeName;
+
+			Extension.Attributes.Add(CreateSchemaAttribute("Input", StringTypeName, XmlSchemaUse.Required));
+			Extension.Attributes.Add(CreateSchemaAttribute("Pattern", StringTypeName, XmlSchemaUse.Required));
+			Extension.Attributes.Add(CreateSchemaAttribute("Capture", ScriptSchemaStandardType.NameList, XmlSchemaUse.Required));
+			Extension.Attributes.Add(CreateSchemaAttribute("If", ScriptSchemaStandardType.BalancedString, XmlSchemaUse.Optional));
+
+			XmlSchemaSimpleContent ContentModel = new XmlSchemaSimpleContent();
+			ContentModel.Content = Extension;
+
+			XmlSchemaComplexType RegexType  = new XmlSchemaComplexType();
+			RegexType.Name = GetTypeName(ScriptSchemaStandardType.Regex);
+			RegexType.ContentModel = ContentModel;
+			return RegexType;
+		}
+
 		/// <summary>
 		/// Creates the schema type representing the macro type
 		/// </summary>
diff --git a/Engine/Source/Programs/AutomationTool/BuildGraph/Script.cs b/Engine/Source/Programs/AutomationTool/BuildGraph/Script.cs
--- a/Engine/Source/Programs/AutomationTool/BuildGraph/Script.cs
+++ b/Engine/Source/Programs/AutomationTool/BuildGraph/Script.cs
@@ -344,6 +344,9 @@ namespace AutomationTool
 					case "Property":
 						ReadProperty(ChildElement);
 						break;
+					case "Regex":
+						ReadRegex(ChildElement);
+						break;
 					case "EnvVar":
 						ReadEnvVar(ChildElement);
 						break;
@@ -703,6 +706,55 @@ namespace AutomationTool
 			}
 		}
 
+		/// <summary>
+		/// Reads a Regex assignment.
+		/// </summary>
+		/// <param name="Element">Xml element to read the definition from</param>
+		void ReadRegex(ScriptElement Element)
+		{
+			if (EvaluateCondition(Element))
+			{
+				// Get the pattern
+				string RegexString = ReadAttribute(Element, "Pattern");
+				// Make sure its a valid regex.
+				Regex RegexValue = ParseRegex(Element, RegexString);
+				if (RegexValue != null)
+				{
+					// read the names in 
+					string[] CaptureNames = ReadListAttribute(Element, "Capture");
+
+					// get number of groups we passed in
+					int[] GroupNumbers = RegexValue.GetGroupNumbers();
+
+					// make sure the number of property names is the same as the number of match groups
+					// this includes the entire string match group as [0], so don't count that one.
+					if (CaptureNames.Length != GroupNumbers.Count() - 1)
+					{
+						LogError(Element, "MatchGroup count: {0} does not match the number of names specified: {1}", GroupNumbers.Count() - 1, CaptureNames.Length);
+					}
+					else
+					{
+						// apply the regex to the value
+						string Input = ReadAttribute(Element, "Input");
+						Match Match = RegexValue.Match(Input);
+
+						if (!Match.Success)
+						{
+							LogError(Element, "Regex {0} did not find a match against input string {1}", RegexString, Input);
+						}
+						else
+						{
+							// assign each property to the group it matches, skip over [0]
+							for (int MatchIdx = 1; MatchIdx < GroupNumbers.Count(); MatchIdx++)
+							{
+								SetPropertyValue(Element, CaptureNames[MatchIdx - 1], Match.Groups[MatchIdx].Value);
+							}
+						}
+					}
+				}
+			}
+		}
+
 		/// <summary>
 		/// Reads a property assignment from an environment variable.
 		/// </summary>
@@ -832,6 +884,9 @@ namespace AutomationTool
 					case "Property":
 						ReadProperty(ChildElement);
 						break;
+					case "Regex":
+						ReadRegex(ChildElement);
+						break;
 					case "Node":
 						ReadNode(ChildElement, ParentAgent, ControllingTrigger);
 						break;
@@ -1050,6 +1105,9 @@ namespace AutomationTool
 					case "Property":
 						ReadProperty(ChildElement);
 						break;
+					case "Regex":
+						ReadRegex(ChildElement);
+						break;
 					case "Trace":
 						ReadDiagnostic(ChildElement, LogEventType.Console, NewNode, ParentAgent, ControllingTrigger);
 						break;
@@ -1659,6 +1717,30 @@ namespace AutomationTool
 			return true;
 		}
 
+		/// <summary>
+		/// Constructs a regex from a regex string and returns it
+		/// </summary>
+		/// <param name="Element">The element that contains the regex</param>
+		/// <param name="Regex">The pattern to construct</param>
+		/// <returns>The regex if is valid, otherwise null</returns>
+		Regex ParseRegex(ScriptElement Element, string Regex)
+		{
+			if(Regex.Length == 0)
+			{
+				LogError(Element, "Regex is empty");
+				return null;
+			}
+			try
+			{
+				return new Regex(Regex);
+			}
+			catch(ArgumentException InvalidRegex)
+			{
+				LogError(Element, "Could not construct the Regex, reason: {0}", InvalidRegex.Message);
+				return null;
+			}
+		}
+
 		/// <summary>
 		/// Expands any properties and reads an attribute.
 		/// </summary>
