Sanitize the contents of all the receipts in the installed engine build, to make sure they don't reference any files we don't distribute.

#rnx

diff --git a/Engine/Build/InstalledEngineBuild.xml b/Engine/Build/InstalledEngineBuild.xml
--- a/Engine/Build/InstalledEngineBuild.xml
+++ b/Engine/Build/InstalledEngineBuild.xml
@@ -802,6 +802,10 @@
 			<!-- Finalize the build and then tag the text file that this process creates -->
 			<Command Name="FinalizeInstalledBuild" Arguments="-OutputDir=&quot;$(LocalInstalledDir)&quot; -ContentOnlyPlatforms=$(ContentOnlyPlatforms) $(ArchitectureCommands) -AnalyticsTypeOverride=$(AnalyticsTypeOverride)"/>
 			<Tag Files="$(LocalInstalledDir)/Engine/Build/InstalledBuild.txt" With="#Installed Build Win64 Files"/>
+			
+			<!-- Sanitize all the receipts in the target directory -->
+			<Tag Files="#Installed Build Win64 Files" Filter="*.target" With="#Installed Build Win64 Receipts"/>
+			<SanitizeReceipt Files="#Installed Build Win64 Receipts" EngineDir="$(LocalInstalledDir)\Engine"/>
 		</Node>
 	</Agent>
 
@@ -923,6 +927,10 @@
 			<!-- Finalize the build and then tag the text file that this process creates -->
 			<Command Name="FinalizeInstalledBuild" Arguments="-OutputDir=$(LocalInstalledDirMac) $(ArchitectureCommands) -AnalyticsTypeOverride=$(AnalyticsTypeOverride)"/>
 			<Tag Files="$(LocalInstalledDirMac)/Engine/Build/InstalledBuild.txt" With="#Installed Build Mac Files"/>
+
+			<!-- Sanitize all the receipts in the target directory -->
+			<Tag Files="#Installed Build Mac Files" Filter="*.target" With="#Installed Build Mac Receipts"/>
+			<SanitizeReceipt Files="#Installed Build Mac Receipts" EngineDir="$(LocalInstalledDirMac)/Engine"/>
 		</Node>
 	</Agent>
 
@@ -1004,6 +1012,10 @@
 			<!-- Finalize the build and then tag the text file that this process creates -->
 			<Command Name="FinalizeInstalledBuild" Arguments="-OutputDir=$(LocalInstalledDirLinux) -ContentOnlyPlatforms=$(ContentOnlyPlatforms) $(ArchitectureCommands) -AnalyticsTypeOverride=$(AnalyticsTypeOverride)"/>
 			<Tag Files="$(LocalInstalledDirLinux)/Engine/Build/InstalledBuild.txt" With="#Installed Build Linux Files"/>
+
+			<!-- Sanitize all the receipts in the target directory -->
+			<Tag Files="#Installed Build Linux Files" Filter="*.target" With="#Installed Build Linux Files"/>
+			<SanitizeReceipt Files="#Installed Build Linux Receipts" EngineDir="$(LocalInstalledDirLinux)/Engine"/>
 		</Node>
 	</Agent>
 
diff --git a/Engine/Build/InstalledEngineFilters.xml b/Engine/Build/InstalledEngineFilters.xml
--- a/Engine/Build/InstalledEngineFilters.xml
+++ b/Engine/Build/InstalledEngineFilters.xml
@@ -60,10 +60,6 @@
 	
 	<!-- Define Editor Filters -->
 	<Property Name="CopyEditorFilter">
-		<!-- Exclude the target files, since we're filtering the output binaries -->
-		Engine/Binaries/.../UE4Editor.target
-		Engine/Binaries/.../UE4Editor-Win64-DebugGame.target
-
 		<!-- This assembly is normally embedded into the UBT executable, but it can technically be rebuilt from an installed build -->
 		Engine/Binaries/ThirdParty/VisualStudio/...
 
@@ -154,6 +150,7 @@
 	<Property Name="CopyEditorFilter" Value="$(CopyEditorFilter);Templates/TP_VirtualRealityBP/..." If="'$(HostPlatform)' != 'Linux'"/>
 		<!-- Define Editor Exceptions -->
 	<Property Name="CopyEditorExceptions">
+	
 		<!-- Content -->
 		Engine/Content/....psd
 		Engine/Content/....pdn
diff --git a/Engine/Source/Programs/AutomationTool/BuildGraph/BuildGraph.Automation.csproj b/Engine/Source/Programs/AutomationTool/BuildGraph/BuildGraph.Automation.csproj
--- a/Engine/Source/Programs/AutomationTool/BuildGraph/BuildGraph.Automation.csproj
+++ b/Engine/Source/Programs/AutomationTool/BuildGraph/BuildGraph.Automation.csproj
@@ -80,6 +80,7 @@
     <Compile Include="Tasks\CsCompileTask.cs" />
     <Compile Include="Tasks\MoveTask.cs" />
     <Compile Include="Tasks\MsBuildTask.cs" />
+    <Compile Include="Tasks\SanitizeReceiptTask.cs" />
     <Compile Include="Tasks\SymStoreTask.cs" />
     <Compile Include="Tasks\TagReceiptTask.cs" />
     <Compile Include="Tasks\PakFileTask.cs" />
diff --git a/Engine/Source/Programs/AutomationTool/BuildGraph/Tasks/SanitizeReceiptTask.cs b/Engine/Source/Programs/AutomationTool/BuildGraph/Tasks/SanitizeReceiptTask.cs
new file mode 100644
--- /dev/null
+++ b/Engine/Source/Programs/AutomationTool/BuildGraph/Tasks/SanitizeReceiptTask.cs
@@ -0,0 +1,154 @@
+using System;
+using System.Collections.Generic;
+using System.Linq;
+using System.Text;
+using System.Threading.Tasks;
+using System.Xml;
+using Tools.DotNETCommon;
+using UnrealBuildTool;
+
+namespace AutomationTool.Tasks
+{
+	/// <summary>
+	/// Parameters for the Tag Receipt task.
+	/// </summary>
+	public class SanitizeReceiptTaskParameters
+	{
+		/// <summary>
+		/// Set of receipt files (*.target) to read, including wildcards and tag names, separated by semicolons.
+		/// </summary>
+		[TaskParameter(ValidationType = TaskParameterValidationType.FileSpec)]
+		public string Files;
+
+		/// <summary>
+		/// Path to the Engine folder, used to expand $(EngineDir) properties in receipt files. Defaults to the Engine directory for the current workspace.
+		/// </summary>
+		[TaskParameter(Optional = true)]
+		public DirectoryReference EngineDir;
+
+		/// <summary>
+		/// Path to the project folder, used to expand $(ProjectDir) properties in receipt files. Defaults to the Engine directory for the current workspace.
+		/// </summary>
+		[TaskParameter(Optional = true)]
+		public DirectoryReference ProjectDir;
+	}
+
+	/// <summary>
+	/// Task which tags build products and/or runtime dependencies by reading from *.target files.
+	/// </summary>
+	[TaskElement("SanitizeReceipt", typeof(SanitizeReceiptTaskParameters))]
+	class SanitizeReceiptTask : CustomTask
+	{
+		/// <summary>
+		/// Parameters to this task
+		/// </summary>
+		SanitizeReceiptTaskParameters Parameters;
+
+		/// <summary>
+		/// Constructor
+		/// </summary>
+		/// <param name="InParameters">Parameters to select which files to search</param>
+		public SanitizeReceiptTask(SanitizeReceiptTaskParameters InParameters)
+		{
+			Parameters = InParameters;
+		}
+
+		/// <summary>
+		/// Execute the task.
+		/// </summary>
+		/// <param name="Job">Information about the current job</param>
+		/// <param name="BuildProducts">Set of build products produced by this node.</param>
+		/// <param name="TagNameToFileSet">Mapping from tag names to the set of files they include</param>
+		public override void Execute(JobContext Job, HashSet<FileReference> BuildProducts, Dictionary<string, HashSet<FileReference>> TagNameToFileSet)
+		{
+			// Set the Engine directory
+			DirectoryReference EngineDir = Parameters.EngineDir ?? CommandUtils.EngineDirectory;
+
+			// Set the Project directory
+			DirectoryReference ProjectDir = Parameters.ProjectDir ?? EngineDir;
+
+			// Resolve the input list
+			IEnumerable<FileReference> TargetFiles = ResolveFilespec(CommandUtils.RootDirectory, Parameters.Files, TagNameToFileSet);
+			foreach(FileReference TargetFile in TargetFiles)
+			{
+				// check all files are .target files
+				if (TargetFile.GetExtension() != ".target")
+				{
+					throw new AutomationException("Invalid file passed to TagReceipt task ({0})", TargetFile.FullName);
+				}
+
+				// Print the name of the file being scanned
+				Log.TraceInformation("Sanitizing {0}", TargetFile);
+				using(new LogIndentScope("  "))
+				{
+					// Read the receipt
+					TargetReceipt Receipt;
+					if (!TargetReceipt.TryRead(TargetFile, EngineDir, ProjectDir, out Receipt))
+					{
+						CommandUtils.LogWarning("Unable to load file using TagReceipt task ({0})", TargetFile.FullName);
+						continue;
+					}
+
+					// Remove any build products that don't exist
+					List<BuildProduct> NewBuildProducts = new List<BuildProduct>(Receipt.BuildProducts.Count);
+					foreach(BuildProduct BuildProduct in Receipt.BuildProducts)
+					{
+						if(FileReference.Exists(BuildProduct.Path))
+						{
+							NewBuildProducts.Add(BuildProduct);
+						}
+						else
+						{
+							Log.TraceInformation("Removing build product: {0}", BuildProduct.Path);
+						}
+					}
+					Receipt.BuildProducts = NewBuildProducts;
+
+					// Remove any runtime dependencies that don't exist
+					RuntimeDependencyList NewRuntimeDependencies = new RuntimeDependencyList();
+					foreach(RuntimeDependency RuntimeDependency in Receipt.RuntimeDependencies)
+					{
+						if(FileReference.Exists(RuntimeDependency.Path))
+						{
+							NewRuntimeDependencies.Add(RuntimeDependency);
+						}
+						else
+						{
+							Log.TraceInformation("Removing runtime dependency: {0}", RuntimeDependency.Path);
+						}
+					}
+					Receipt.RuntimeDependencies = NewRuntimeDependencies;
+				
+					// Save the new receipt
+					Receipt.Write(TargetFile, EngineDir, ProjectDir);
+				}
+			}
+		}
+
+		/// <summary>
+		/// Output this task out to an XML writer.
+		/// </summary>
+		public override void Write(XmlWriter Writer)
+		{
+			Write(Writer, Parameters);
+		}
+
+		/// <summary>
+		/// Find all the tags which are required by this task
+		/// </summary>
+		/// <returns>The tag names which are required by this task</returns>
+		public override IEnumerable<string> FindConsumedTagNames()
+		{
+			return FindTagNamesFromFilespec(Parameters.Files);
+		}
+
+		/// <summary>
+		/// Find all the referenced tags from tasks in this task
+		/// </summary>
+		/// <returns>The tag names which are produced/modified by this task</returns>
+		public override IEnumerable<string> FindProducedTagNames()
+		{
+			return new string[0];
+		}
+	}
+}
