Moved platform files under Gauntlet/Platform and added conditional compilation

diff --git a/Engine/Source/Programs/AutomationTool/Gauntlet/Framework/Devices/Gauntlet.TargetDeviceAndroid.cs b/Engine/Source/Programs/AutomationTool/Gauntlet/Framework/Devices/Gauntlet.TargetDeviceAndroid.cs
deleted file mode 100644
diff --git a/Engine/Source/Programs/AutomationTool/Gauntlet/Framework/Devices/Gauntlet.TargetDeviceMac.cs b/Engine/Source/Programs/AutomationTool/Gauntlet/Framework/Devices/Gauntlet.TargetDeviceMac.cs
deleted file mode 100644
diff --git a/Engine/Source/Programs/AutomationTool/Gauntlet/Framework/Devices/Gauntlet.TargetDeviceNull.cs b/Engine/Source/Programs/AutomationTool/Gauntlet/Framework/Devices/Gauntlet.TargetDeviceNull.cs
deleted file mode 100644
diff --git a/Engine/Source/Programs/AutomationTool/Gauntlet/Framework/Devices/Gauntlet.TargetDeviceWindows.cs b/Engine/Source/Programs/AutomationTool/Gauntlet/Framework/Devices/Gauntlet.TargetDeviceWindows.cs
deleted file mode 100644
diff --git a/Engine/Source/Programs/AutomationTool/Gauntlet/Gauntlet.Automation.csproj b/Engine/Source/Programs/AutomationTool/Gauntlet/Gauntlet.Automation.csproj
--- a/Engine/Source/Programs/AutomationTool/Gauntlet/Gauntlet.Automation.csproj
+++ b/Engine/Source/Programs/AutomationTool/Gauntlet/Gauntlet.Automation.csproj
@@ -76,23 +76,14 @@
     <Compile Include="Framework\Base\Gauntlet.Blacklist.cs" />
     <Compile Include="Framework\Base\Gauntlet.ConfigOption.cs" />
     <Compile Include="Framework\Base\Gauntlet.BuildSource.cs" />
-    <Compile Include="Framework\Devices\Gauntlet.TargetDeviceAndroid.cs" />
-    <Compile Include="Framework\Devices\Gauntlet.TargetDeviceNull.cs" />
-    <Compile Include="Framework\Devices\Gauntlet.TargetDevicePS4.cs" />
-    <Compile Include="Framework\Devices\Gauntlet.TargetDeviceSwitch.cs" />
-    <Compile Include="Framework\Devices\Gauntlet.TargetDeviceXboxOne.cs" />
     <Compile Include="Framework\FileLockInfo.cs" />
     <Compile Include="Framework\Gauntlet.AutoParam.cs" />
     <Compile Include="Framework\Gauntlet.Exception.cs" />
     <Compile Include="Framework\Gauntlet.Params.cs" />
     <Compile Include="Framework\Gauntlet.TestGroup.cs" />
     <Compile Include="Framework\Utils\Gauntlet.MarkdownBuilder.cs" />
+    <Compile Include="Platform\Null\Gauntlet.TargetDeviceNull.cs" />
     <Compile Include="SelfTest\Unreal\Gauntlet.SelfTest.TestUnrealBuildParsing.cs" />
-    <Compile Include="Unreal\BuildSource\Gauntlet.AndroidBuildSource.cs" />
-    <Compile Include="Unreal\BuildSource\Gauntlet.PS4BuildSources.cs" />
-    <Compile Include="Unreal\BuildSource\Gauntlet.WindowsBuildSources.cs" />
-    <Compile Include="Unreal\BuildSource\Gauntlet.SwitchBuildSource.cs" />
-    <Compile Include="Unreal\BuildSource\Gauntlet.XboxOneBuildSource.cs" />
     <Compile Include="Properties\Resources.Designer.cs">
       <AutoGen>True</AutoGen>
       <DesignTime>True</DesignTime>
@@ -100,31 +91,23 @@
     </Compile>
     <Compile Include="Framework\Base\Gauntlet.TestContext.cs" />
     <Compile Include="Framework\Base\Gauntlet.AppInstance.cs" />
-    <Compile Include="Framework\Devices\Gauntlet.TargetDeviceWindows.cs" />
     <Compile Include="Framework\Devices\Gauntlet.DevicePool.cs" />
     <Compile Include="Framework\Base\Gauntlet.TestNode.cs" />
     <Compile Include="Framework\Base\Gauntlet.TargetDevice.cs" />
     <Compile Include="Framework\Gauntlet.TestExecutor.cs" />
     <Compile Include="Framework\Gauntlet.Utils.cs" />
     <Compile Include="Properties\AssemblyInfo.cs" />
-    <Compile Include="SelfTest\Devices\Gauntlet.SelfTest.TestTargetDeviceXboxOne.cs" />
     <Compile Include="SelfTest\Framework\Gauntlet.SelfTest.TestAutoParam.cs" />
     <Compile Include="SelfTest\Framework\Gauntlet.SelfTest.TestParsing.cs" />
     <Compile Include="SelfTest\Gauntlet.SelfTest.Blank.cs" />
     <Compile Include="SelfTest\Framework\Gauntlet.SelfTest.CreateGif.cs" />
     <Compile Include="SelfTest\Framework\Gauntlet.SelfTest.OrderOfOpsTest.cs" />
     <Compile Include="SelfTest\Devices\Gauntlet.SelfTest.TestTargetDevice.cs" />
-    <Compile Include="SelfTest\Devices\Gauntlet.SelfTest.TestTargetDevicePS4.cs" />
-    <Compile Include="SelfTest\Devices\Gauntlet.SelfTest.TestTargetDeviceWindows.cs" />
-    <Compile Include="SelfTest\Devices\Gauntlet.SelfTest.TestTargetDeviceMac.cs" />
     <Compile Include="SelfTest\Framework\Gauntlet.SelfTest.TimeoutTest.cs" />
     <Compile Include="SelfTest\Gauntlet.SelfTest.Globals.cs" />
     <Compile Include="SelfTest\Unreal\Gauntlet.SelfTest.LogParserTest.cs" />
     <Compile Include="SelfTest\Unreal\Gauntlet.SelfTest.TestShortSoloGame.cs" />
     <Compile Include="SelfTest\Unreal\Gauntlet.SelfTest.TestUnrealInstallAndRunBase.cs" />
-    <Compile Include="SelfTest\Unreal\Gauntlet.SelfTest.TestUnrealInstallAndRunPS4.cs" />
-    <Compile Include="SelfTest\Unreal\Gauntlet.SelfTest.TestUnrealInstallAndRunXboxOne.cs" />
-    <Compile Include="SelfTest\Unreal\Gauntlet.SelfTest.TestUnrealInstallAndRunMac.cs" />
     <Compile Include="SelfTest\Unreal\Gauntlet.SelfTest.TestUnrealSession.cs" />
     <Compile Include="SelfTest\Unreal\Gauntlet.SelfTest.TestUnrealOptions.cs" />
     <Compile Include="SelfTest\Unreal\Gauntlet.SelfTest.TestUnrealBase.cs" />
@@ -147,8 +130,6 @@
     <Compile Include="Unreal\Game\UE4Game.TestConfig.cs" />
     <Compile Include="Unreal\Utils\Gauntlet.UnrealHelpers.cs" />
     <Compile Include="Unreal\Utils\Gauntlet.UnrealLogParser.cs" />
-    <Compile Include="Framework\Devices\Gauntlet.TargetDeviceMac.cs" />
-    <Compile Include="Unreal\BuildSource\Gauntlet.MacBuildSource.cs" />
     <Compile Include="Unreal\BuildSource\Gauntlet.StagedBuildSource.cs" />
     <Compile Include="Unreal\Utils\Gauntlet.UnrealSnapshotParser.cs" />
   </ItemGroup>
@@ -193,7 +174,6 @@
   </ItemGroup>
   <ItemGroup>
     <Compile Include="SelfTest\Gauntlet.SelfTest.BaseNode.cs" />
-    <Compile Include="SelfTest\Unreal\Gauntlet.SelfTest.TestUnrealInstallAndRunWindows.cs" />
   </ItemGroup>
   <ItemGroup>
     <EmbeddedResource Include="Properties\Resources.resx">
@@ -218,5 +198,60 @@
   </ItemGroup>
   <ItemGroup />
   <ItemGroup />
+  <Choose>
+    <When Condition="Exists('Platform\Android\Gauntlet.TargetDeviceAndroid.cs')">
+      <ItemGroup>
+        <Compile Include="Platform\Android\Gauntlet.TargetDeviceAndroid.cs" />
+        <Compile Include="Platform\Android\Gauntlet.AndroidBuildSource.cs" />
+      </ItemGroup>
+    </When>
+  </Choose>
+  <Choose>
+    <When Condition="Exists('Platform\Mac\Gauntlet.TargetDeviceMac.cs')">
+      <ItemGroup>
+        <Compile Include="Platform\Mac\Gauntlet.TargetDeviceMac.cs" />
+        <Compile Include="Platform\Mac\Gauntlet.MacBuildSource.cs" />
+        <Compile Include="Platform\Mac\Gauntlet.SelfTest.TestUnrealInstallAndRunMac.cs" />
+        <Compile Include="Platform\Mac\Gauntlet.SelfTest.TestTargetDeviceMac.cs" />
+      </ItemGroup>
+    </When>
+  </Choose>
+  <Choose>
+    <When Condition="Exists('Platform\PS4\Gauntlet.TargetDevicePS4.cs')">
+      <ItemGroup>
+        <Compile Include="Platform\PS4\Gauntlet.TargetDevicePS4.cs" />
+        <Compile Include="Platform\PS4\Gauntlet.PS4BuildSources.cs" />
+        <Compile Include="Platform\PS4\Gauntlet.SelfTest.TestUnrealInstallAndRunPS4.cs" />
+        <Compile Include="Platform\PS4\Gauntlet.SelfTest.TestTargetDevicePS4.cs" />
+      </ItemGroup>
+    </When>
+  </Choose>
+  <Choose>
+    <When Condition="Exists('Platform\Switch\Gauntlet.TargetDeviceSwitch.cs')">
+      <ItemGroup>
+        <Compile Include="Platform\Switch\Gauntlet.TargetDeviceSwitch.cs" />
+        <Compile Include="Platform\Switch\Gauntlet.SwitchBuildSource.cs" />
+      </ItemGroup>
+    </When>
+  </Choose>
+  <Choose>
+    <When Condition="Exists('Platform\Windows\Gauntlet.TargetDeviceWindows.cs')">
+      <ItemGroup>
+        <Compile Include="Platform\Windows\Gauntlet.TargetDeviceWindows.cs" />
+        <Compile Include="Platform\Windows\Gauntlet.WindowsBuildSources.cs" />
+        <Compile Include="Platform\Windows\Gauntlet.SelfTest.TestTargetDeviceWindows.cs" />
+      </ItemGroup>
+    </When>
+  </Choose>
+  <Choose>
+    <When Condition="Exists('Platform\XboxOne\Gauntlet.TargetDeviceXboxOne.cs')">
+      <ItemGroup>
+        <Compile Include="Platform\XboxOne\Gauntlet.TargetDeviceXboxOne.cs" />
+        <Compile Include="Platform\XboxOne\Gauntlet.XboxOneBuildSource.cs" />
+        <Compile Include="Platform\XboxOne\Gauntlet.SelfTest.TestUnrealInstallAndRunXboxOne.cs" />
+        <Compile Include="Platform\XboxOne\Gauntlet.SelfTest.TestTargetDeviceXboxOne.cs" />
+      </ItemGroup>
+    </When>
+  </Choose>
   <Import Project="$(MSBuildToolsPath)\Microsoft.CSharp.targets" />
 </Project>
\ No newline at end of file
diff --git a/Engine/Source/Programs/AutomationTool/Gauntlet/Platform/Android/Gauntlet.AndroidBuildSource.cs b/Engine/Source/Programs/AutomationTool/Gauntlet/Platform/Android/Gauntlet.AndroidBuildSource.cs
new file mode 100644
--- /dev/null
+++ b/Engine/Source/Programs/AutomationTool/Gauntlet/Platform/Android/Gauntlet.AndroidBuildSource.cs
@@ -0,0 +1,232 @@
+﻿// Copyright 1998-2018 Epic Games, Inc. All Rights Reserved.
+
+using System;
+using System.Collections.Generic;
+using System.IO;
+using AutomationTool;
+using UnrealBuildTool;
+using System.Threading;
+using System.Text.RegularExpressions;
+using System.Linq;
+
+namespace Gauntlet
+{
+
+
+	public class AndroidBuild : IBuild
+	{
+		public UnrealTargetConfiguration Configuration { get; protected set; }
+
+		public string SourceApkPath;
+
+		public Dictionary<string, string> FilesToInstall;
+
+		public string AndroidPackageName;
+
+		public BuildFlags Flags { get; protected set; }
+
+		public UnrealTargetPlatform Platform { get { return UnrealTargetPlatform.Android; } }
+
+		public AndroidBuild(UnrealTargetConfiguration InConfig, string InAndroidPackageName, string InApkPath, Dictionary<string, string> InFilesToInstall, BuildFlags InFlags)
+		{
+			Configuration = InConfig;
+			AndroidPackageName = InAndroidPackageName;
+			SourceApkPath = InApkPath;
+			FilesToInstall = InFilesToInstall;
+			Flags = InFlags;
+		}
+
+		public bool CanSupportRole(UnrealTargetRole RoleType)
+		{
+			if (RoleType.IsClient())
+			{
+				return true;
+			}
+
+			return false;
+		}
+
+		public static IEnumerable<AndroidBuild> CreateFromPath(string InProjectName, string InPath)
+		{
+			string BuildPath = InPath;
+
+			List<AndroidBuild> DiscoveredBuilds = new List<AndroidBuild>();
+
+			DirectoryInfo Di = new DirectoryInfo(BuildPath);
+
+			
+			// find all install batchfiles
+			FileInfo[] InstallFiles = Di.GetFiles("Install_*");
+
+			foreach (FileInfo Fi in InstallFiles)
+			{
+				// todo 
+				if (Fi.FullName.Contains("armv7"))
+				{
+					continue;
+				}
+
+                // Unreal naming. In test & shipping the platform name and config will be includedn
+                string RegEx = "Install_(.+?)(client|game|server)-(Android)?-?(test|shipping)?-?(.+).bat";
+
+				Match Info = Regex.Match(Fi.Name, RegEx, RegexOptions.IgnoreCase);
+
+				if (Info.Success == false)
+				{
+					continue;
+				}
+
+				string TargetName = Info.Groups[1].ToString();
+				string TypeName = Info.Groups[2].ToString();
+				string ConfigName = Info.Groups[4].ToString();
+
+				if (string.IsNullOrEmpty(ConfigName))
+				{
+					ConfigName = "Development";
+				}
+
+
+				Log.Verbose("Pulling install data from {0}", Fi.FullName);
+
+				string AbsPath = Fi.Directory.FullName;
+
+				// read contents and replace linefeeds (regex doesn't stop on them :((
+				string BatContents = File.ReadAllText(Fi.FullName).Replace(Environment.NewLine, "\n");
+
+				// Replace .bat with .apk and strip up to and including the first _, that is then our APK name
+				string SourceApkPath = Regex.Replace(Fi.Name, ".bat", ".apk", RegexOptions.IgnoreCase);
+				SourceApkPath = SourceApkPath.Substring(SourceApkPath.IndexOf("_") + 1);
+				SourceApkPath = Path.Combine(AbsPath, SourceApkPath);
+
+				// save com.companyname.product
+				string AndroidPackageName = Regex.Match(BatContents, @"uninstall\s+(com\..+)").Groups[1].ToString();
+
+				// pull all OBBs (probably just one..)
+				var OBBMatches = Regex.Matches(BatContents, @"push\s+(.+?)\s+(.+)");
+
+				// save them as a dict of full paths as keys and dest paths as values
+				Dictionary<string, string> FilesToInstall = OBBMatches.Cast<Match>().ToDictionary(M => Path.Combine(AbsPath, M.Groups[1].ToString()), M => M.Groups[2].ToString());
+	
+				if (string.IsNullOrEmpty(SourceApkPath))
+				{
+					Log.Warning("No APK found for build at {0}", Fi.FullName);
+					continue;
+				}
+
+				if (string.IsNullOrEmpty(AndroidPackageName))
+				{
+					Log.Warning("No product name found for build at {0}", Fi.FullName);
+					continue;
+				}
+
+				UnrealTargetConfiguration UnrealConfig;
+				if (Enum.TryParse(ConfigName, true, out UnrealConfig) == false)
+				{
+					Log.Warning("Could not determine Unreal Config type from '{0}'", ConfigName);
+					continue;
+				}
+
+				// Android builds are always packaged, and we can always replace the command line
+				BuildFlags Flags = BuildFlags.Packaged | BuildFlags.CanReplaceCommandLine;
+
+				// if there's data then the pak files are in an obb and we can sub in a new exe
+				if (FilesToInstall.Count() > 0)
+				{
+					Flags |= BuildFlags.CanReplaceExecutable;
+				}
+                if (AbsPath.Contains("Bulk"))
+                {
+                    Flags |= BuildFlags.Bulk;
+                }
+
+                AndroidBuild NewBuild = new AndroidBuild(UnrealConfig, AndroidPackageName, SourceApkPath, FilesToInstall, Flags);
+
+				DiscoveredBuilds.Add(NewBuild);
+
+				Log.Verbose("Found {0} {1} build at {2}", UnrealConfig, ((Flags & BuildFlags.Bulk) == BuildFlags.Bulk) ? "(bulk)" : "(not bulk)", AbsPath);
+
+			}
+
+			return DiscoveredBuilds;
+		}
+	}
+	
+	public class AndroidBuildSource : IFolderBuildSource
+	{	
+		public string BuildName { get { return "AndroidBuildSource";  } }
+
+		public bool CanSupportPlatform(UnrealTargetPlatform InPlatform)
+		{
+			return InPlatform == UnrealTargetPlatform.Android;
+		}
+
+		public string ProjectName { get; protected set; }
+
+		public AndroidBuildSource()
+		{
+		}
+
+		public List<IBuild> GetBuildsAtPath(string InProjectName, string InPath, int MaxRecursion = 3)
+		{
+
+			List<DirectoryInfo> AllDirs = new List<DirectoryInfo>();
+
+			// c:\path\to\build
+			DirectoryInfo PathDI = new DirectoryInfo(InPath);
+
+			if (PathDI.Name.IndexOf("Android", StringComparison.OrdinalIgnoreCase) >= 0)
+			{
+				AllDirs.Add(PathDI);
+			}
+
+			// find all directories that begin with Android
+			DirectoryInfo[] AndroidDirs = PathDI.GetDirectories("Android*", SearchOption.TopDirectoryOnly);
+
+			AllDirs.AddRange(AndroidDirs);
+
+			List<DirectoryInfo> DirsToRecurse = AllDirs;
+
+			// now get subdirs
+			while (MaxRecursion-- > 0)
+			{
+				List<DirectoryInfo> DiscoveredDirs = new List<DirectoryInfo>();
+
+				DirsToRecurse.ToList().ForEach((D) =>
+				{
+					DiscoveredDirs.AddRange(D.GetDirectories("*", SearchOption.TopDirectoryOnly));
+				});
+
+				AllDirs.AddRange(DiscoveredDirs);
+				DirsToRecurse = DiscoveredDirs;
+			}
+
+			//AndroidBuildSource BuildSource = null;
+			List<IBuild> Builds = new List<IBuild>();
+
+            string AndroidBuildFilter = Globals.Params.ParseValue("AndroidBuildFilter", "");
+            foreach (DirectoryInfo Di in AllDirs)
+			{
+				IEnumerable<AndroidBuild> FoundBuilds = AndroidBuild.CreateFromPath(InProjectName, Di.FullName);
+
+				if (FoundBuilds != null)
+				{
+                    if (!string.IsNullOrEmpty(AndroidBuildFilter))
+                    {
+                        //IndexOf used because Contains must be case-sensitive
+                        FoundBuilds = FoundBuilds.Where(B => B.SourceApkPath.IndexOf(AndroidBuildFilter, StringComparison.OrdinalIgnoreCase) >= 0);
+                    }
+                    Builds.AddRange(FoundBuilds);
+				}
+			}
+
+			return Builds;
+		}
+
+		/*public AndroidBuild GetBuild(UnrealTargetConfiguration InConfig, BuildFlags InFlags)
+		{
+			return Builds.Where((B) => {
+				return B.Configuration == InConfig && (B.Flags & InFlags) > 0;
+				}).FirstOrDefault();
+		}*/
+	}
+}
\ No newline at end of file
diff --git a/Engine/Source/Programs/AutomationTool/Gauntlet/Platform/Android/Gauntlet.TargetDeviceAndroid.cs b/Engine/Source/Programs/AutomationTool/Gauntlet/Platform/Android/Gauntlet.TargetDeviceAndroid.cs
new file mode 100644
--- /dev/null
+++ b/Engine/Source/Programs/AutomationTool/Gauntlet/Platform/Android/Gauntlet.TargetDeviceAndroid.cs
@@ -0,0 +1,1281 @@
+// Copyright 1998-2018 Epic Games, Inc. All Rights Reserved.
+
+using System;
+using System.Collections.Generic;
+using System.IO;
+using AutomationTool;
+using UnrealBuildTool;
+using System.Text.RegularExpressions;
+using System.Linq;
+using System.Text;
+using System.Threading;
+
+namespace Gauntlet
+{
+	// device data from json
+	public sealed class AndroidDeviceData
+	{
+		// remote device settings (wifi)
+
+		// host of PC which is tethered 
+		public string hostIP { get; set; }
+
+		// public key 
+		public string publicKey { get; set; }
+
+		// private key
+		public string privateKey { get; set; }
+	}
+
+	// become IAppInstance when implemented enough
+	class AndroidAppInstance : IAppInstance
+	{
+		protected TargetDeviceAndroid AndroidDevice;
+
+		protected AndroidAppInstall Install;
+
+		internal IProcessResult LaunchProcess;
+
+		internal bool bHaveSavedArtifacts;
+
+		public string CommandLine { get { return Install.CommandLine; } }
+
+		public AndroidAppInstance(TargetDeviceAndroid InDevice, AndroidAppInstall InInstall, IProcessResult InProcess)
+		{
+			AndroidDevice = InDevice;
+			Install = InInstall;
+			LaunchProcess = InProcess;
+		}
+
+		public string ArtifactPath
+		{
+			get
+			{
+				if (bHaveSavedArtifacts == false)
+				{
+					if (HasExited)
+					{
+						SaveArtifacts();
+						bHaveSavedArtifacts = true;
+					}
+				}
+				
+				return Path.Combine(AndroidDevice.LocalCachePath, "Saved");
+			}
+		}
+
+		public ITargetDevice Device
+		{
+			get
+			{
+				return AndroidDevice;
+			}
+		}
+
+		public bool HasExited
+		{
+			get
+			{
+				try
+				{
+					if (!LaunchProcess.HasExited)
+					{
+						return false;
+					}
+				}
+				catch (System.InvalidOperationException)
+				{
+					return true;
+				}
+
+				return IsActivityRunning();
+			}
+		}
+
+		/// <summary>
+		/// Checks on device whether the activity is running, this is an expensive shell with output operation
+		/// the result is cached, with checks at ActivityCheckDelta seconds
+		/// </summary>
+		private bool IsActivityRunning()
+		{
+			if (ActivityExited)
+			{
+				return true;
+			}
+
+			if ((DateTime.UtcNow - ActivityCheckTime) < ActivityCheckDelta)
+			{
+				return false;
+			}
+
+			ActivityCheckTime = DateTime.UtcNow;
+
+			// get activities filtered by our package name
+			IProcessResult ActivityQuery = AndroidDevice.RunAdbDeviceCommand("shell dumpsys activity -p " + Install.AndroidPackageName + " a");
+
+			// We have exited if our activity doesn't appear in the activity query or is not the focused activity.
+			bool bActivityPresent = ActivityQuery.Output.Contains(Install.AndroidPackageName);
+			bool bActivityInForeground = ActivityQuery.Output.Contains("mResumedActivity");
+			bool bHasExited = !bActivityPresent || !bActivityInForeground;
+			if (bHasExited)
+			{
+				ActivityExited = true;
+				Log.VeryVerbose("{0}: process exited, Activity running={1}, Activity in foreground={2} ", ToString(), bActivityPresent.ToString(), bActivityInForeground.ToString());
+			}
+
+			return bHasExited;
+
+		}
+
+		private static readonly TimeSpan ActivityCheckDelta = TimeSpan.FromSeconds(10);
+		private DateTime ActivityCheckTime = DateTime.UtcNow;
+		private bool ActivityExited = false;
+
+		public bool WasKilled { get; protected set; }
+
+		/// <summary>
+		/// The output of the test activity, runs a shell command returning the full log from device per call (possibly over wifi)
+		/// result is cached, and updated at ActivityLogDelta frequency
+		/// </summary>
+		public string StdOut
+		{
+			get
+			{				
+				if (!String.IsNullOrEmpty(ActivityLogCached) && (ActivityLogTime == DateTime.MinValue || ((DateTime.UtcNow - ActivityLogTime) < ActivityLogDelta)))
+				{
+					return ActivityLogCached;
+				}
+
+				ActivityLogTime = DateTime.UtcNow;
+
+				string GetLogCommand = string.Format("shell cat {0}/Logs/{1}.log", Install.AndroidDevice.DeviceArtifactPath, Install.Name);
+				IProcessResult LogQuery = Install.AndroidDevice.RunAdbDeviceCommand(GetLogCommand, true);
+				ActivityLogCached = LogQuery.Output;
+
+				// the activity has exited, mark final log sentinel 
+				if (ActivityExited)
+				{
+					ActivityLogTime = DateTime.MinValue;
+				}
+
+				return ActivityLogCached;
+			}
+		}
+
+		private static readonly TimeSpan ActivityLogDelta = TimeSpan.FromSeconds(15);
+		private DateTime ActivityLogTime = DateTime.UtcNow;
+		private string ActivityLogCached = "";
+		
+
+		public int WaitForExit()
+		{
+			if (!HasExited)
+			{
+				LaunchProcess.WaitForExit();
+			}
+
+			return ExitCode;
+		}
+
+		public void Kill()
+		{
+			if (!HasExited)
+			{
+				WasKilled = true;
+				Install.AndroidDevice.KillRunningProcess(Install.AndroidPackageName);
+			}
+		}
+		public int ExitCode { get { return LaunchProcess.ExitCode; } }
+
+		protected void SaveArtifacts()
+		{
+			// copy remote artifacts to local
+			string LocalSaved = Path.Combine(Install.AndroidDevice.LocalCachePath, "Saved");
+			Directory.CreateDirectory(LocalSaved);
+			string ArtifactPullCommand = string.Format("pull {0} {1}", Install.AndroidDevice.DeviceArtifactPath, Install.AndroidDevice.LocalCachePath);
+			IProcessResult PullCmd = Install.AndroidDevice.RunAdbDeviceCommand(ArtifactPullCommand);
+
+			if (PullCmd.ExitCode != 0)
+			{
+				Log.Warning("Failed to retrieve artifacts. {0}", PullCmd.Output);
+			}
+
+			// pull the logcat over from device.
+			IProcessResult LogcatResult = Install.AndroidDevice.RunAdbDeviceCommand("logcat -d");
+
+			string LogcatFilename = "Logcat.log";
+			// Save logcat dump to local artifact path.
+			System.IO.File.WriteAllText(Path.Combine(LocalSaved, LogcatFilename), LogcatResult.Output);
+
+			Install.AndroidDevice.PostRunCleanup();
+		}
+	}
+
+	class AndroidAppInstall : IAppInstall
+	{
+		public string Name { get; protected set; }
+
+		public string AndroidPackageName { get; protected set; }
+
+		public TargetDeviceAndroid AndroidDevice { get; protected set; }
+
+		public ITargetDevice Device { get { return AndroidDevice; } }
+
+		public string CommandLine { get; protected set; }
+
+		public IAppInstance Run()
+		{
+			return AndroidDevice.Run(this);
+		}
+
+		public AndroidAppInstall(TargetDeviceAndroid InDevice, string InName, string InAndroidPackageName, string InCommandLine)
+		{
+			AndroidDevice = InDevice;
+			Name = InName;
+			AndroidPackageName = InAndroidPackageName;
+			CommandLine = InCommandLine;
+		}
+	}
+
+	public class DefaultAndroidDevices : IDefaultDeviceSource
+	{
+		public bool CanSupportPlatform(UnrealTargetPlatform Platform)
+		{
+			return Platform == UnrealTargetPlatform.Android;
+		}
+
+		public ITargetDevice[] GetDefaultDevices()
+		{
+			return TargetDeviceAndroid.GetDefaultDevices();
+		}
+	}
+
+	public class AndroidDeviceFactory : IDeviceFactory
+	{
+		public bool CanSupportPlatform(UnrealTargetPlatform Platform)
+		{
+			return Platform == UnrealTargetPlatform.Android;
+		}
+
+		public ITargetDevice CreateDevice(string InRef, string InParam)
+		{
+			AndroidDeviceData DeviceData = null;
+
+			if (!String.IsNullOrEmpty(InParam))
+			{
+				DeviceData = fastJSON.JSON.Instance.ToObject<AndroidDeviceData>(InParam);
+			}
+
+			return new TargetDeviceAndroid(InRef, DeviceData);
+		}
+	}
+
+	/// <summary>
+	/// Android implementation of a device that can run applications
+	/// </summary>
+	public class TargetDeviceAndroid : ITargetDevice
+	{ 
+		/// <summary>
+		/// Friendly name for this target
+		/// </summary>
+		public string Name { get; protected set; }
+	
+		/// <summary>
+		/// Low-level device name
+		/// </summary>
+		public string DeviceName { get; protected set; }
+
+		/// <summary>
+		/// Platform type.
+		/// </summary>
+		public UnrealTargetPlatform Platform { get { return UnrealTargetPlatform.Android; } }
+
+		/// <summary>
+		/// Options for executing commands
+		/// </summary>
+		public CommandUtils.ERunOptions RunOptions { get; set; }
+
+		/// <summary>
+		/// Temp path we use to push/pull things from the device
+		/// </summary>
+		public string LocalCachePath { get; protected set; }
+
+
+		/// <summary>
+		/// Artifact (e.g. Saved) path on the device
+		/// </summary>
+		public string DeviceArtifactPath { get; protected set;  }
+
+		/// <summary>
+		/// Path to a command line if installed
+		/// </summary>
+		protected string CommandLineFilePath { get; set; }
+
+		public bool IsAvailable
+		{
+			get
+			{
+                // ensure our device is present in 'adb devices' output.
+				var AllDevices = GetAllConnectedDevices();
+
+				if (AllDevices.Keys.Contains(DeviceName) == false)
+				{
+					return false;
+				}
+
+				if (AllDevices[DeviceName] == false)
+				{
+					Log.Warning("Device {0} is connected but we are not authorized", DeviceName);
+					return false;
+				}
+
+				// any device will do, but only one at a time.
+				return true;
+            }
+		}
+
+		protected Dictionary<EIntendedBaseCopyDirectory, string> LocalDirectoryMappings { get; set; }
+        void SetUpDirectoryMappings()
+        {
+            LocalDirectoryMappings = new Dictionary<EIntendedBaseCopyDirectory, string>();
+        }
+
+        public void PopulateDirectoryMappings(string ProjectDir)
+        {
+            LocalDirectoryMappings.Add(EIntendedBaseCopyDirectory.Binaries, Path.Combine(ProjectDir, "Binaries"));
+            LocalDirectoryMappings.Add(EIntendedBaseCopyDirectory.Config, Path.Combine(ProjectDir, "Config"));
+            LocalDirectoryMappings.Add(EIntendedBaseCopyDirectory.Content, Path.Combine(ProjectDir, "Content"));
+            LocalDirectoryMappings.Add(EIntendedBaseCopyDirectory.Demos, Path.Combine(ProjectDir, "Demos"));
+            LocalDirectoryMappings.Add(EIntendedBaseCopyDirectory.Profiling, Path.Combine(ProjectDir, "Profiling"));
+            LocalDirectoryMappings.Add(EIntendedBaseCopyDirectory.Saved, ProjectDir);
+        }
+        public bool IsConnected { get	{ return IsAvailable; }	}
+
+		protected bool IsExistingDevice = false;		
+
+		/// <summary>
+		/// Constructor
+		/// </summary>
+		/// <param name="InReferenceName"></param>
+		/// <param name="InRemoveOnDestruction"></param>
+		public TargetDeviceAndroid(string InDeviceName = "", AndroidDeviceData DeviceData = null)
+		{
+			DeviceName = InDeviceName;
+
+			AdbCredentialCache.AddInstance(DeviceData);
+
+			// If no device name or its 'default' then use the first default device
+			if (string.IsNullOrEmpty(DeviceName) || DeviceName.Equals("default", StringComparison.OrdinalIgnoreCase))
+			{
+				var DefaultDevices = GetAllAvailableDevices();	
+
+				if (DefaultDevices.Count() == 0)
+				{
+					if (GetAllConnectedDevices().Count > 0)
+					{
+						throw new AutomationException("No default device available. One or more devices are connected but unauthorized. See 'adb devices'");
+					}
+					else
+					{
+						throw new AutomationException("No default device available. See 'adb devices'");
+					}
+				}
+
+				DeviceName = DefaultDevices.First();
+			}
+
+			if (Log.IsVerbose)
+			{
+				RunOptions = CommandUtils.ERunOptions.NoWaitForExit;
+			}
+			else
+			{
+				RunOptions = CommandUtils.ERunOptions.NoWaitForExit | CommandUtils.ERunOptions.NoLoggingOfRunCommand;
+			}
+
+			// if this is not a connected device then remove when done
+			var ConnectedDevices = GetAllConnectedDevices();
+
+			IsExistingDevice = ConnectedDevices.Keys.Contains(DeviceName);
+
+			if (!IsExistingDevice)
+			{
+				// adb uses 5555 by default
+				if (DeviceName.Contains(":") == false)
+				{
+					DeviceName = DeviceName + ":5555";
+				}
+
+				lock (Globals.MainLock)
+				{
+					using (var PauseEC = new ScopedSuspendECErrorParsing())
+					{
+						IProcessResult AdbResult = RunAdbGlobalCommand(string.Format("connect {0}", DeviceName));
+
+						if (AdbResult.ExitCode != 0)
+						{
+							throw new AutomationException("adb failed to connect to {0}. {1}", DeviceName, AdbResult.Output);
+						}
+					}
+
+					Log.Info("Connected to {0}", DeviceName);
+
+					// Need to sleep for adb service process to register, otherwise get an unauthorized (especially on parallel device use)
+					Thread.Sleep(5000);
+				}
+			}
+
+			LocalDirectoryMappings = new Dictionary<EIntendedBaseCopyDirectory, string>();
+
+			// for IP devices need to sanitize this
+			Name = DeviceName.Replace(":", "_");
+
+			LocalCachePath = Path.Combine(Path.GetTempPath(), "AndroidDevice_" + Name);
+
+			ConnectedDevices = GetAllConnectedDevices();
+
+            SetUpDirectoryMappings();
+
+            // sanity check that it was now dound
+            if (ConnectedDevices.Keys.Contains(DeviceName) == false)
+			{
+				throw new AutomationException("Failed to find new device {0} in connection list", DeviceName);
+			}
+
+			if (ConnectedDevices[DeviceName] == false)
+			{
+				Dispose();
+				throw new AutomationException("Device {0} is connected but this PC is not authorized.", DeviceName);
+			}
+		}
+
+		~TargetDeviceAndroid()
+		{
+			Dispose(false);
+		}
+
+		#region IDisposable Support
+		private bool disposedValue = false; // To detect redundant calls
+
+		protected virtual void Dispose(bool disposing)
+		{
+			if (!disposedValue)
+			{
+				try
+				{
+					if (!IsExistingDevice)
+					{
+						// disconnect
+						RunAdbGlobalCommand(string.Format("disconnect {0}", DeviceName));
+
+						Log.Info("Disconnected {0}", DeviceName);
+					}
+
+					if (Directory.Exists(LocalCachePath))
+					{
+						Directory.Delete(LocalCachePath, true);
+					}
+				}
+				catch (Exception Ex)
+				{
+					Log.Warning("TargetDeviceAndroid.Dispose() threw: {0}", Ex.Message);
+				}
+				finally
+				{
+					disposedValue = true;
+					AdbCredentialCache.RemoveInstance();					
+				}
+
+			}
+		}
+
+		// This code added to correctly implement the disposable pattern.
+		public void Dispose()
+		{
+			// Do not change this code. Put cleanup code in Dispose(bool disposing) above.
+			Dispose(true);
+			// TODO: uncomment the following line if the finalizer is overridden above.
+			// GC.SuppressFinalize(this);
+		}
+		#endregion
+
+		/// <summary>
+		/// Returns a list of locally connected devices (e.g. 'adb devices'). 
+		/// </summary>
+		/// <returns></returns>
+		static private Dictionary<string, bool> GetAllConnectedDevices()
+		{
+           var Result = RunAdbGlobalCommand("devices");
+
+            MatchCollection DeviceMatches = Regex.Matches(Result.Output, @"^([\d\w\.\:]{6,32})\s+(\w+)", RegexOptions.Multiline);
+
+            var DeviceList = DeviceMatches.Cast<Match>().ToDictionary(
+                M => M.Groups[1].ToString(),
+                M => !M.Groups[2].ToString().ToLower().Contains("unauthorized")
+            );
+
+            return DeviceList;
+		}
+
+		static private IEnumerable<string> GetAllAvailableDevices()
+		{
+			var AllDevices = GetAllConnectedDevices();
+			return AllDevices.Keys.Where(D => AllDevices[D] == true);
+		}
+
+		static public ITargetDevice[] GetDefaultDevices()
+		{
+			var Result = RunAdbGlobalCommand("devices");
+
+			MatchCollection DeviceMatches = Regex.Matches(Result.Output, @"([\d\w\.\:]{8,32})\s+device");
+
+			List<ITargetDevice> Devices = new List<ITargetDevice>();
+
+			foreach (string Device in GetAllAvailableDevices())
+			{
+				ITargetDevice NewDevice = new TargetDeviceAndroid(Device);
+				Devices.Add(NewDevice);
+			}
+
+			return Devices.ToArray();
+		}
+
+		internal void PostRunCleanup()
+		{
+			// Delete the commandline file, if someone installs an APK on top of ours
+			// they will get very confusing behavior...
+			if (string.IsNullOrEmpty(CommandLineFilePath) == false)
+			{
+				Log.Verbose("Removing {0}", CommandLineFilePath);
+				DeleteFileFromDevice(CommandLineFilePath);
+				CommandLineFilePath = null;
+			}
+		}
+
+		public bool IsOn
+		{
+			get
+			{
+				string CommandLine = "shell dumpsys power";
+				IProcessResult OnAndUnlockedQuery = RunAdbDeviceCommand(CommandLine);
+
+				return OnAndUnlockedQuery.Output.Contains("mHoldingDisplaySuspendBlocker=true")
+					&& OnAndUnlockedQuery.Output.Contains("mHoldingWakeLockSuspendBlocker=true");
+			}
+		}
+
+		public bool PowerOn()
+		{
+			Log.Verbose("{0}: Powering on", ToString());
+			string CommandLine = "shell \"input keyevent KEYCODE_WAKEUP && input keyevent KEYCODE_MENU\"";
+			RunAdbDeviceCommand(CommandLine);
+			return true;
+		}
+		public bool PowerOff()
+		{
+			Log.Verbose("{0}: Powering off", ToString());
+
+			string CommandLine = "shell \"input keyevent KEYCODE_SLEEP\"";
+			RunAdbDeviceCommand(CommandLine);
+			return true;
+		}
+
+		public bool Reboot()
+		{
+			return true;
+		}
+
+		public bool Connect()
+		{
+			AllowDeviceSleepState(true);
+			return true;
+		}
+
+		public bool Disconnect()
+		{
+			AllowDeviceSleepState(false);
+			return true;
+		}
+
+		public override string ToString()
+		{
+            // TODO: device id
+			if (Name == DeviceName)
+			{
+				return Name;
+			}
+			return string.Format("{0} ({1})", Name, DeviceName);
+		}
+
+		protected bool DeleteFileFromDevice(string DestPath)
+		{
+			var AdbResult = RunAdbDeviceCommand(string.Format("shell rm -f {0}", DestPath));
+			return AdbResult.ExitCode == 0;
+		}
+
+		protected bool CopyFileToDevice(string PackageName, string SourcePath, string DestPath, bool IgnoreDependencies = false)
+		{
+			bool IsAPK = string.Equals(Path.GetExtension(SourcePath), ".apk", StringComparison.OrdinalIgnoreCase);
+
+            // for the APK there's no easy/reliable way to get the date of the version installed, so 
+            // we write this out to a dependency file in the demote dir and check it each time.
+            // current file time
+            DateTime LocalModifiedTime = File.GetLastWriteTime(SourcePath);
+
+            string QuotedSourcePath = SourcePath;
+            if (SourcePath.Contains(" "))
+            {
+                QuotedSourcePath = '"' + SourcePath + '"';
+            }
+
+            // dependency info is a hash of the destination name, saved under a folder on /sdcard
+            int DestHash = DestPath.GetHashCode();
+			string DependencyCacheDir = "/sdcard/gdeps";
+			string DepFile = string.Format("{0}/{1:X}", DependencyCacheDir, DestHash);	
+
+			IProcessResult AdbResult = null;
+
+	
+			// get info from the device about this file
+			string CurrentFileInfo = null;
+
+			if (IsAPK)
+			{
+				// for APK query the package info and get the update time
+				AdbResult = RunAdbDeviceCommand(string.Format("shell dumpsys package {0} | grep lastUpdateTime", PackageName));
+
+				if (AdbResult.ExitCode == 0)
+				{
+					CurrentFileInfo = AdbResult.Output.ToString().Trim();
+				}
+			}
+			else
+			{
+				// for other files get the file info
+				AdbResult = RunAdbDeviceCommand(string.Format("shell ls -l {0}", DestPath));
+
+				if (AdbResult.ExitCode == 0)
+				{
+					CurrentFileInfo = AdbResult.Output.ToString().Trim();
+				}
+			}
+
+			bool SkipInstall = false;
+
+			// If this is valid then there is some form of that file on the device, now figure out if it matches the 
+			if (string.IsNullOrEmpty(CurrentFileInfo) == false)
+			{
+				// read the dep file
+				AdbResult = RunAdbDeviceCommand(string.Format("shell cat {0}", DepFile));
+
+				if (AdbResult.ExitCode == 0)
+				{
+					// Dependency info is the modified time of the source, and the post-copy file stats of the installed file, separated by ###
+					string[] DepLines = AdbResult.Output.ToString().Split(new[] { "###" }, StringSplitOptions.RemoveEmptyEntries).Select(S => S.Trim()).ToArray();
+
+					if (DepLines.Length >= 2)
+					{
+						string InstalledSourceModifiedTime = DepLines[0];
+						string InstalledFileInfo = DepLines[1];
+
+						if (InstalledSourceModifiedTime == LocalModifiedTime.ToString()
+							&& CurrentFileInfo == InstalledFileInfo)
+						{
+							SkipInstall = true;
+						}
+					}
+				}
+			}
+
+			if (SkipInstall && IgnoreDependencies == false)
+			{
+				Log.Info("Skipping install of {0} - remote file up to date", Path.GetFileName(SourcePath));
+			}
+			else
+			{
+				if (IsAPK)
+				{
+					// we need to ununstall then install the apk - don't care if it fails, may have been deleted
+					string AdbCommand = string.Format("uninstall {0}", PackageName);
+					AdbResult = RunAdbDeviceCommand(AdbCommand);
+
+					Log.Info("Installing {0} to {1}", SourcePath, Name);
+
+					AdbCommand = string.Format("install {0}", QuotedSourcePath);
+					AdbResult = RunAdbDeviceCommand(AdbCommand);
+
+					if (AdbResult.ExitCode != 0)
+					{
+						throw new AutomationException("Failed to install {0}. Error {1}", SourcePath, AdbResult.Output);
+					}
+
+					// for APK query the package info and get the update time
+					AdbResult = RunAdbDeviceCommand(string.Format("shell dumpsys package {0} | grep lastUpdateTime", PackageName));
+					CurrentFileInfo = AdbResult.Output.ToString().Trim();
+				}
+				else
+				{
+					Log.Info("Installing {0} to {1} via adb push", SourcePath, DestPath);
+					string AdbCommand = string.Format("push {0} {1}", QuotedSourcePath, DestPath);
+					AdbResult = RunAdbDeviceCommand(AdbCommand);
+
+					if (AdbResult.ExitCode != 0)
+					{
+						throw new AutomationException("Failed to push {0} to device. Error {1}", SourcePath, AdbResult.Output);
+					}
+
+					// Now pull info about the file which we'll write as a dep
+					AdbResult = RunAdbDeviceCommand(string.Format("shell ls -l {0}", DestPath));
+					CurrentFileInfo = AdbResult.Output.ToString().Trim();
+				}
+
+				// write the actual dependency info
+				string DepContents = LocalModifiedTime + "###" + CurrentFileInfo;
+
+				// save last modified time to remote deps after success
+				RunAdbDeviceCommand(string.Format("shell mkdir -p {0}", DependencyCacheDir));
+
+				string Cmd = string.Format("shell echo \"{0}\" > {1}", DepContents, DepFile);
+				AdbResult = RunAdbDeviceCommand(Cmd);
+
+				if (AdbResult.ExitCode != 0)
+				{
+					Log.Warning("Failed to write dependency file {0}", DepFile);
+				}
+			}
+
+			return true;
+		}
+
+		public IAppInstall InstallApplication(UnrealAppConfig AppConfig)
+		{
+			// todo - pass this through
+			AndroidBuild Build = AppConfig.Build as AndroidBuild;
+
+			// Ensure APK exists
+			if (Build == null)
+			{
+				throw new AutomationException("Invalid build for Android!");
+			}
+
+			// kill any currently running instance:
+			KillRunningProcess(Build.AndroidPackageName);
+
+			bool SkipDeploy = Globals.Params.ParseParam("SkipDeploy");
+
+			if (SkipDeploy == false)
+			{
+				// Establish remote directory locations
+				string DeviceStorageQueryCommand = AndroidPlatform.GetStorageQueryCommand();
+				IProcessResult StorageQueryResult = RunAdbDeviceCommand(DeviceStorageQueryCommand);
+				string StorageLocation = StorageQueryResult.Output.Trim(); // "/mnt/sdcard";
+
+				// remote dir used to save things
+				string RemoteDir = StorageLocation + "/UE4Game/" + AppConfig.ProjectName;
+
+				// if not a bulk/dev build, remote dir will be under /{StorageLocation}/Android/data/{PackageName}
+				if ((Build.Flags & ( BuildFlags.Bulk | BuildFlags.CanReplaceExecutable)) == 0)
+				{
+					RemoteDir = StorageLocation + "/Android/data/" + Build.AndroidPackageName + "/files/UE4Game/" + AppConfig.ProjectName;
+				}
+				
+				string DependencyDir = RemoteDir + "/deps";
+
+				// device artifact path, always clear between runs
+				DeviceArtifactPath = string.Format("{0}/{1}/Saved", RemoteDir, AppConfig.ProjectName);
+				RunAdbDeviceCommand(string.Format("shell rm -r {0}", DeviceArtifactPath));
+
+				// path for OBB files
+				string OBBRemoteDestination = string.Format("{0}/obb/{1}", StorageLocation, Build.AndroidPackageName);
+
+				if (Globals.Params.ParseParam("cleandevice"))
+				{
+					Log.Info("Cleaning previous builds due to presence of -cleandevice");
+
+					// we need to ununstall then install the apk - don't care if it fails, may have been deleted
+					Log.Info("Uninstalling {0}", Build.AndroidPackageName);
+					RunAdbDeviceCommand(string.Format("uninstall {0}", Build.AndroidPackageName));
+
+					Log.Info("Removing {0}", RemoteDir);
+					RunAdbDeviceCommand(string.Format("shell rm -r {0}", RemoteDir));
+
+					Log.Info("Removing {0}", OBBRemoteDestination);
+					RunAdbDeviceCommand(string.Format("shell rm -r {0}", OBBRemoteDestination));
+				}
+
+				// remote dir on the device, create it if it doesn't exist
+				RunAdbDeviceCommand(string.Format("shell mkdir -p {0}/", RemoteDir));
+
+				IProcessResult AdbResult;
+				string AdbCommand;
+
+				// path to the APK to install.
+				string ApkPath = Build.SourceApkPath;
+
+				// check for a local newer executable
+				if (Globals.Params.ParseParam("dev"))
+				{
+					//string ApkFileName = Path.GetFileName(ApkPath);
+
+					string ApkFileName2 = UnrealHelpers.GetExecutableName(AppConfig.ProjectName, UnrealTargetPlatform.Android, AppConfig.Configuration, AppConfig.ProcessType, "apk");
+
+					string LocalAPK = Path.Combine(Environment.CurrentDirectory, AppConfig.ProjectName, "Binaries/Android", ApkFileName2);
+
+					bool LocalFileExists = File.Exists(LocalAPK);
+					bool LocalFileNewer = LocalFileExists && File.GetLastWriteTime(LocalAPK) > File.GetLastWriteTime(ApkPath);
+
+					Log.Verbose("Checking for newer binary at {0}", LocalAPK);
+					Log.Verbose("LocalFile exists: {0}. Newer: {1}", LocalFileExists, LocalFileNewer);
+
+					if (LocalFileExists && LocalFileNewer)
+					{
+						ApkPath = LocalAPK;
+					}
+				}
+
+				// first install the APK
+				CopyFileToDevice(Build.AndroidPackageName, ApkPath, "");
+
+				// obb files need to be named based on APK version (grrr), so find that out. This should return something like
+				// versionCode=2 minSdk=21 targetSdk=21
+				string PackageInfo = RunAdbDeviceCommand(string.Format("shell dumpsys package {0} | grep versionCode", Build.AndroidPackageName)).Output;
+				var Match = Regex.Match(PackageInfo, @"versionCode=([\d\.]+)\s");
+				if (Match.Success == false)
+				{
+					throw new AutomationException("Failed to find version info for APK!");
+				}
+				string PackageVersion = Match.Groups[1].ToString();
+
+				// Convert the files from the source to final destination names
+				Dictionary<string, string> FilesToInstall = new Dictionary<string, string>();
+
+                Console.WriteLine("trying to copy files over.");
+                if (AppConfig.FilesToCopy != null)
+                {
+                    if (LocalDirectoryMappings.Count == 0)
+                    {
+                        Console.WriteLine("Populating Directory");
+                        PopulateDirectoryMappings(DeviceArtifactPath);
+                    }
+                    Console.WriteLine("trying to copy files over.");
+                    foreach (UnrealFileToCopy FileToCopy in AppConfig.FilesToCopy)
+                    {
+                        string PathToCopyTo = Path.Combine(LocalDirectoryMappings[FileToCopy.TargetBaseDirectory], FileToCopy.TargetRelativeLocation);
+                        if (File.Exists(FileToCopy.SourceFileLocation))
+                        {
+                            FileInfo SrcInfo = new FileInfo(FileToCopy.SourceFileLocation);
+                            SrcInfo.IsReadOnly = false;
+                            FilesToInstall.Add(FileToCopy.SourceFileLocation, PathToCopyTo.Replace("\\", "/"));
+                            Console.WriteLine("Copying {0} to {1}", FileToCopy.SourceFileLocation, PathToCopyTo);
+                        }
+
+                        else
+                        {
+                            Log.Warning("File to copy {0} not found", FileToCopy);
+                        }
+                    }
+                }
+
+                Build.FilesToInstall.Keys.ToList().ForEach(K =>
+				{
+
+					string SrcPath = K;
+					string DestPath = Build.FilesToInstall[K];
+
+					string DestFile = Path.GetFileName(DestPath);
+
+					// If we installed a new APK we need to change the package version
+					Match OBBMatch = Regex.Match(DestFile, @"main\.(\d+)\.com.*\.obb");
+					if (OBBMatch.Success)
+					{
+						string NewFileName = DestFile.Replace(OBBMatch.Groups[1].ToString(), PackageVersion);
+						DestPath = DestPath.Replace(DestFile, NewFileName);
+					}
+
+					DestPath = Regex.Replace(DestPath, "%STORAGE%", StorageLocation, RegexOptions.IgnoreCase);
+
+					FilesToInstall.Add(SrcPath, DestPath);
+				});
+
+
+
+                // get a list of files in the destination OBB directory
+                AdbResult = RunAdbDeviceCommand(string.Format("shell ls {0}", OBBRemoteDestination));
+
+				// if != 0 then no folder exists
+				if (AdbResult.ExitCode == 0)
+				{
+					IEnumerable<string> CurrentRemoteFileList = AdbResult.Output.Replace("\r\n", "\n").Split('\n');
+					IEnumerable<string> NewRemoteFileList = FilesToInstall.Values.Select(F => Path.GetFileName(F));
+
+					// delete any files that should not be there
+					foreach (string FileName in CurrentRemoteFileList)
+					{
+						if (FileName.StartsWith(".") || FileName.Length == 0)
+						{
+							continue;
+						}
+
+						if (NewRemoteFileList.Contains(FileName) == false)
+						{
+							RunAdbDeviceCommand(string.Format("shell rm {0}/{1}", OBBRemoteDestination, FileName));
+						}
+					}
+				}
+
+				foreach (var KV in FilesToInstall)
+				{
+					string LocalFile = KV.Key;
+					string RemoteFile = KV.Value;
+
+					CopyFileToDevice(Build.AndroidPackageName, LocalFile, RemoteFile);
+				}
+
+				// create a tempfile, insert the command line, and push it over
+				string TmpFile = Path.GetTempFileName();
+
+				CommandLineFilePath = string.Format("{0}/UE4CommandLine.txt", RemoteDir);
+
+				// I've seen a weird thing where adb push truncates by a byte, so add some padding...
+				File.WriteAllText(TmpFile, AppConfig.CommandLine + "    ");
+				AdbCommand = string.Format("push {0} {1}", TmpFile, CommandLineFilePath);
+				RunAdbDeviceCommand(AdbCommand);
+
+
+				File.Delete(TmpFile);
+			}
+			else
+			{
+				Log.Info("Skipping install of {0} (-skipdeploy)", Build.AndroidPackageName);
+			}
+
+			AndroidAppInstall AppInstall = new AndroidAppInstall(this, AppConfig.ProjectName, Build.AndroidPackageName, AppConfig.CommandLine);
+
+			return AppInstall;
+		}
+
+		public IAppInstance Run(IAppInstall App)
+		{
+			AndroidAppInstall DroidAppInstall = App as AndroidAppInstall;
+
+			if (DroidAppInstall == null)
+			{
+				throw new Exception("AppInstance is of incorrect type!");
+			}
+
+			// wake the device - we can install while its asleep but not run
+			PowerOn();
+
+			// kill any currently running instance:
+			KillRunningProcess(DroidAppInstall.AndroidPackageName);
+
+			string LaunchActivity = AndroidPlatform.GetLaunchableActivityName();
+
+			Log.Info("Launching {0} on '{1}' ", DroidAppInstall.AndroidPackageName + "/" + LaunchActivity, ToString());
+			Log.Verbose("\t{0}", DroidAppInstall.CommandLine);
+
+			// Clear the device's logcat in preparation for the test..
+			RunAdbDeviceCommand("logcat --clear");
+
+			// start the app on device!
+			string CommandLine = "shell am start -W -S -n " + DroidAppInstall.AndroidPackageName + "/" + LaunchActivity;
+			IProcessResult Process = RunAdbDeviceCommand(CommandLine, false, true);
+
+			return new AndroidAppInstance(this, DroidAppInstall, Process);
+		}
+
+		/// <summary>
+		/// Runs an ADB command, automatically adding the name of the current device to
+		/// the arguments sent to adb
+		/// </summary>
+		/// <param name="Args"></param>
+		/// <param name="Wait"></param>
+		/// <param name="Input"></param>
+		/// <returns></returns>
+		public IProcessResult RunAdbDeviceCommand(string Args, bool Wait=true, bool bShouldLogCommand = false)
+		{
+			if (string.IsNullOrEmpty(DeviceName) == false)
+			{
+				Args = string.Format("-s {0} {1}", DeviceName, Args);
+			}
+
+			return RunAdbGlobalCommand(Args, Wait, bShouldLogCommand);
+		}
+
+		/// <summary>
+		/// Runs an ADB command, automatically adding the name of the current device to
+		/// the arguments sent to adb
+		/// </summary>
+		/// <param name="Args"></param>
+		/// <param name="Wait"></param>
+		/// <param name="Input"></param>
+		/// <returns></returns>
+		public string RunAdbDeviceCommandAndGetOutput(string Args)
+		{
+			if (string.IsNullOrEmpty(DeviceName) == false)
+			{
+				Args = string.Format("-s {0} {1}", DeviceName, Args);
+			}
+
+			IProcessResult Result = RunAdbGlobalCommand(Args);
+
+			if (Result.ExitCode != 0)
+			{
+				throw new AutomationException("adb command {0} failed. {1}", Args, Result.Output);
+			}
+
+			return Result.Output;
+		}
+
+		/// <summary>
+		/// Runs an ADB command at the global scope
+		/// </summary>
+		/// <param name="Args"></param>
+		/// <param name="Wait"></param>
+		/// <returns></returns>
+		public static IProcessResult RunAdbGlobalCommand(string Args, bool Wait = true, bool bShouldLogCommand = false)
+		{
+			CommandUtils.ERunOptions RunOptions = CommandUtils.ERunOptions.AppMustExist;
+
+			if (Log.IsVeryVerbose)
+			{
+				RunOptions |= CommandUtils.ERunOptions.AllowSpew;
+			}
+			else
+			{
+				RunOptions |= CommandUtils.ERunOptions.NoLoggingOfRunCommand;
+			}
+
+			if (Wait == false)
+			{
+				RunOptions |= CommandUtils.ERunOptions.NoWaitForExit;
+			}
+
+			if (bShouldLogCommand)
+			{
+				Log.Verbose("Running ADB Command: adb {0}", Args);
+			}
+			
+			IProcessResult Process = AndroidPlatform.RunAdbCommand(null, null, Args, null, RunOptions);
+			return Process;
+		}
+
+		public void AllowDeviceSleepState(bool bAllowSleep)
+		{
+			string CommandLine = "shell svc power stayon " + (bAllowSleep ? "false" : "usb");
+			RunAdbDeviceCommand(CommandLine);
+		}
+
+		public void KillRunningProcess(string AndroidPackageName)
+		{
+			Log.Verbose("{0}: Killing process '{1}' ", ToString(), AndroidPackageName);
+			string KillProcessCommand = string.Format("shell am force-stop {0}", AndroidPackageName);
+			RunAdbDeviceCommand(KillProcessCommand);
+		}
+
+		public Dictionary<EIntendedBaseCopyDirectory, string> GetPlatformDirectoryMappings()
+		{
+			if (LocalDirectoryMappings.Count == 0)
+			{
+				Log.Warning("Platform directory mappings have not been populated for this platform! This should be done within InstallApplication()");
+			}
+			return LocalDirectoryMappings;
+		}
+	}
+
+	/// <summary>
+	/// ADB key credentials, running adb-server commands (must) use same pub/private key store
+	/// </summary>
+	internal static class AdbCredentialCache
+	{
+
+		private static int InstanceCount = 0;
+		private static bool bUsingCustomKeys = false;
+
+		private static string PrivateKey;
+		private static string PublicKey;
+
+		private const string KeyBackupExt = ".gauntlet.bak";
+
+		private static void Reset()
+		{
+			if (InstanceCount != 0)
+			{
+				throw new AutomationException("AdbCredentialCache.Reset() called with outstanding instances");
+			}
+
+			PrivateKey = PublicKey = String.Empty;
+			bUsingCustomKeys = false;
+
+			RestoreBackupKeys();
+		}
+
+		public static void AddInstance(AndroidDeviceData DeviceData = null)
+		{
+			lock (Globals.MainLock)
+			{
+				string KeyPath = Globals.Params.ParseValue("adbkeys", null);
+
+				// setup key store from device data
+				if (String.IsNullOrEmpty(KeyPath) && DeviceData != null)
+				{
+					// checked that cached keys are the same
+					if (!String.IsNullOrEmpty(PrivateKey))
+					{
+						if (PrivateKey != DeviceData.privateKey)
+						{
+							throw new AutomationException("ADB device private keys must match");
+						}
+					}
+
+					if (!String.IsNullOrEmpty(PublicKey))
+					{
+						if (PublicKey != DeviceData.publicKey)
+						{
+							throw new AutomationException("ADB device public keys must match");
+						}
+					}
+
+					PrivateKey = DeviceData.privateKey;
+					PublicKey = DeviceData.publicKey;
+
+					if (String.IsNullOrEmpty(PublicKey) || String.IsNullOrEmpty(PrivateKey))
+					{
+						throw new AutomationException("Invalid key in device data");
+					}
+
+					KeyPath = Path.Combine(Globals.TempDir, "AndroidADBKeys");
+
+					if (!Directory.Exists(KeyPath))
+					{
+						Directory.CreateDirectory(KeyPath);
+					}
+
+					if (InstanceCount == 0)
+					{
+						byte[] data = Convert.FromBase64String(PrivateKey);
+						File.WriteAllText(KeyPath + "/adbkey", Encoding.UTF8.GetString(data));
+
+						data = Convert.FromBase64String(PublicKey);
+						File.WriteAllText(KeyPath + "/adbkey.pub", Encoding.UTF8.GetString(data));
+					}
+
+				}
+
+				if (InstanceCount == 0 && !String.IsNullOrEmpty(KeyPath))
+				{
+
+					Log.Info("Using adb keys at {0}", KeyPath);
+
+					string LocalKeyPath = Path.Combine(Environment.GetEnvironmentVariable("USERPROFILE"), ".android");
+
+					string RemoteKeyFile = Path.Combine(KeyPath, "adbkey");
+					string RemotePubKeyFile = Path.Combine(KeyPath, "adbkey.pub");
+					string LocalKeyFile = Path.Combine(LocalKeyPath, "adbkey");
+					string LocalPubKeyFile = Path.Combine(LocalKeyPath, "adbkey.pub");
+					string BackupSentry = Path.Combine(LocalKeyPath, "gauntlet.inuse");
+
+					if (File.Exists(RemoteKeyFile) == false)
+					{
+						throw new AutomationException("adbkey at {0} does not exist", KeyPath);
+					}
+
+					if (File.Exists(RemotePubKeyFile) == false)
+					{
+						throw new AutomationException("adbkey.pub at {0} does not exist", KeyPath);
+					}
+
+					if (File.Exists(BackupSentry) == false)
+					{
+						if (File.Exists(LocalKeyFile))
+						{
+							File.Copy(LocalKeyFile, LocalKeyFile + KeyBackupExt, true);
+						}
+
+						if (File.Exists(LocalPubKeyFile))
+						{
+							File.Copy(LocalPubKeyFile, LocalPubKeyFile + KeyBackupExt, true);
+						}
+						File.WriteAllText(BackupSentry, "placeholder");
+					}
+
+					File.Copy(RemoteKeyFile, LocalKeyFile, true);
+					File.Copy(RemotePubKeyFile, LocalPubKeyFile, true);
+
+					bUsingCustomKeys = true;
+
+					Log.Info("Running adb kill-server to refresh credentials");
+					TargetDeviceAndroid.RunAdbGlobalCommand("kill-server");
+
+					Thread.Sleep(5000);
+				}
+
+				InstanceCount++;
+			}
+
+		}
+
+		public static void RemoveInstance()
+		{
+			lock (Globals.MainLock)
+			{
+
+				InstanceCount--;
+
+				if (InstanceCount == 0 && bUsingCustomKeys)
+				{
+					Reset();
+
+					Log.Info("Running adb kill-server to refresh credentials");
+					TargetDeviceAndroid.RunAdbGlobalCommand("kill-server");
+					Thread.Sleep(2500);
+				}
+
+			}
+		}
+
+		public static void RestoreBackupKeys()
+		{			
+			string LocalKeyPath = Path.Combine(Environment.GetEnvironmentVariable("USERPROFILE"), ".android");
+			string LocalKeyFile = Path.Combine(LocalKeyPath, "adbkey");
+			string LocalPubKeyFile = Path.Combine(LocalKeyPath, "adbkey.pub");
+			string BackupSentry = Path.Combine(LocalKeyPath, "gauntlet.inuse");
+
+			if (File.Exists(BackupSentry))
+			{
+				Log.Info("Restoring original adb keys");
+
+				if (File.Exists(LocalKeyFile + KeyBackupExt))
+				{
+					File.Copy(LocalKeyFile + KeyBackupExt, LocalKeyFile, true);
+					File.Delete(LocalKeyFile + KeyBackupExt);
+				}
+				else
+				{
+					File.Delete(LocalKeyFile);
+				}
+
+				if (File.Exists(LocalPubKeyFile + KeyBackupExt))
+				{
+					File.Copy(LocalPubKeyFile + KeyBackupExt, LocalPubKeyFile, true);
+					File.Delete(LocalPubKeyFile + KeyBackupExt);
+				}
+				else
+				{
+					File.Delete(LocalPubKeyFile);
+				}
+
+				File.Delete(BackupSentry);
+			}
+
+		}
+
+		static AdbCredentialCache()
+		{
+			Reset();
+		}
+
+	}
+
+
+}
\ No newline at end of file
diff --git a/Engine/Source/Programs/AutomationTool/Gauntlet/Platform/Mac/Gauntlet.MacBuildSource.cs b/Engine/Source/Programs/AutomationTool/Gauntlet/Platform/Mac/Gauntlet.MacBuildSource.cs
new file mode 100644
--- /dev/null
+++ b/Engine/Source/Programs/AutomationTool/Gauntlet/Platform/Mac/Gauntlet.MacBuildSource.cs
@@ -0,0 +1,22 @@
+﻿// Copyright 1998-2018 Epic Games, Inc. All Rights Reserved.
+
+using System;
+using System.Collections.Generic;
+using System.IO;
+using AutomationTool;
+using UnrealBuildTool;
+using System.Threading;
+using System.Text.RegularExpressions;
+using System.Linq;
+
+namespace Gauntlet
+{
+	public class MacBuildSource : StagedBuildSource<StagedBuild>
+	{
+		public override string BuildName {  get { return "MacBuild";  } }
+
+		public override UnrealTargetPlatform Platform { get { return UnrealTargetPlatform.Mac; }}
+
+		public override string PlatformFolderPrefix { get { return "Mac"; }}
+	}
+}
diff --git a/Engine/Source/Programs/AutomationTool/Gauntlet/Platform/Mac/Gauntlet.SelfTest.TestTargetDeviceMac.cs b/Engine/Source/Programs/AutomationTool/Gauntlet/Platform/Mac/Gauntlet.SelfTest.TestTargetDeviceMac.cs
new file mode 100644
--- /dev/null
+++ b/Engine/Source/Programs/AutomationTool/Gauntlet/Platform/Mac/Gauntlet.SelfTest.TestTargetDeviceMac.cs
@@ -0,0 +1,29 @@
+﻿// Copyright 1998-2018 Epic Games, Inc. All Rights Reserved.
+
+using Gauntlet;
+using System;
+using System.Collections.Generic;
+using System.Diagnostics;
+using System.IO;
+using System.Linq;
+using System.Text;
+using System.Threading.Tasks;
+using UnrealBuildTool;
+
+namespace Gauntlet.SelfTest
+{
+	[TestGroup("Devices")]
+	class TestTargetDeviceMac : TestTargetDevice
+	{
+		public override void TickTest()
+		{
+			string TempDir = Path.Combine(Environment.CurrentDirectory, "GauntletTemp");
+
+			TargetDeviceMac Device = new TargetDeviceMac("LocalMac", Globals.TempDir);
+
+			CheckEssentialFunctions(Device);
+
+			MarkComplete();
+		}
+	}
+}
diff --git a/Engine/Source/Programs/AutomationTool/Gauntlet/Platform/Mac/Gauntlet.SelfTest.TestUnrealInstallAndRunMac.cs b/Engine/Source/Programs/AutomationTool/Gauntlet/Platform/Mac/Gauntlet.SelfTest.TestUnrealInstallAndRunMac.cs
new file mode 100644
--- /dev/null
+++ b/Engine/Source/Programs/AutomationTool/Gauntlet/Platform/Mac/Gauntlet.SelfTest.TestUnrealInstallAndRunMac.cs
@@ -0,0 +1,63 @@
+﻿// Copyright 1998-2018 Epic Games, Inc. All Rights Reserved.
+
+using AutomationTool;
+using Gauntlet;
+using System;
+using System.Collections.Generic;
+using System.Diagnostics;
+using System.IO;
+using System.Linq;
+using System.Text;
+using System.Threading.Tasks;
+using UnrealBuildTool;
+
+namespace Gauntlet.SelfTest
+{
+	/// <summary>
+	/// Validates that we can install and run all Windows based targets that are supported. P8 because this is the longest
+	/// running Unreal-based test
+	/// </summary>
+	[TestGroup("Unreal", 8)]
+	class TestUnrealInstallThenRunMac : TestUnrealInstallAndRunBase
+	{
+		public override void TickTest()
+		{
+			// create a new build
+			UnrealBuildSource Build = new UnrealBuildSource(this.GameName, this.UsesSharedBuildType, Environment.CurrentDirectory, this.BuildPath, new string[] { "" });
+
+			// check it's valid
+			if (!CheckResult(Build.BuildCount > 0, "staged build was invalid"))
+			{
+				MarkComplete();
+				return;
+			}
+
+			// Create devices to run the client and server
+			ITargetDevice PCDevice = new TargetDeviceMac("LocalMac", Globals.TempDir);
+
+			UnrealOptions GameOptions = new UnrealOptions();
+			GameOptions.Windowed = true;
+			GameOptions.ResX = 1280;
+			GameOptions.ResY = 720;
+			GameOptions.CommonArgs = "-log";
+
+			UnrealOptions OtherOptions = new UnrealOptions();
+			OtherOptions.CommonArgs = "-log";
+
+			// editor configs
+			TestInstallThenRun(Build, UnrealTargetRole.EditorGame, PCDevice, UnrealTargetConfiguration.Development, GameOptions);
+			TestInstallThenRun(Build, UnrealTargetRole.EditorServer, PCDevice, UnrealTargetConfiguration.Development, OtherOptions);
+			TestInstallThenRun(Build, UnrealTargetRole.Editor, PCDevice, UnrealTargetConfiguration.Development, OtherOptions);
+
+			// Monolithic clients
+			//TestInstallThenRun(Build, UnrealRoleType.Client, PCDevice, UnrealTargetConfiguration.Development, GameOptions);
+			TestInstallThenRun(Build, UnrealTargetRole.Client, PCDevice, UnrealTargetConfiguration.Test, GameOptions);
+
+			// Monolithic servers
+			//TestInstallThenRun(Build, UnrealRoleType.Server, PCDevice, UnrealTargetConfiguration.Development, OtherOptions);
+			//TestInstallThenRun(Build, UnrealRoleType.Server, PCDevice, UnrealTargetConfiguration.Test, OtherOptions);
+
+			MarkComplete();
+		}
+	}
+}
diff --git a/Engine/Source/Programs/AutomationTool/Gauntlet/Platform/Mac/Gauntlet.TargetDeviceMac.cs b/Engine/Source/Programs/AutomationTool/Gauntlet/Platform/Mac/Gauntlet.TargetDeviceMac.cs
new file mode 100644
--- /dev/null
+++ b/Engine/Source/Programs/AutomationTool/Gauntlet/Platform/Mac/Gauntlet.TargetDeviceMac.cs
@@ -0,0 +1,341 @@
+﻿// Copyright 1998-2018 Epic Games, Inc. All Rights Reserved.
+
+using System;
+using System.Collections.Generic;
+using System.IO;
+using AutomationTool;
+using UnrealBuildTool;
+using System.Text.RegularExpressions;
+
+namespace Gauntlet
+{
+	
+	public class MacDeviceFactory : IDeviceFactory
+	{
+		public bool CanSupportPlatform(UnrealTargetPlatform Platform)
+		{
+			return Platform == UnrealTargetPlatform.Mac;
+		}
+
+		public ITargetDevice CreateDevice(string InRef, string InParam)
+		{
+			return new TargetDeviceMac(InRef, InParam);
+		}
+	}
+
+	class MacAppInstall : IAppInstall
+	{
+		public string Name { get; private set; }
+
+		public string LocalPath;
+
+		public string WorkingDirectory;
+
+		public string ExecutablePath;
+
+		public string CommandArguments;
+
+		public string ArtifactPath;
+
+		public ITargetDevice Device { get; protected set; }
+
+		public CommandUtils.ERunOptions RunOptions { get; set; }
+
+		public MacAppInstall(string InName, TargetDeviceMac InDevice)
+		{
+			Name = InName;
+			Device = InDevice;
+			CommandArguments = "";
+			this.RunOptions = CommandUtils.ERunOptions.NoWaitForExit;
+		}
+
+		public IAppInstance Run()
+		{
+			return Device.Run(this);
+		}
+	}
+
+	class MacAppInstance : LocalAppProcess
+	{
+		protected MacAppInstall Install;
+
+		public MacAppInstance(MacAppInstall InInstall, IProcessResult InProcess)
+			: base(InProcess, InInstall.CommandArguments)
+		{
+			Install = InInstall;
+		}
+
+		public override string ArtifactPath
+		{
+			get
+			{
+				return Install.ArtifactPath;
+			}
+		}
+
+		public override ITargetDevice Device
+		{
+			get
+			{
+				return Install.Device;
+			}
+		}
+	}
+
+	public class TargetDeviceMac : ITargetDevice
+	{
+		public string Name { get; protected set; }
+
+		public UnrealTargetPlatform Platform { get { return UnrealTargetPlatform.Mac; } }
+
+		public bool IsAvailable { get { return true; } }
+		public bool IsConnected { get { return true; } }
+		public bool IsOn { get { return true; } }
+		public bool PowerOn() { return true; }
+		public bool PowerOff() { return true; }
+		public bool Reboot() { return true; }
+		public bool Connect() { return true; }
+		public bool Disconnect() { return true; }
+
+		protected string UserDir { get; set; }
+
+		protected string TempDir { get; set; }
+
+		public CommandUtils.ERunOptions RunOptions { get; set; }
+
+		protected Dictionary<EIntendedBaseCopyDirectory, string> LocalDirectoryMappings { get; set; }
+
+		public TargetDeviceMac(string InName, string InTempDir)
+		{
+			Name = InName;
+			TempDir = InTempDir;
+			UserDir = Path.Combine(TempDir, string.Format("{0}_UserDir", Name));
+			this.RunOptions = CommandUtils.ERunOptions.NoWaitForExit;
+			LocalDirectoryMappings = new Dictionary<EIntendedBaseCopyDirectory, string>();
+		}
+
+		#region IDisposable Support
+		private bool disposedValue = false; // To detect redundant calls
+
+		protected virtual void Dispose(bool disposing)
+		{
+			if (!disposedValue)
+			{
+				if (disposing)
+				{
+					// TODO: dispose managed state (managed objects).
+				}
+
+				// TODO: free unmanaged resources (unmanaged objects) and override a finalizer below.
+				// TODO: set large fields to null.
+
+				disposedValue = true;
+			}
+		}
+
+		// This code added to correctly implement the disposable pattern.
+		public void Dispose()
+		{
+			// Do not change this code. Put cleanup code in Dispose(bool disposing) above.
+			Dispose(true);
+			// TODO: uncomment the following line if the finalizer is overridden above.
+			// GC.SuppressFinalize(this);
+		}
+		#endregion
+
+		public override string ToString()
+		{
+			return Name;
+		}
+
+		protected string GetVolumeName(string InPath)
+		{
+			Match M = Regex.Match(InPath, @"/Volumes/(.+?)/");
+
+			if (M.Success)
+			{
+				return M.Groups[1].ToString();
+			}
+
+			return "";
+		}
+
+		protected IAppInstall InstallStagedBuild(UnrealAppConfig AppConfig, StagedBuild InBuild)
+		{
+			bool SkipDeploy = Globals.Params.ParseParam("SkipDeploy");
+
+			string BuildPath = InBuild.BuildPath;
+
+			// Must be on our volume to run
+			string BuildVolume = GetVolumeName(BuildPath);
+			string LocalRoot = GetVolumeName(Environment.CurrentDirectory);
+
+			if (BuildVolume.Equals(LocalRoot, StringComparison.OrdinalIgnoreCase) == false)
+			{
+				string SubDir = string.IsNullOrEmpty(AppConfig.Sandbox) ? AppConfig.ProjectName : AppConfig.Sandbox;
+				string DestPath = Path.Combine(this.TempDir, SubDir, AppConfig.ProcessType.ToString());
+
+				if (!SkipDeploy)
+				{
+					Log.Info("Installing {0} to {1}", AppConfig.Name, ToString());
+					Log.Verbose("\tCopying {0} to {1}", BuildPath, DestPath);
+					Gauntlet.Utils.SystemHelpers.CopyDirectory(BuildPath, DestPath, Utils.SystemHelpers.CopyOptions.Mirror);
+				}
+				else
+				{
+					Log.Info("Skipping install of {0} (-skipdeploy)", BuildPath);
+				}
+
+				// write a token, used to detect and old gauntlet-installedbuilds periodically
+				string TokenPath = Path.Combine(DestPath, "gauntlet.token");
+				File.WriteAllText(TokenPath, "Created by Gauntlet");
+
+				BuildPath = DestPath;
+			}
+
+			MacAppInstall MacApp = new MacAppInstall(AppConfig.Name, this);
+			MacApp.LocalPath = BuildPath;
+			MacApp.WorkingDirectory = MacApp.LocalPath;
+			MacApp.RunOptions = RunOptions;
+
+			// Set commandline replace any InstallPath arguments with the path we use
+			MacApp.CommandArguments = Regex.Replace(AppConfig.CommandLine, @"\$\(InstallPath\)", BuildPath, RegexOptions.IgnoreCase);
+
+			// Mac always forces this to stop logs and other artifacts going to different places
+			// Mac always forces this to stop logs and other artifacts going to different places
+			MacApp.CommandArguments += string.Format(" -userdir={0}", UserDir);
+			MacApp.ArtifactPath = Path.Combine(UserDir, @"Saved");
+
+			// temp - Mac doesn't support -userdir?
+			//MacApp.ArtifactPath = Path.Combine(Environment.GetFolderPath(Environment.SpecialFolder.Personal), "Library/Logs", AppConfig.ProjectName);
+
+			// clear artifact path
+			if (Directory.Exists(MacApp.ArtifactPath))
+			{
+				try
+				{
+					Directory.Delete(MacApp.ArtifactPath, true);
+				}
+				catch (Exception Ex)
+				{
+					Log.Warning("Failed to delete {0}. {1}", MacApp.ArtifactPath, Ex.Message);
+				}
+			}
+
+			if (Path.IsPathRooted(InBuild.ExecutablePath))
+			{
+				MacApp.ExecutablePath = InBuild.ExecutablePath;
+			}
+			else
+			{
+				// TODO - this check should be at a higher level....
+				string BinaryPath = Path.Combine(BuildPath, InBuild.ExecutablePath);
+
+				// check for a local newer executable
+				if (Globals.Params.ParseParam("dev") && AppConfig.ProcessType.UsesEditor() == false)
+				{
+					string LocalBinary = Path.Combine(Environment.CurrentDirectory, InBuild.ExecutablePath);
+
+					bool LocalFileExists = File.Exists(LocalBinary);
+					bool LocalFileNewer = LocalFileExists && File.GetLastWriteTime(LocalBinary) > File.GetLastWriteTime(BinaryPath);
+
+					Log.Verbose("Checking for newer binary at {0}", LocalBinary);
+					Log.Verbose("LocalFile exists: {0}. Newer: {1}", LocalFileExists, LocalFileNewer);
+
+					if (LocalFileExists && LocalFileNewer)
+					{
+						// need to -basedir to have our exe load content from the path
+						MacApp.CommandArguments += string.Format(" -basedir={0}", Path.GetDirectoryName(BinaryPath));
+
+						BinaryPath = LocalBinary;
+					}
+				}
+
+				MacApp.ExecutablePath = BinaryPath;
+			}
+
+			// now turn the Foo.app into Foo/Content/MacOS/Foo
+			string AppPath = Path.GetDirectoryName(MacApp.ExecutablePath);
+			string FileName = Path.GetFileNameWithoutExtension(MacApp.ExecutablePath);
+			MacApp.ExecutablePath = Path.Combine(MacApp.ExecutablePath, "Contents", "MacOS", FileName);
+
+			return MacApp;
+		}
+
+
+		public IAppInstall InstallApplication(UnrealAppConfig AppConfig)
+		{
+			if (AppConfig.Build is StagedBuild)
+			{
+				return InstallStagedBuild(AppConfig, AppConfig.Build as StagedBuild);
+			}
+
+			EditorBuild EditorBuild = AppConfig.Build as EditorBuild;
+
+			if (EditorBuild == null)
+			{
+				throw new AutomationException("Invalid build type!");
+			}
+
+			MacAppInstall MacApp = new MacAppInstall(AppConfig.Name, this);
+
+			MacApp.WorkingDirectory = Path.GetFullPath(EditorBuild.ExecutablePath);
+			MacApp.CommandArguments = AppConfig.CommandLine;
+			MacApp.RunOptions = RunOptions;
+
+			// Mac always forces this to stop logs and other artifacts going to different places
+			MacApp.CommandArguments += string.Format(" -userdir={0}", UserDir);
+			MacApp.ArtifactPath = Path.Combine(UserDir, @"Saved");
+
+			// now turn the Foo.app into Foo/Content/MacOS/Foo
+			string AppPath = Path.GetDirectoryName(EditorBuild.ExecutablePath);
+			string FileName = Path.GetFileNameWithoutExtension(EditorBuild.ExecutablePath);
+			MacApp.ExecutablePath = Path.Combine(EditorBuild.ExecutablePath, "Contents", "MacOS", FileName);
+
+			return MacApp;
+		}
+
+		public IAppInstance Run(IAppInstall App)
+		{
+			MacAppInstall MacInstall = App as MacAppInstall;
+
+			if (MacInstall == null)
+			{
+				throw new AutomationException("Invalid install type!");
+			}
+
+			IProcessResult Result = null;
+
+			lock (Globals.MainLock)
+			{
+				string NewWorkingDir = string.IsNullOrEmpty(MacInstall.WorkingDirectory) ? MacInstall.LocalPath : MacInstall.WorkingDirectory;
+				string OldWD = Environment.CurrentDirectory;
+				Environment.CurrentDirectory = NewWorkingDir;
+
+				Log.Info("Launching {0} on {1}", App.Name, ToString());
+				Log.Verbose("\t{0}", MacInstall.CommandArguments);
+
+				Result = CommandUtils.Run(MacInstall.ExecutablePath, MacInstall.CommandArguments, Options: MacInstall.RunOptions);
+
+				if (Result.HasExited && Result.ExitCode != 0)
+				{
+					throw new AutomationException("Failed to launch {0}. Error {1}", MacInstall.ExecutablePath, Result.ExitCode);
+				}
+
+				Environment.CurrentDirectory = OldWD;
+			}
+
+			return new MacAppInstance(MacInstall, Result);
+		}
+
+		public Dictionary<EIntendedBaseCopyDirectory, string> GetPlatformDirectoryMappings()
+		{
+			if (LocalDirectoryMappings.Count == 0)
+			{
+				Log.Warning("Platform directory mappings have not been populated for this platform! This should be done within InstallApplication()");
+			}
+			return LocalDirectoryMappings;
+		}
+	}
+
+}
diff --git a/Engine/Source/Programs/AutomationTool/Gauntlet/Platform/Null/Gauntlet.TargetDeviceNull.cs b/Engine/Source/Programs/AutomationTool/Gauntlet/Platform/Null/Gauntlet.TargetDeviceNull.cs
new file mode 100644
--- /dev/null
+++ b/Engine/Source/Programs/AutomationTool/Gauntlet/Platform/Null/Gauntlet.TargetDeviceNull.cs
@@ -0,0 +1,187 @@
+﻿// Copyright 1998-2018 Epic Games, Inc. All Rights Reserved.
+
+using System;
+using System.Collections.Generic;
+using System.IO;
+using AutomationTool;
+using UnrealBuildTool;
+using System.Text.RegularExpressions;
+
+namespace Gauntlet
+{
+
+	class NullAppInstance : IAppInstance
+	{
+
+		public NullAppInstance(ITargetDevice InDevice)
+		{
+			Device = InDevice;
+		}
+
+		public string ArtifactPath
+		{
+			get
+			{
+				return Globals.TempDir;
+			}
+		}
+
+		public string CommandLine { get { return ""; } }
+
+		public ITargetDevice Device { get; protected set; }
+
+		public bool HasExited { get; protected set; }
+
+		public bool WasKilled { get; protected set; }
+
+		public string StdOut { get { return ""; } }
+
+		public int ExitCode { get { return 0; }}
+
+		public void Kill()
+		{
+			if (!HasExited)
+			{
+				WasKilled = true;
+				HasExited = true;
+			}
+		}
+
+		public int WaitForExit()
+		{
+			HasExited = true;
+			return 0;
+		}
+	}
+
+
+	class NullAppInstall : IAppInstall
+	{
+		public ITargetDevice Device { get; protected set; }
+
+		public string Name { get; protected set; }
+			
+		public string CommandLine { get; protected set; }
+
+		public NullAppInstall(string InName, TargetDeviceNull InDevice, string InCommandLine)
+		{
+			Name = InName;
+			Device = InDevice;
+			CommandLine = InCommandLine;
+		}
+
+		public IAppInstance Run()
+		{
+			return Device.Run(this);
+		}
+	}
+
+	public class NullDeviceFactory : IDeviceFactory
+	{
+		public bool CanSupportPlatform(UnrealTargetPlatform Platform)
+		{
+			return Platform == UnrealTargetPlatform.Unknown;
+		}
+
+		public ITargetDevice CreateDevice(string InRef, string InParam)
+		{
+			return new TargetDeviceNull(InRef);
+		}
+	}
+
+	/// <summary>
+	/// Win32/64 implementation of a device to run applications
+	/// </summary>
+	public class TargetDeviceNull : ITargetDevice
+	{
+		public string Name { get; protected set; }
+
+		protected Dictionary<EIntendedBaseCopyDirectory, string> LocalDirectoryMappings { get; set; }
+
+		public TargetDeviceNull(string InName)
+		{
+			Name = InName;
+			LocalDirectoryMappings = new Dictionary<EIntendedBaseCopyDirectory, string>();
+		}
+
+		#region IDisposable Support
+		private bool disposedValue = false; // To detect redundant calls
+
+		protected virtual void Dispose(bool disposing)
+		{
+			if (!disposedValue)
+			{
+				if (disposing)
+				{
+					// TODO: dispose managed state (managed objects).
+				}
+
+				// TODO: free unmanaged resources (unmanaged objects) and override a finalizer below.
+				// TODO: set large fields to null.
+
+				disposedValue = true;
+			}
+		}
+
+		// This code added to correctly implement the disposable pattern.
+		public void Dispose()
+		{
+			// Do not change this code. Put cleanup code in Dispose(bool disposing) above.
+			Dispose(true);
+			// TODO: uncomment the following line if the finalizer is overridden above.
+			// GC.SuppressFinalize(this);
+		}
+		#endregion
+
+		public CommandUtils.ERunOptions RunOptions { get; set; }
+
+		public IAppInstance Run(IAppInstall App)
+		{
+			NullAppInstall NullApp = App as NullAppInstall;
+
+			if (NullApp == null)
+			{
+				throw new DeviceException("AppInstance is of incorrect type!");
+			}
+
+			Log.Info("Launching {0} on {1}", App.Name, ToString());
+			Log.Info("\t{0}", NullApp.CommandLine);
+
+			return new NullAppInstance(this);
+		}
+
+		public IAppInstall InstallApplication(UnrealAppConfig AppConfig)
+		{
+
+			NullAppInstall NullApp = new NullAppInstall(AppConfig.Name, this, AppConfig.CommandLine);
+
+	
+			return NullApp;
+		}
+
+		public UnrealTargetPlatform Platform { get { return UnrealTargetPlatform.Unknown; } }
+
+		public bool IsAvailable { get { return true; } }
+		public bool IsConnected { get { return true; } }
+		public bool IsOn { get { return true; } }
+		public bool PowerOn() { return true; }
+		public bool PowerOff() { return true; }
+		public bool Reboot() { return true; }
+		public bool Connect() { return true; }
+		public bool Disconnect() { return true; }
+
+		public override string ToString()
+		{
+			return Name;
+		}
+
+		public Dictionary<EIntendedBaseCopyDirectory, string> GetPlatformDirectoryMappings()
+		{
+			if (LocalDirectoryMappings.Count == 0)
+			{
+				Log.Warning("Platform directory mappings have not been populated for this platform! This should be done within InstallApplication()");
+			}
+			return LocalDirectoryMappings;
+		}
+	}
+}
\ No newline at end of file
diff --git a/Engine/Source/Programs/AutomationTool/Gauntlet/Platform/Windows/Gauntlet.SelfTest.TestTargetDeviceWindows.cs b/Engine/Source/Programs/AutomationTool/Gauntlet/Platform/Windows/Gauntlet.SelfTest.TestTargetDeviceWindows.cs
new file mode 100644
--- /dev/null
+++ b/Engine/Source/Programs/AutomationTool/Gauntlet/Platform/Windows/Gauntlet.SelfTest.TestTargetDeviceWindows.cs
@@ -0,0 +1,29 @@
+﻿// Copyright 1998-2018 Epic Games, Inc. All Rights Reserved.
+
+using Gauntlet;
+using System;
+using System.Collections.Generic;
+using System.Diagnostics;
+using System.IO;
+using System.Linq;
+using System.Text;
+using System.Threading.Tasks;
+using UnrealBuildTool;
+
+namespace Gauntlet.SelfTest
+{
+	[TestGroup("Devices")]
+	class TestTargetDeviceWindows : TestTargetDevice
+	{
+		public override void TickTest()
+		{
+			string TempDir = Path.Combine(Environment.CurrentDirectory, "GauntletTemp");
+
+			TargetDeviceWindows Device = new TargetDeviceWindows("Local PC", Globals.TempDir);
+
+			CheckEssentialFunctions(Device);
+
+			MarkComplete();
+		}
+	}
+}
diff --git a/Engine/Source/Programs/AutomationTool/Gauntlet/Platform/Windows/Gauntlet.TargetDeviceWindows.cs b/Engine/Source/Programs/AutomationTool/Gauntlet/Platform/Windows/Gauntlet.TargetDeviceWindows.cs
new file mode 100644
--- /dev/null
+++ b/Engine/Source/Programs/AutomationTool/Gauntlet/Platform/Windows/Gauntlet.TargetDeviceWindows.cs
@@ -0,0 +1,424 @@
+// Copyright 1998-2018 Epic Games, Inc. All Rights Reserved.
+
+using System;
+using System.Collections.Generic;
+using System.IO;
+using AutomationTool;
+using UnrealBuildTool;
+using System.Text.RegularExpressions;
+
+namespace Gauntlet
+{
+
+	public abstract class LocalAppProcess : IAppInstance
+	{
+		public IProcessResult ProcessResult { get; private set; }
+
+		public bool HasExited { get { return ProcessResult.HasExited; } }
+
+		public bool WasKilled { get; protected set; }
+
+		public string StdOut { get { return ProcessResult.Output; } }
+
+		public int ExitCode { get { return ProcessResult.ExitCode; } }
+
+		public string CommandLine { get; private set; }
+
+		public LocalAppProcess(IProcessResult InProcess, string InCommandLine)
+		{
+			this.CommandLine = InCommandLine;
+			this.ProcessResult = InProcess;
+		}
+		public int WaitForExit()
+		{
+			if (!HasExited)
+			{
+				ProcessResult.WaitForExit();
+			}
+
+			return ExitCode;
+		}
+
+		public void Kill()
+		{
+			if (!HasExited)
+			{
+				WasKilled = true;
+				ProcessResult.ProcessObject.Kill();
+			}
+		}
+
+		public abstract string ArtifactPath { get; }
+
+		public abstract ITargetDevice Device { get; }
+		
+	}
+
+	class WindowsAppInstance : LocalAppProcess
+	{
+		protected WindowsAppInstall Install;
+
+		public WindowsAppInstance(WindowsAppInstall InInstall, IProcessResult InProcess)
+			: base(InProcess, InInstall.CommandArguments)
+		{
+			Install = InInstall;
+		}
+
+		public override string ArtifactPath
+		{
+			get
+			{
+				return Install.ArtifactPath;
+			}
+		}
+
+		public override ITargetDevice Device
+		{
+			get
+			{
+				return Install.Device;
+			}
+		}
+	}
+
+
+	class WindowsAppInstall : IAppInstall
+	{
+		public string Name { get; private set; }
+
+		public string WorkingDirectory;
+
+		public string ExecutablePath;
+
+		public string CommandArguments;
+
+		public string ArtifactPath;
+
+		public TargetDeviceWindows WinDevice { get; private set; }
+
+		public ITargetDevice Device { get { return WinDevice; } }
+
+		public CommandUtils.ERunOptions RunOptions { get; set; }
+
+		public WindowsAppInstall(string InName, TargetDeviceWindows InDevice)
+		{
+			Name = InName;
+			WinDevice = InDevice;
+			CommandArguments = "";
+			this.RunOptions = CommandUtils.ERunOptions.NoWaitForExit;
+		}
+
+		public IAppInstance Run()
+		{
+			return Device.Run(this);
+		}
+	}
+
+	public class WindowsDeviceFactory : IDeviceFactory
+	{
+		public bool CanSupportPlatform(UnrealTargetPlatform Platform)
+		{
+			return Platform == UnrealTargetPlatform.Win64;
+		}
+
+		public ITargetDevice CreateDevice(string InRef, string InParam)
+		{
+			return new TargetDeviceWindows(InRef, InParam);
+		}
+	}
+
+	/// <summary>
+	/// Win32/64 implementation of a device to run applications
+	/// </summary>
+	public class TargetDeviceWindows : ITargetDevice
+	{
+		public string Name { get; protected set; }
+
+		protected string UserDir { get; set; }
+
+		/// <summary>
+		/// Our mappings of Intended directories to where they actually represent on this platform.
+		/// </summary>
+		protected Dictionary<EIntendedBaseCopyDirectory, string> LocalDirectoryMappings { get; set; }
+
+		public TargetDeviceWindows(string InName, string InTempDir)
+		{
+			Name = InName;
+			TempDir = InTempDir;
+			RunOptions = CommandUtils.ERunOptions.NoWaitForExit | CommandUtils.ERunOptions.NoLoggingOfRunCommand;
+
+			UserDir = Path.Combine(TempDir, string.Format("{0}_UserDir", Name));
+            LocalDirectoryMappings = new Dictionary<EIntendedBaseCopyDirectory, string>();
+		}
+
+		#region IDisposable Support
+		private bool disposedValue = false; // To detect redundant calls
+
+		protected virtual void Dispose(bool disposing)
+		{
+			if (!disposedValue)
+			{
+				if (disposing)
+				{
+					// TODO: dispose managed state (managed objects).
+				}
+
+				// TODO: free unmanaged resources (unmanaged objects) and override a finalizer below.
+				// TODO: set large fields to null.
+
+				disposedValue = true;
+			}
+		}
+
+		// This code added to correctly implement the disposable pattern.
+		public void Dispose()
+		{
+			// Do not change this code. Put cleanup code in Dispose(bool disposing) above.
+			Dispose(true);
+			// TODO: uncomment the following line if the finalizer is overridden above.
+			// GC.SuppressFinalize(this);
+		}
+		#endregion
+
+		public CommandUtils.ERunOptions RunOptions { get; set; }
+
+        // We care about UserDir in windows as some of the roles may require files going into user instead of build dir.
+        public void PopulateDirectoryMappings(string BasePath, string UserDir)
+        {
+            LocalDirectoryMappings.Add(EIntendedBaseCopyDirectory.Binaries, Path.Combine(BasePath, "Binaries"));
+            LocalDirectoryMappings.Add(EIntendedBaseCopyDirectory.Config, Path.Combine(BasePath, "Saved", "Config"));
+            LocalDirectoryMappings.Add(EIntendedBaseCopyDirectory.Content, Path.Combine(BasePath, "Content"));
+            LocalDirectoryMappings.Add(EIntendedBaseCopyDirectory.Demos, Path.Combine(UserDir, "Saved", "Demos"));
+            LocalDirectoryMappings.Add(EIntendedBaseCopyDirectory.Profiling, Path.Combine(BasePath, "Saved", "Profiling"));
+            LocalDirectoryMappings.Add(EIntendedBaseCopyDirectory.Saved, Path.Combine(BasePath, "Saved"));
+        }
+
+        public IAppInstance Run(IAppInstall App)
+		{
+			WindowsAppInstall WinApp = App as WindowsAppInstall;
+
+			if (WinApp == null)
+			{
+				throw new DeviceException("AppInstance is of incorrect type!");
+			}
+
+			if (File.Exists(WinApp.ExecutablePath) == false)
+			{
+				throw new DeviceException("Specified path {0} not found!", WinApp.ExecutablePath);
+			}
+
+			IProcessResult Result = null;
+
+			lock (Globals.MainLock)
+			{
+				string ExePath = Path.GetDirectoryName(WinApp.ExecutablePath);
+				string NewWorkingDir = string.IsNullOrEmpty(WinApp.WorkingDirectory) ? ExePath : WinApp.WorkingDirectory;
+				string OldWD = Environment.CurrentDirectory;
+				Environment.CurrentDirectory = NewWorkingDir;
+
+				Log.Info("Launching {0} on {1}", App.Name, ToString());
+				Log.Verbose("\t{0}", WinApp.CommandArguments);
+
+				Result = CommandUtils.Run(WinApp.ExecutablePath, WinApp.CommandArguments, Options: WinApp.RunOptions);
+
+				if (Result.HasExited && Result.ExitCode != 0)
+				{
+					throw new AutomationException("Failed to launch {0}. Error {1}", WinApp.ExecutablePath, Result.ExitCode);
+				}
+
+				Environment.CurrentDirectory = OldWD;
+			}
+
+			return new WindowsAppInstance(WinApp, Result);
+		}
+
+		protected IAppInstall InstallStagedBuild(UnrealAppConfig AppConfig, StagedBuild InBuild)
+		{
+			bool SkipDeploy = Globals.Params.ParseParam("SkipDeploy");
+
+			string BuildPath = InBuild.BuildPath;
+
+			if (CanRunFromPath(BuildPath) == false)
+			{
+				string SubDir = string.IsNullOrEmpty(AppConfig.Sandbox) ? AppConfig.ProjectName : AppConfig.Sandbox;
+				string DestPath = Path.Combine(this.TempDir, SubDir, AppConfig.ProcessType.ToString());
+
+				if (!SkipDeploy)
+				{
+					Log.Info("Installing {0} to {1}", AppConfig.Name, ToString());
+					Log.Verbose("\tCopying {0} to {1}", BuildPath, DestPath);
+					Gauntlet.Utils.SystemHelpers.CopyDirectory(BuildPath, DestPath, Utils.SystemHelpers.CopyOptions.Mirror);
+				}
+				else
+				{
+					Log.Info("Skipping install of {0} (-skipdeploy)", BuildPath);
+				}
+
+				// write a token, used to detect and old gauntlet-installedbuilds periodically
+				string TokenPath = Path.Combine(DestPath, "gauntlet.token");
+				File.WriteAllText(TokenPath, "Created by Gauntlet");
+
+				BuildPath = DestPath;
+			}
+
+			WindowsAppInstall WinApp = new WindowsAppInstall(AppConfig.Name, this);
+			WinApp.RunOptions = RunOptions;
+
+			// Set commandline replace any InstallPath arguments with the path we use
+			WinApp.CommandArguments = Regex.Replace(AppConfig.CommandLine, @"\$\(InstallPath\)", BuildPath, RegexOptions.IgnoreCase);
+
+			if (string.IsNullOrEmpty(UserDir) == false)
+			{
+				WinApp.CommandArguments += string.Format(" -userdir={0}", UserDir);
+				WinApp.ArtifactPath = Path.Combine(UserDir, @"Saved");
+			}
+			else
+			{
+				// e.g d:\Unreal\GameName\Saved
+				WinApp.ArtifactPath = Path.Combine(BuildPath, AppConfig.ProjectName, @"Saved");
+
+			}
+
+			// clear artifact path
+			if (Directory.Exists(WinApp.ArtifactPath))
+			{
+				try
+				{
+					Directory.Delete(WinApp.ArtifactPath, true);
+				}
+				catch (Exception Ex)
+				{
+					Log.Warning("Failed to delete {0}. {1}", WinApp.ArtifactPath, Ex.Message);
+				}
+			}
+
+            if (LocalDirectoryMappings.Count == 0)
+            {
+                PopulateDirectoryMappings(Path.Combine(BuildPath, AppConfig.ProjectName), UserDir);
+            }
+
+            if (AppConfig.FilesToCopy != null)
+            {
+                foreach (UnrealFileToCopy FileToCopy in AppConfig.FilesToCopy)
+                {
+                    string PathToCopyTo = Path.Combine(LocalDirectoryMappings[FileToCopy.TargetBaseDirectory], FileToCopy.TargetRelativeLocation);
+                    if (File.Exists(FileToCopy.SourceFileLocation))
+                    {
+                        FileInfo SrcInfo = new FileInfo(FileToCopy.SourceFileLocation);
+                        SrcInfo.IsReadOnly = false;
+                        string DirectoryToCopyTo = Path.GetDirectoryName(PathToCopyTo);
+                        if (!Directory.Exists(DirectoryToCopyTo))
+                        {
+                            Directory.CreateDirectory(DirectoryToCopyTo);
+                        }
+                        if (File.Exists(PathToCopyTo))
+                        {
+                            FileInfo ExistingFile = new FileInfo(PathToCopyTo);
+                            ExistingFile.IsReadOnly = false;
+                        }
+                        SrcInfo.CopyTo(PathToCopyTo, true);
+                        Log.Info("Copying {0} to {1}", FileToCopy.SourceFileLocation, PathToCopyTo);
+                    }
+                    else
+                    {
+                        Log.Warning("File to copy {0} not found", FileToCopy);
+                    }
+                }
+            }
+
+            if (Path.IsPathRooted(InBuild.ExecutablePath))
+			{
+				WinApp.ExecutablePath = InBuild.ExecutablePath;
+			}
+			else
+			{
+				// TODO - this check should be at a higher level....
+				string BinaryPath = Path.Combine(BuildPath, InBuild.ExecutablePath);
+
+				// check for a local newer executable
+				if (Globals.Params.ParseParam("dev") && AppConfig.ProcessType.UsesEditor() == false)
+				{
+					string LocalBinary = Path.Combine(Environment.CurrentDirectory, InBuild.ExecutablePath);
+
+					bool LocalFileExists = File.Exists(LocalBinary);
+					bool LocalFileNewer = LocalFileExists && File.GetLastWriteTime(LocalBinary) > File.GetLastWriteTime(BinaryPath);
+
+					Log.Verbose("Checking for newer binary at {0}", LocalBinary);
+					Log.Verbose("LocalFile exists: {0}. Newer: {1}", LocalFileExists, LocalFileNewer);
+
+					if (LocalFileExists && LocalFileNewer)
+					{
+						// need to -basedir to have our exe load content from the path
+						WinApp.CommandArguments += string.Format(" -basedir={0}", Path.GetDirectoryName(BinaryPath));
+
+						BinaryPath = LocalBinary;
+					}
+				}
+
+				WinApp.ExecutablePath = BinaryPath;
+			}
+
+			return WinApp;
+		}
+
+
+		public IAppInstall InstallApplication(UnrealAppConfig AppConfig)
+		{
+			if (AppConfig.Build is StagedBuild)
+			{
+				return InstallStagedBuild(AppConfig, AppConfig.Build as StagedBuild);
+			}
+
+			EditorBuild EditorBuild = AppConfig.Build as EditorBuild;
+
+			if (EditorBuild == null)
+			{
+				throw new AutomationException("Invalid build type!");
+			}
+
+			WindowsAppInstall WinApp = new WindowsAppInstall(AppConfig.Name, this);
+
+			WinApp.WorkingDirectory = Path.GetDirectoryName(EditorBuild.ExecutablePath);
+			WinApp.RunOptions = RunOptions;
+	
+			// Force this to stop logs and other artifacts going to different places
+			WinApp.CommandArguments = AppConfig.CommandLine + string.Format(" -userdir={0}", UserDir);
+			WinApp.ArtifactPath = Path.Combine(UserDir, @"Saved");
+			WinApp.ExecutablePath = EditorBuild.ExecutablePath;
+
+			return WinApp;
+		}
+		
+		public bool CanRunFromPath(string InPath)
+		{
+			// path must be under our mapped drive (e.g. w:\KitName);
+			return string.Compare(Path.GetPathRoot(InPath), Path.GetPathRoot(this.TempDir), StringComparison.OrdinalIgnoreCase) == 0;
+		}
+
+		public UnrealTargetPlatform Platform { get { return UnrealTargetPlatform.Win64; } }
+
+		public string TempDir { get; private set; }
+		public bool IsAvailable { get { return true; } }
+		public bool IsConnected { get { return true; } }
+		public bool IsOn { get { return true; } }
+		public bool PowerOn() { return true; }
+		public bool PowerOff() { return true; }
+		public bool Reboot() { return true; }
+		public bool Connect() { return true; }
+		public bool Disconnect() { return true; }
+
+		public override string ToString()
+		{
+			return Name;
+		}
+
+		public Dictionary<EIntendedBaseCopyDirectory, string> GetPlatformDirectoryMappings()
+		{
+			if (LocalDirectoryMappings.Count == 0)
+			{
+				Log.Warning("Platform directory mappings have not been populated for this platform! This should be done within InstallApplication()");
+			}
+			return LocalDirectoryMappings;
+		}
+	}
+}
\ No newline at end of file
diff --git a/Engine/Source/Programs/AutomationTool/Gauntlet/Platform/Windows/Gauntlet.WindowsBuildSources.cs b/Engine/Source/Programs/AutomationTool/Gauntlet/Platform/Windows/Gauntlet.WindowsBuildSources.cs
new file mode 100644
--- /dev/null
+++ b/Engine/Source/Programs/AutomationTool/Gauntlet/Platform/Windows/Gauntlet.WindowsBuildSources.cs
@@ -0,0 +1,16 @@
+﻿// Copyright 1998-2018 Epic Games, Inc. All Rights Reserved.
+
+using System;
+using UnrealBuildTool;
+
+namespace Gauntlet
+{
+	public class Win64BuildSource : StagedBuildSource<StagedBuild>
+	{
+		public override string BuildName { get { return "Win64StagedBuild"; } }
+
+		public override UnrealTargetPlatform Platform { get { return UnrealTargetPlatform.Win64; } }
+
+		public override string PlatformFolderPrefix { get { return "Windows"; } }
+	}
+}
diff --git a/Engine/Source/Programs/AutomationTool/Gauntlet/SelfTest/Devices/Gauntlet.SelfTest.TestTargetDeviceMac.cs b/Engine/Source/Programs/AutomationTool/Gauntlet/SelfTest/Devices/Gauntlet.SelfTest.TestTargetDeviceMac.cs
deleted file mode 100644
diff --git a/Engine/Source/Programs/AutomationTool/Gauntlet/SelfTest/Devices/Gauntlet.SelfTest.TestTargetDeviceWindows.cs b/Engine/Source/Programs/AutomationTool/Gauntlet/SelfTest/Devices/Gauntlet.SelfTest.TestTargetDeviceWindows.cs
deleted file mode 100644
diff --git a/Engine/Source/Programs/AutomationTool/Gauntlet/SelfTest/Unreal/Gauntlet.SelfTest.TestUnrealInstallAndRunMac.cs b/Engine/Source/Programs/AutomationTool/Gauntlet/SelfTest/Unreal/Gauntlet.SelfTest.TestUnrealInstallAndRunMac.cs
deleted file mode 100644
diff --git a/Engine/Source/Programs/AutomationTool/Gauntlet/SelfTest/Unreal/Gauntlet.SelfTest.TestUnrealInstallAndRunWindows.cs b/Engine/Source/Programs/AutomationTool/Gauntlet/SelfTest/Unreal/Gauntlet.SelfTest.TestUnrealInstallAndRunWindows.cs
deleted file mode 100644
diff --git a/Engine/Source/Programs/AutomationTool/Gauntlet/SelfTest/Unreal/Windows/Gauntlet.SelfTest.TestUnrealInstallAndRunWindows.cs b/Engine/Source/Programs/AutomationTool/Gauntlet/SelfTest/Unreal/Windows/Gauntlet.SelfTest.TestUnrealInstallAndRunWindows.cs
new file mode 100644
--- /dev/null
+++ b/Engine/Source/Programs/AutomationTool/Gauntlet/SelfTest/Unreal/Windows/Gauntlet.SelfTest.TestUnrealInstallAndRunWindows.cs
@@ -0,0 +1,54 @@
+﻿// Copyright 1998-2018 Epic Games, Inc. All Rights Reserved.
+
+using AutomationTool;
+using Gauntlet;
+using System;
+using System.Collections.Generic;
+using System.Diagnostics;
+using System.IO;
+using System.Linq;
+using System.Text;
+using System.Threading.Tasks;
+using UnrealBuildTool;
+
+namespace Gauntlet.SelfTest
+{
+	/// <summary>
+	/// Validates that we can install and run all Windows based targets that are supported. P8 because this is the longest
+	/// running Unreal-based test
+	/// </summary>
+	[TestGroup("Unreal", 8)]
+	class TestUnrealInstallThenRunWindows : TestUnrealInstallAndRunBase
+	{
+		public override void TickTest()
+		{
+			// create a new build
+			UnrealBuildSource Build = new UnrealBuildSource(this.GameName, this.UsesSharedBuildType, Environment.CurrentDirectory, this.BuildPath, new string[] { "" });
+
+			// check it's valid
+			if (!CheckResult(Build.BuildCount > 0, "staged build was invalid"))
+			{
+				MarkComplete();
+				return;
+			}
+
+			// Create devices to run the client and server
+			ITargetDevice PCDevice = new TargetDeviceWindows("Local PC", Globals.TempDir);
+
+			// Monolithic servers
+			TestInstallThenRun(Build, UnrealTargetRole.Server, PCDevice, UnrealTargetConfiguration.Development);
+			TestInstallThenRun(Build, UnrealTargetRole.Server, PCDevice, UnrealTargetConfiguration.Test);
+
+			// Monolithic clients
+			TestInstallThenRun(Build, UnrealTargetRole.Client, PCDevice, UnrealTargetConfiguration.Development);
+			TestInstallThenRun(Build, UnrealTargetRole.Client, PCDevice, UnrealTargetConfiguration.Test);
+
+			// editor configs
+			TestInstallThenRun(Build, UnrealTargetRole.EditorGame, PCDevice, UnrealTargetConfiguration.Development);
+			TestInstallThenRun(Build, UnrealTargetRole.EditorServer, PCDevice, UnrealTargetConfiguration.Development);
+			TestInstallThenRun(Build, UnrealTargetRole.Editor, PCDevice, UnrealTargetConfiguration.Development);
+
+			MarkComplete();
+		}
+	}
+}
diff --git a/Engine/Source/Programs/AutomationTool/Gauntlet/Unreal/BuildSource/Gauntlet.AndroidBuildSource.cs b/Engine/Source/Programs/AutomationTool/Gauntlet/Unreal/BuildSource/Gauntlet.AndroidBuildSource.cs
deleted file mode 100644
diff --git a/Engine/Source/Programs/AutomationTool/Gauntlet/Unreal/BuildSource/Gauntlet.MacBuildSource.cs b/Engine/Source/Programs/AutomationTool/Gauntlet/Unreal/BuildSource/Gauntlet.MacBuildSource.cs
deleted file mode 100644
diff --git a/Engine/Source/Programs/AutomationTool/Gauntlet/Unreal/BuildSource/Gauntlet.WindowsBuildSources.cs b/Engine/Source/Programs/AutomationTool/Gauntlet/Unreal/BuildSource/Gauntlet.WindowsBuildSources.cs
deleted file mode 100644
