Compilation fixes for people who do not have access to specific platforms

diff --git a/Engine/Source/Programs/AutomationTool/Gauntlet/SelfTest/Unreal/Gauntlet.SelfTest.TestShortSoloGame.cs b/Engine/Source/Programs/AutomationTool/Gauntlet/SelfTest/Unreal/Gauntlet.SelfTest.TestShortSoloGame.cs
--- a/Engine/Source/Programs/AutomationTool/Gauntlet/SelfTest/Unreal/Gauntlet.SelfTest.TestShortSoloGame.cs
+++ b/Engine/Source/Programs/AutomationTool/Gauntlet/SelfTest/Unreal/Gauntlet.SelfTest.TestShortSoloGame.cs
@@ -61,11 +61,11 @@ namespace Gauntlet.SelfTest
 
 			// Create devices to run the client and server
 			ITargetDevice ServerDevice = new TargetDeviceWindows("PC Server", Gauntlet.Globals.TempDir);
-			ITargetDevice ClientDevice;
+			ITargetDevice ClientDevice = null;
 
 			if (Platform == UnrealTargetPlatform.PS4)
 			{
-				ClientDevice = new TargetDevicePS4(this.PS4Name);
+				//ClientDevice = new TargetDevicePS4(this.PS4Name);
 			}
 			else
 			{
diff --git a/Engine/Source/Programs/AutomationTool/Gauntlet/SelfTest/Unreal/Gauntlet.SelfTest.TestUnrealSession.cs b/Engine/Source/Programs/AutomationTool/Gauntlet/SelfTest/Unreal/Gauntlet.SelfTest.TestUnrealSession.cs
--- a/Engine/Source/Programs/AutomationTool/Gauntlet/SelfTest/Unreal/Gauntlet.SelfTest.TestUnrealSession.cs
+++ b/Engine/Source/Programs/AutomationTool/Gauntlet/SelfTest/Unreal/Gauntlet.SelfTest.TestUnrealSession.cs
@@ -28,7 +28,6 @@ namespace Gauntlet.SelfTest
 			DevicePool.Instance.RegisterDevices(new ITargetDevice[] {
 				new TargetDeviceWindows("Local PC1", Globals.TempDir)
 				, new TargetDeviceWindows("Local PC2", Globals.TempDir)
-				, new TargetDevicePS4(this.PS4Name)
 			});
 			
 			// Create a new build (params come from our base class will be similar to "OrionGame" and "p:\builds\orion\branch-cl")
diff --git a/Engine/Source/Programs/AutomationTool/Gauntlet/Unreal/Tasks/Gauntlet.Task.CleanDevices.cs b/Engine/Source/Programs/AutomationTool/Gauntlet/Unreal/Tasks/Gauntlet.Task.CleanDevices.cs
--- a/Engine/Source/Programs/AutomationTool/Gauntlet/Unreal/Tasks/Gauntlet.Task.CleanDevices.cs
+++ b/Engine/Source/Programs/AutomationTool/Gauntlet/Unreal/Tasks/Gauntlet.Task.CleanDevices.cs
@@ -89,10 +89,10 @@ namespace Gauntlet
 				return;
 			}
 
-			if (Device is TargetDevicePS4)
+			/*if (Device is TargetDevicePS4)
 			{
 				CleanPS4(Device as TargetDevicePS4);
-			}
+			}*/
 
 			// disconnect and power down
 			Gauntlet.Log.Info("Powering down and disconnecting from {0}", Device.Name);
@@ -100,90 +100,5 @@ namespace Gauntlet
 			// turns out this may be a bad idea.. sorry environment.
 			//Device.PowerOff();
 		}
-
-		protected void CleanPS4(TargetDevicePS4 PS4)
-		{
-			// TODO - it would be nice to make both enumeration of builds and their removal tasks that the targetdevice provides...
-			DirectoryInfo Di = new DirectoryInfo(PS4.DataPath);
-
-			if (Di.Exists == false)
-			{
-				Gauntlet.Log.Warning("Data path {0} not found for {1}", PS4.DataPath, PS4.DeviceName);
-				return;
-			}
-
-			// first check each toplevel directory for a token that gauntlet creates
-			foreach (DirectoryInfo SubDir in Di.GetDirectories())
-			{
-				try
-				{
-					FileInfo TokenFile = SubDir.GetFiles().Where(F => F.Name == "gauntlet.token" || F.Name == "testdata.token").FirstOrDefault();
-
-					if (TokenFile != null)
-					{
-						double Age = (DateTime.Now - TokenFile.LastWriteTime).TotalDays;
-						if (Age >= MaxDays)
-						{
-							Gauntlet.Log.Info("Build at {0} is {1:00} days old (max={2}). Removing", SubDir.Name, Age, MaxDays);
-
-							try
-							{
-								SubDir.Delete(true);
-							}
-							catch (Exception Ex)
-							{
-								Gauntlet.Log.Error("Error deleting {0}: {1}", SubDir.FullName, Ex);
-							}
-
-						}
-					}
-					else
-					{
-						Gauntlet.Log.Verbose("Directory {0} was not created by Gauntlet. Ignoring", SubDir.Name);
-					}
-				}
-				catch (Exception Ex)
-				{
-					// can occur when the data drive has characters considered invalid for a windows FS.
-					Gauntlet.Log.Warning("Error getting directory tree for {0} on {1}. {2}", SubDir.FullName, PS4.Name, Ex);
-				}
-			}
-
-			// now delete old Crashdumps
-			DirectoryInfo DumpDir = new DirectoryInfo(Path.Combine(Di.FullName, "sce_coredumps"));
-
-			if (DumpDir.Exists)
-			{
-
-				foreach (DirectoryInfo SubDir in DumpDir.GetDirectories())
-				{
-					// each crashdump has a report
-					FileInfo DumpFile = SubDir.GetFiles("*.orbisdmp", SearchOption.TopDirectoryOnly).FirstOrDefault();
-
-					if (DumpFile != null)
-					{
-						double Age = (DateTime.Now - DumpFile.LastWriteTime).TotalDays;
-						if (Age >= MaxDays)
-						{
-							Gauntlet.Log.Info("Crashdump at {0} is {1:00} days old (max={2}). Removing", SubDir.Name, Age, MaxDays);
-
-							try
-							{
-								SubDir.Delete(true);
-							}
-							catch (Exception Ex)
-							{
-								Gauntlet.Log.Error("Error deleting {0}: {1}", SubDir.FullName, Ex);
-							}
-						}
-					}
-					else
-					{
-						Gauntlet.Log.Info("No orbisdmp found in {0}", SubDir.Name);
-					}
-				}
-			}
-			PS4.SetSetting("NP Environment", "sp-int");
-		}
 	}
 }
