PR #5130: Fix ZipUtil to include all files, not just those with a `.` (Contributed by [REDACTED_PERSON])

diff --git a/Engine/Source/Programs/AutomationTool/Scripts/ZipUtils.Automation.cs b/Engine/Source/Programs/AutomationTool/Scripts/ZipUtils.Automation.cs
--- a/Engine/Source/Programs/AutomationTool/Scripts/ZipUtils.Automation.cs
+++ b/Engine/Source/Programs/AutomationTool/Scripts/ZipUtils.Automation.cs
@@ -51,7 +51,7 @@ public class ZipUtils : BuildCommand
 
 		//setup filter to include all files
 		FileFilter Filter = new FileFilter();
-		Filter.Include("*.*");
+		Filter.Include("*");
 
 		if (bUnzip)
 		{
