Add exception handler to CleanAutomationReports script, and move automated test screenshot cleanup into the same job. Suspect both are running at the same time, causing a race.

diff --git a/Engine/Source/Programs/AutomationTool/Scripts/CleanAutomationReports.cs b/Engine/Source/Programs/AutomationTool/Scripts/CleanAutomationReports.cs
--- a/Engine/Source/Programs/AutomationTool/Scripts/CleanAutomationReports.cs
+++ b/Engine/Source/Programs/AutomationTool/Scripts/CleanAutomationReports.cs
@@ -50,11 +50,18 @@ namespace AutomationTool
 			List<DirectoryInfo> FoldersToDelete = new List<DirectoryInfo>();
 			foreach (DirectoryInfo BuildDirectory in new DirectoryInfo(ReportDir).EnumerateDirectories())
 			{
-				if(!BuildDirectory.EnumerateFiles("*", SearchOption.AllDirectories).Any(x => x.LastWriteTimeUtc > RetainTime))
+				try
+				{
+					if(!BuildDirectory.EnumerateFiles("*", SearchOption.AllDirectories).Any(x => x.LastWriteTimeUtc > RetainTime))
+					{
+						FoldersToDelete.Add(BuildDirectory);
+					}
+					NumFolders++;
+				}
+				catch(Exception Ex)
 				{
-					FoldersToDelete.Add(BuildDirectory);
+					CommandUtils.LogWarning("Unable to enumerate {0}: {1}", BuildDirectory.FullName, Ex.ToString());
 				}
-				NumFolders++;
 			}
 			CommandUtils.LogInformation("Found {0} builds; {1} to delete.", NumFolders, FoldersToDelete.Count);
 
