UBT: Remove the ability to override whether the engine is installed or not via the command line. There are a lot of pathways that have to maintain support for this for forwarding on to other tools, they inconsistently check for a command line argument or existence of InstalledBuild.txt, and there is no way to override the enterprise/project installed flags in this way. For testing, it's easy enough to create InstalledBuild.txt.

diff --git a/Engine/Source/Programs/AutomationTool/AutomationUtils/Automation.cs b/Engine/Source/Programs/AutomationTool/AutomationUtils/Automation.cs
--- a/Engine/Source/Programs/AutomationTool/AutomationUtils/Automation.cs
+++ b/Engine/Source/Programs/AutomationTool/AutomationUtils/Automation.cs
@@ -116,7 +116,6 @@ namespace AutomationTool
 		public static CommandLineArg List = new CommandLineArg("-List");
 		public static CommandLineArg VS2015 = new CommandLineArg("-2015");
 		public static CommandLineArg NoKill = new CommandLineArg("-NoKill");
-		public static CommandLineArg Installed = new CommandLineArg("-Installed");
 		public static CommandLineArg UTF8Output = new CommandLineArg("-UTF8Output");
         public static CommandLineArg AllowStdOutLogVerbosity = new CommandLineArg("-AllowStdOutLogVerbosity");
         public static CommandLineArg NoAutoSDK = new CommandLineArg("-NoAutoSDK");
@@ -125,8 +124,6 @@ namespace AutomationTool
         /// Allows you to use local storage for your root build storage dir (default of P:\Builds (on PC) is changed to Engine\Saved\LocalBuilds). Used for local testing.
         /// </summary>
         public static CommandLineArg UseLocalBuildStorage = new CommandLineArg("-UseLocalBuildStorage");
-		public static CommandLineArg InstalledEngine = new CommandLineArg("-InstalledEngine");
-		public static CommandLineArg NotInstalledEngine = new CommandLineArg("-NotInstalledEngine");
 
 		/// <summary>
 		/// Force initialize static members by calling this.
@@ -464,26 +461,13 @@ AutomationTool.exe [-verbose] [-compileonly] [-p4] Command0 [-Arg0 -Arg1 -Arg2 .
 			Log.TraceLog("Setting up command environment.");
 			CommandUtils.InitCommandEnvironment();
 
-			// Determine if the engine is installed
-			bIsEngineInstalled = GlobalCommandLine.Installed;
-			string InstalledBuildFile = Path.Combine(CommandUtils.CmdEnv.LocalRoot, "Engine", "Build", "InstalledBuild.txt");
-			bIsEngineInstalled |= File.Exists(InstalledBuildFile);
-			if (bIsEngineInstalled.Value)
-			{
-				bIsEngineInstalled = !GlobalCommandLine.NotInstalledEngine;
-			}
-			else
-			{
-				bIsEngineInstalled = GlobalCommandLine.InstalledEngine;
-			}
-
 			// Create the log file, and flush the startup listener to it
 			TraceListener LogTraceListener = LogUtils.AddLogFileListener(CommandUtils.CmdEnv.LogFolder, CommandUtils.CmdEnv.FinalLogFolder);
 			StartupListener.CopyTo(LogTraceListener);
 			Trace.Listeners.Remove(StartupListener);
 
 			// Initialize UBT
-			if(!UnrealBuildTool.PlatformExports.Initialize(bIsEngineInstalled.Value))
+			if(!UnrealBuildTool.PlatformExports.Initialize())
 			{
 				Log.TraceInformation("Failed to initialize UBT");
 				return ExitCode.Error_Unknown;
@@ -640,16 +624,11 @@ AutomationTool.exe [-verbose] [-compileonly] [-p4] Command0 [-Arg0 -Arg1 -Arg2 .
 		/// Returns true if AutomationTool is running using installed Engine components
 		/// </summary>
 		/// <returns>True if running using installed Engine components</returns>
+		[Obsolete("IsEngineInstalled() is deprecated in 4.22. Call CommandUtils.IsEngineInstalled instead.")]
 		static public bool IsEngineInstalled()
 		{
-			if (!bIsEngineInstalled.HasValue)
-			{
-				throw new AutomationException("IsEngineInstalled has not been initialized yet");
-			}
-
-			return bIsEngineInstalled.Value;
+			return CommandUtils.IsEngineInstalled();
 		}
-		static private bool? bIsEngineInstalled;
 
 		#endregion
 
diff --git a/Engine/Source/Programs/AutomationTool/AutomationUtils/CommandEnvironment.cs b/Engine/Source/Programs/AutomationTool/AutomationUtils/CommandEnvironment.cs
--- a/Engine/Source/Programs/AutomationTool/AutomationUtils/CommandEnvironment.cs
+++ b/Engine/Source/Programs/AutomationTool/AutomationUtils/CommandEnvironment.cs
@@ -102,7 +102,7 @@ namespace AutomationTool
 			LogFolder = CommandUtils.GetEnvVar(EnvVarNames.LogFolder);
 			if (String.IsNullOrEmpty(LogFolder))
 			{
-				if (GlobalCommandLine.Installed)
+				if (CommandUtils.IsEngineInstalled())
 				{
 					LogFolder = GetInstalledLogFolder();
 				}
diff --git a/Engine/Source/Programs/AutomationTool/AutomationUtils/CommandUtils.cs b/Engine/Source/Programs/AutomationTool/AutomationUtils/CommandUtils.cs
--- a/Engine/Source/Programs/AutomationTool/AutomationUtils/CommandUtils.cs
+++ b/Engine/Source/Programs/AutomationTool/AutomationUtils/CommandUtils.cs
@@ -106,6 +106,21 @@ namespace AutomationTool
 
 		#endregion
 
+		/// <summary>
+		/// Returns true if AutomationTool is running using installed Engine components
+		/// </summary>
+		/// <returns>True if running using installed Engine components</returns>
+		static public bool IsEngineInstalled()
+		{
+			if (!bIsEngineInstalled.HasValue)
+			{
+				bIsEngineInstalled = FileReference.Exists(FileReference.Combine(CommandUtils.EngineDirectory, "Build", "InstalledBuild.txt"));
+			}
+			return bIsEngineInstalled.Value;
+		}
+
+		static private bool? bIsEngineInstalled;
+
 		#region Logging
 
 		/// <summary>
diff --git a/Engine/Source/Programs/AutomationTool/AutomationUtils/ProjectUtils.cs b/Engine/Source/Programs/AutomationTool/AutomationUtils/ProjectUtils.cs
--- a/Engine/Source/Programs/AutomationTool/AutomationUtils/ProjectUtils.cs
+++ b/Engine/Source/Programs/AutomationTool/AutomationUtils/ProjectUtils.cs
@@ -217,7 +217,7 @@ namespace AutomationTool
 				bool RetVal = false;
 				foreach (UnrealTargetPlatform TargetPlatformType in TargetPlatforms)
 				{
-					if((!Automation.IsEngineInstalled() && !PlatformExports.HasDefaultBuildConfig(RawProjectPath, TargetPlatformType)) || PlatformExports.RequiresBuild(RawProjectPath, TargetPlatformType))
+					if((!CommandUtils.IsEngineInstalled() && !PlatformExports.HasDefaultBuildConfig(RawProjectPath, TargetPlatformType)) || PlatformExports.RequiresBuild(RawProjectPath, TargetPlatformType))
 					{
 						RetVal = true;
 						break;
@@ -438,7 +438,7 @@ namespace AutomationTool
 		private static string GetRulesAssemblyFolder()
 		{
 			string RulesFolder;
-			if (GlobalCommandLine.Installed.IsSet)
+			if (CommandUtils.IsEngineInstalled())
 			{
 				RulesFolder = CommandUtils.CombinePaths(Path.GetTempPath(), "UAT", CommandUtils.EscapePath(CommandUtils.CmdEnv.LocalRoot), "Rules"); 
 			}
diff --git a/Engine/Source/Programs/AutomationTool/AutomationUtils/UE4Build.cs b/Engine/Source/Programs/AutomationTool/AutomationUtils/UE4Build.cs
--- a/Engine/Source/Programs/AutomationTool/AutomationUtils/UE4Build.cs
+++ b/Engine/Source/Programs/AutomationTool/AutomationUtils/UE4Build.cs
@@ -1429,7 +1429,7 @@ namespace AutomationTool
 		FileReference GetManifestFile(FileReference ProjectFile)
 		{
 			// Can't write to Engine directory on installed builds
-			if (Automation.IsEngineInstalled() && ProjectFile != null)
+			if (CommandUtils.IsEngineInstalled() && ProjectFile != null)
 			{
 				return FileReference.Combine(ProjectFile.Directory, "Intermediate", "Build", "Manifest.xml");
 			}
diff --git a/Engine/Source/Programs/AutomationTool/HTML5/HTML5Platform.PakFiles.Automation.cs b/Engine/Source/Programs/AutomationTool/HTML5/HTML5Platform.PakFiles.Automation.cs
--- a/Engine/Source/Programs/AutomationTool/HTML5/HTML5Platform.PakFiles.Automation.cs
+++ b/Engine/Source/Programs/AutomationTool/HTML5/HTML5Platform.PakFiles.Automation.cs
@@ -1,4 +1,4 @@
-﻿// Copyright 1998-2018 Epic Games, Inc. All Rights Reserved.
+// Copyright 1998-2018 Epic Games, Inc. All Rights Reserved.
 using System;
 using System.Collections.Generic;
 using System.Text.RegularExpressions;
@@ -290,7 +290,7 @@ public class HTMLPakAutomation
 		{
 			CmdLine += " -sign=" + CommandUtils.MakePathSafeToUseWithCommandLine(EncryptionKeys);
 		}
-		if (GlobalCommandLine.Installed)
+		if (CommandUtils.IsEngineInstalled())
 		{
 			CmdLine += " -installed";
 		}
diff --git a/Engine/Source/Programs/AutomationTool/IOS/IOSPlatform.Automation.cs b/Engine/Source/Programs/AutomationTool/IOS/IOSPlatform.Automation.cs
--- a/Engine/Source/Programs/AutomationTool/IOS/IOSPlatform.Automation.cs
+++ b/Engine/Source/Programs/AutomationTool/IOS/IOSPlatform.Automation.cs
@@ -533,7 +533,7 @@ public class IOSPlatform : Platform
 			Arguments += " -platforms=" + PlatformName + " -game -nointellisense -" + PlatformName + "deployonly -ignorejunk";
 
 			// If engine is installed then UBT doesn't need to be built
-			if (Automation.IsEngineInstalled())
+			if (CommandUtils.IsEngineInstalled())
 			{
 				// Get the path to UBT
 				string InstalledUBT = CombinePaths(CmdEnv.LocalRoot, "Engine/Binaries/DotNET/UnrealBuildTool.exe");
@@ -1648,7 +1648,7 @@ public class IOSPlatform : Platform
 
 	public override void PreBuildAgenda(UE4Build Build, UE4Build.BuildAgenda Agenda)
 	{
-		if (UnrealBuildTool.BuildHostPlatform.Current.Platform != UnrealTargetPlatform.Mac && !Automation.IsEngineInstalled())
+		if (UnrealBuildTool.BuildHostPlatform.Current.Platform != UnrealTargetPlatform.Mac && !CommandUtils.IsEngineInstalled())
 		{
 			Agenda.DotNetProjects.Add(@"Engine\Source\Programs\IOS\MobileDeviceInterface\MobileDeviceInterface.csproj");
 			Agenda.DotNetProjects.Add(@"Engine\Source\Programs\IOS\iPhonePackager\iPhonePackager.csproj");
diff --git a/Engine/Source/Programs/AutomationTool/Scripts/BuildProjectCommand.Automation.cs b/Engine/Source/Programs/AutomationTool/Scripts/BuildProjectCommand.Automation.cs
--- a/Engine/Source/Programs/AutomationTool/Scripts/BuildProjectCommand.Automation.cs
+++ b/Engine/Source/Programs/AutomationTool/Scripts/BuildProjectCommand.Automation.cs
@@ -62,7 +62,7 @@ public partial class Project : CommandUtils
 		{
 			return;
 		}
-		if (Automation.IsEngineInstalled() && !Params.IsCodeBasedProject)
+		if (CommandUtils.IsEngineInstalled() && !Params.IsCodeBasedProject)
 		{
 			return;
 		}
@@ -82,7 +82,7 @@ public partial class Project : CommandUtils
 
             Agenda.AddTargets(Params.EditorTargets.ToArray(), EditorPlatform, EditorConfiguration, Params.CodeBasedUprojectPath);
 
-			if(!Automation.IsEngineInstalled())
+			if(!CommandUtils.IsEngineInstalled())
 			{
 				CrashReportPlatforms.Add(EditorPlatform);
 				if (Params.EditorTargets.Contains("UnrealHeaderTool") == false)
@@ -110,7 +110,7 @@ public partial class Project : CommandUtils
 		}
 
 		// Build any tools we need to stage
-		if ((TargetMask & ProjectBuildTargets.UnrealPak) == ProjectBuildTargets.UnrealPak && !Automation.IsEngineInstalled())
+		if ((TargetMask & ProjectBuildTargets.UnrealPak) == ProjectBuildTargets.UnrealPak && !CommandUtils.IsEngineInstalled())
 		{
 			if (Params.EditorTargets.Contains("UnrealPak") == false)
 			{
@@ -170,7 +170,7 @@ public partial class Project : CommandUtils
 				}
 			}
 		}
-		if (!Params.NoBootstrapExe && !Automation.IsEngineInstalled() && (TargetMask & ProjectBuildTargets.Bootstrap) == ProjectBuildTargets.Bootstrap)
+		if (!Params.NoBootstrapExe && !CommandUtils.IsEngineInstalled() && (TargetMask & ProjectBuildTargets.Bootstrap) == ProjectBuildTargets.Bootstrap)
 		{
 			UnrealBuildTool.UnrealTargetPlatform[] BootstrapPackagedGamePlatforms = { UnrealBuildTool.UnrealTargetPlatform.Win32, UnrealBuildTool.UnrealTargetPlatform.Win64 };
 			foreach(UnrealBuildTool.UnrealTargetPlatform BootstrapPackagedGamePlatformType in BootstrapPackagedGamePlatforms)
@@ -181,7 +181,7 @@ public partial class Project : CommandUtils
 				}
 			}
 		}
-		if (Params.CrashReporter && !Automation.IsEngineInstalled() && (TargetMask & ProjectBuildTargets.CrashReporter) == ProjectBuildTargets.CrashReporter)
+		if (Params.CrashReporter && !CommandUtils.IsEngineInstalled() && (TargetMask & ProjectBuildTargets.CrashReporter) == ProjectBuildTargets.CrashReporter)
 		{
 			foreach (var CrashReportPlatform in CrashReportPlatforms)
 			{
diff --git a/Engine/Source/Programs/AutomationTool/Scripts/CookCommand.Automation.cs b/Engine/Source/Programs/AutomationTool/Scripts/CookCommand.Automation.cs
--- a/Engine/Source/Programs/AutomationTool/Scripts/CookCommand.Automation.cs
+++ b/Engine/Source/Programs/AutomationTool/Scripts/CookCommand.Automation.cs
@@ -48,7 +48,7 @@ public partial class Project : CommandUtils
 			if (Params.ClientTargetPlatforms.Count > 0)
 			{
 				var LogFolderOutsideOfSandbox = GetLogFolderOutsideOfSandbox();
-				if (!GlobalCommandLine.Installed)
+				if (!CommandUtils.IsEngineInstalled())
 				{
 					// In the installed runs, this is the same folder as CmdEnv.LogFolder so delete only in not-installed
 					DeleteDirectory(LogFolderOutsideOfSandbox);
diff --git a/Engine/Source/Programs/AutomationTool/Scripts/RunProjectCommand.Automation.cs b/Engine/Source/Programs/AutomationTool/Scripts/RunProjectCommand.Automation.cs
--- a/Engine/Source/Programs/AutomationTool/Scripts/RunProjectCommand.Automation.cs
+++ b/Engine/Source/Programs/AutomationTool/Scripts/RunProjectCommand.Automation.cs
@@ -1,4 +1,4 @@
-﻿// Copyright 1998-2018 Epic Games, Inc. All Rights Reserved.
+// Copyright 1998-2018 Epic Games, Inc. All Rights Reserved.
 using System;
 using System.Collections.Generic;
 using System.IO;
@@ -64,7 +64,7 @@ public partial class Project : CommandUtils
 	/// <returns></returns>
 	private static string GetLogFolderOutsideOfSandbox()
 	{
-		return GlobalCommandLine.Installed ?
+		return CommandUtils.IsEngineInstalled() ?
 			CmdEnv.LogFolder :
 			CombinePaths(Path.GetTempPath(), CommandUtils.EscapePath(CmdEnv.LocalRoot), "Logs");
 	}
@@ -74,7 +74,7 @@ public partial class Project : CommandUtils
 	/// </summary>
 	private static void CopyLogsBackToLogFolder()
 	{
-		if (!GlobalCommandLine.Installed)
+		if (!CommandUtils.IsEngineInstalled())
 		{
 			var LogFolderOutsideOfSandbox = GetLogFolderOutsideOfSandbox();
 			var TempLogFiles = FindFiles_NoExceptions("*", false, LogFolderOutsideOfSandbox);
@@ -97,7 +97,7 @@ public partial class Project : CommandUtils
 		LogInformation("********** RUN COMMAND STARTED **********");
 
 		var LogFolderOutsideOfSandbox = GetLogFolderOutsideOfSandbox();
-		if (!GlobalCommandLine.Installed && ServerProcess == null)
+		if (!CommandUtils.IsEngineInstalled() && ServerProcess == null)
 		{
 			// In the installed runs, this is the same folder as CmdEnv.LogFolder so delete only in not-installed
 			DeleteDirectory(LogFolderOutsideOfSandbox);
diff --git a/Engine/Source/Programs/UnrealBuildTool/System/PlatformExports.cs b/Engine/Source/Programs/UnrealBuildTool/System/PlatformExports.cs
--- a/Engine/Source/Programs/UnrealBuildTool/System/PlatformExports.cs
+++ b/Engine/Source/Programs/UnrealBuildTool/System/PlatformExports.cs
@@ -149,12 +149,9 @@ namespace UnrealBuildTool
 		/// <summary>
 		/// Initialize UBT in the context of another host process (presumably UAT)
 		/// </summary>
-		/// <param name="bIsEngineInstalled">Whether the engine is installed</param>
 		/// <returns>True if initialization was successful</returns>
-		public static bool Initialize(bool bIsEngineInstalled)
+		public static bool Initialize()
 		{
-			UnrealBuildTool.SetIsEngineInstalled(bIsEngineInstalled);
-
 			// Read the XML configuration files
 			if(!XmlConfig.ReadConfigFiles(false))
 			{
diff --git a/Engine/Source/Programs/UnrealBuildTool/UnrealBuildTool.cs b/Engine/Source/Programs/UnrealBuildTool/UnrealBuildTool.cs
--- a/Engine/Source/Programs/UnrealBuildTool/UnrealBuildTool.cs
+++ b/Engine/Source/Programs/UnrealBuildTool/UnrealBuildTool.cs
@@ -131,15 +131,6 @@ namespace UnrealBuildTool
 			private set;
 		}
 
-		/// <summary>
-		/// Sets the global flag indicating whether the engine is installed
-		/// </summary>
-		/// <param name="bInIsEngineInstalled">Whether the engine is installed or not</param>
-		public static void SetIsEngineInstalled(bool bInIsEngineInstalled)
-		{
-			bIsEngineInstalled = bInIsEngineInstalled;
-		}
-
 		/// <summary>
 		/// Returns true if UnrealBuildTool is running using installed Engine components
 		/// </summary>
@@ -148,7 +139,7 @@ namespace UnrealBuildTool
 		{
 			if (!bIsEngineInstalled.HasValue)
 			{
-				throw new BuildException("IsEngineInstalled() called before being initialized.");
+				bIsEngineInstalled = FileReference.Exists(FileReference.Combine(EngineDirectory, "Build", "InstalledBuild.txt"));
 			}
 			return bIsEngineInstalled.Value;
 		}
@@ -453,24 +444,6 @@ namespace UnrealBuildTool
 						Directory.SetCurrentDirectory(EngineSourceDirectory);
 					}
 
-					// Figure out if the engine is installed or not
-					foreach (string Argument in Arguments)
-					{
-						string LowercaseArg = Argument.ToLowerInvariant();
-						if (LowercaseArg == "-installed" || LowercaseArg == "-installedengine")
-						{
-							bIsEngineInstalled = true;
-						}
-						else if (LowercaseArg == "-notinstalledengine")
-						{
-							bIsEngineInstalled = false;
-						}
-					}
-					if (!bIsEngineInstalled.HasValue)
-					{
-						bIsEngineInstalled = FileReference.Exists(FileReference.Combine(RootDirectory, "Engine", "Build", "InstalledBuild.txt"));
-					}
-
 					// Read the XML configuration files
 					bool bForceXmlConfigCache = Arguments.Any(x => x.Equals("-ForceXmlConfigCache", StringComparison.InvariantCultureIgnoreCase));
 					if (!XmlConfig.ReadConfigFiles(bForceXmlConfigCache))
