Missing file.

#rnx

diff --git a/Engine/Source/Programs/AutomationTool/BuildGraph/Tasks/PakFileTask.cs b/Engine/Source/Programs/AutomationTool/BuildGraph/Tasks/PakFileTask.cs
--- a/Engine/Source/Programs/AutomationTool/BuildGraph/Tasks/PakFileTask.cs
+++ b/Engine/Source/Programs/AutomationTool/BuildGraph/Tasks/PakFileTask.cs
@@ -149,7 +149,7 @@ namespace BuildGraph.Tasks
 			{
 				CommandLine.AppendFormat(" -order={0}", CommandUtils.MakePathSafeToUseWithCommandLine(Parameters.Order.FullName));
 			}
-			if (GlobalCommandLine.Installed)
+			if (CommandUtils.IsEngineInstalled())
 			{
 				CommandLine.Append(" -installed");
 			}
