- Add support for linux installed builds
- Add win64 host support to DDC Group Linux
- Change project params split to remove empty entries
- Change BuildDerivedDataCache to skip building a feature pack if there are no project specific platforms for it out of the available targetplatforms

diff --git a/Engine/Build/InstalledEngineBuild.xml b/Engine/Build/InstalledEngineBuild.xml
--- a/Engine/Build/InstalledEngineBuild.xml
+++ b/Engine/Build/InstalledEngineBuild.xml
@@ -616,8 +616,11 @@
 	<Property Name="DDCPlatformsLinux" Value="Linux"/>
 	<Property Name="DDCPlatformsLinux" Value="$(DDCPlatformsLinux)+LinuxNoEditor" If="'$(WithLinux)' == true"/>
 
-	<Agent Name="DDC Group Linux" Type="Linux">
-		<Node Name="Build DDC Linux" Requires="#UE4Editor Linux;Build Tools Linux;#Build Tools CS Binaries" Produces="#CompressedDDCLinux">
+	<Agent Name="DDC Group Linux" Type="Linux;Win64">
+		<!-- Additionally require the Win64 editor binaries if building on win64 -->
+		<Property Name="DDCLinuxRequires" Value="#UE4Editor Linux;Build Tools Linux;#Build Tools CS Binaries"/>
+		<Property Name="DDCLinuxRequires" Value="$(DDCLinuxRequires);#UE4Editor Win64;#Build Tools Win64" If="'$(HostPlatform)' == 'Win64'"/>
+		<Node Name="Build DDC Linux" Requires="$(DDCLinuxRequires)" Produces="#CompressedDDCLinux">
 			<!-- Get our temp dir -->
 			<Property Name="DDCDir" Value="$(RootDir)/LocalBuilds/InstalledDDC"/>
 			<Property Name="SavedDir" Value="$(SavedOutput)/Installed/Linux"/>
@@ -626,6 +629,7 @@
 			<!-- Build up a list of files needed to build DDC -->
 			<CsCompile Project="Engine/Source/Programs/UnrealBuildTool/UnrealBuildTool.csproj" Configuration="Development" Platform="AnyCPU" Tag="#ToCopy" TagReferences="#ToCopy" EnumerateOnly="true"/>
 			<Tag Files="#Compile UnrealHeaderTool $(HostPlatform);#UE4Editor Linux;#Build Tools Linux" With="#ToCopy"/>
+			<Tag Files="#UE4Editor Win64;#Build Tools Win64" With="#ToCopy" If="'$(HostPlatform)' == 'Win64'"/>
 			<Tag Files="#ToCopy" Filter="*.target" With="#TargetReceipts"/>
 			<TagReceipt Files="#TargetReceipts" RuntimeDependencies="true" With="#ToCopy"/>
 			<Tag Files="Engine/Content/..." Except="*.psd;*.pdn;*.fbx;*.po" With="#ToCopy"/>
@@ -647,7 +651,7 @@
 			<Copy From="$(RootDir)" To="$(DDCDir)" Files="#FilteredCopyList"/>
 			<Copy From="$(CsToolsDir)" To="$(DDCDir)" Files="#Build Tools CS Binaries"/>
 
-			<Command Name="BuildDerivedDataCache" Arguments="-TempDir=&quot;$(DDCDir)&quot; -FeaturePacks=$(ProjectsToBuildDDC) -TargetPlatforms=$(DDCPlatformsLinux) -HostPlatform=Linux -SavedDir=&quot;$(SavedDir)&quot;"/>
+			<Command Name="BuildDerivedDataCache" Arguments="-TempDir=&quot;$(DDCDir)&quot; -FeaturePacks=$(ProjectsToBuildDDC) -TargetPlatforms=$(DDCPlatformsLinux) -SavedDir=&quot;$(SavedDir)&quot;"/>
 
 			<Tag Files="$(SavedDir)/Engine/DerivedDataCache/Compressed.ddp" With="#CompressedDDCLinux"/>
 		</Node>
diff --git a/Engine/Source/Programs/AutomationTool/AutomationUtils/ProjectParams.cs b/Engine/Source/Programs/AutomationTool/AutomationUtils/ProjectParams.cs
--- a/Engine/Source/Programs/AutomationTool/AutomationUtils/ProjectParams.cs
+++ b/Engine/Source/Programs/AutomationTool/AutomationUtils/ProjectParams.cs
@@ -161,7 +161,7 @@ namespace AutomationTool
                         string CmdLineCookFlavor = Command.ParseParamValue("cookflavor");
                         if (!String.IsNullOrEmpty(CmdLineCookFlavor))
                         {
-                            CookFlavors = new List<string>(CmdLineCookFlavor.Split('+'));
+                            CookFlavors = new List<string>(CmdLineCookFlavor.Split(new char[] { '+' }, StringSplitOptions.RemoveEmptyEntries));
                         }
                     }
 
@@ -169,11 +169,11 @@ namespace AutomationTool
 					{
 						// Get all platforms from the param value: Platform_1+Platform_2+...+Platform_k
 						TargetPlatforms = new List<TargetPlatformDescriptor>();
-						var PlatformNames = (new HashSet<string>(CmdLinePlatform.Split('+'))).ToList();
+						var PlatformNames = (new HashSet<string>(CmdLinePlatform.Split(new char[] { '+' }, StringSplitOptions.RemoveEmptyEntries))).ToList();
 						foreach (var PlatformName in PlatformNames)
 						{
                             // Look for dependent platforms, Source_1.Dependent_1+Source_2.Dependent_2+Standalone_3
-                            var SubPlatformNames = new List<string>(PlatformName.Split('.'));
+                            var SubPlatformNames = new List<string>(PlatformName.Split(new char[] { '.' }, StringSplitOptions.RemoveEmptyEntries));
 
                             foreach (var SubPlatformName in SubPlatformNames)
                             {
@@ -524,7 +524,7 @@ namespace AutomationTool
                     var CookCulturesString = Command.ParseParamValue("CookCultures");
                     if (CookCulturesString != null)
                     {
-                        this.CulturesToCook = new ParamList<string>(CookCulturesString.Split(','));
+                        this.CulturesToCook = new ParamList<string>(CookCulturesString.Split(new char[] { ',' }, StringSplitOptions.RemoveEmptyEntries));
                     }
                 }
             }
@@ -780,7 +780,7 @@ namespace AutomationTool
             }
             else
             {
-                this.Devices = new ParamList<string>(DeviceString.Split('+'));
+                this.Devices = new ParamList<string>(DeviceString.Split(new char[] { '+' }, StringSplitOptions.RemoveEmptyEntries));
                 this.DeviceNames = new ParamList<string>();
                 foreach (var d in this.Devices)
                 {
@@ -828,7 +828,7 @@ namespace AutomationTool
                     if (ClientConfig != null)
 					{
 						this.ClientConfigsToBuild = new List<UnrealTargetConfiguration>();
-						var Configs = new ParamList<string>(ClientConfig.Split('+'));
+						var Configs = new ParamList<string>(ClientConfig.Split(new char[] { '+' }, StringSplitOptions.RemoveEmptyEntries));
 						foreach (var ConfigName in Configs)
 						{
 							this.ClientConfigsToBuild.Add((UnrealTargetConfiguration)Enum.Parse(typeof(UnrealTargetConfiguration), ConfigName, true));
@@ -850,7 +850,7 @@ namespace AutomationTool
                     var PortString = Command.ParseParamValue("port");
                     if (String.IsNullOrEmpty(PortString) == false)
                     {
-                        var Ports = new ParamList<string>(PortString.Split('+'));
+                        var Ports = new ParamList<string>(PortString.Split(new char[] { '+' }, StringSplitOptions.RemoveEmptyEntries));
                         foreach (var P in Ports)
                         {
                             this.Port.Add(P);
@@ -873,7 +873,7 @@ namespace AutomationTool
                     var MapsString = Command.ParseParamValue("MapsToCook");
                     if (String.IsNullOrEmpty(MapsString) == false)
                     {
-                        var MapNames = new ParamList<string>(MapsString.Split('+'));
+                        var MapNames = new ParamList<string>(MapsString.Split(new char[] { '+' }, StringSplitOptions.RemoveEmptyEntries));
                         foreach ( var M in MapNames ) 
                         {
                             this.MapsToCook.Add( M );
@@ -895,7 +895,7 @@ namespace AutomationTool
 					var MapsString = Command.ParseParamValue("MapIniSectionsToCook");
 					if (String.IsNullOrEmpty(MapsString) == false)
 					{
-						var MapNames = new ParamList<string>(MapsString.Split('+'));
+						var MapNames = new ParamList<string>(MapsString.Split(new char[] { '+' }, StringSplitOptions.RemoveEmptyEntries));
 						foreach (var M in MapNames)
 						{
 							this.MapIniSectionsToCook.Add(M);
@@ -922,7 +922,7 @@ namespace AutomationTool
 					var MapsString = Command.ParseParamValue("MapsToRebuildLightMaps");
 					if (String.IsNullOrEmpty(MapsString) == false)
 					{
-						var MapNames = new ParamList<string>(MapsString.Split('+'));
+						var MapNames = new ParamList<string>(MapsString.Split(new char[] { '+' }, StringSplitOptions.RemoveEmptyEntries));
 						foreach (var M in MapNames)
 						{
 							this.MapsToRebuildLightMaps.Add(M);
@@ -944,7 +944,7 @@ namespace AutomationTool
                     var MapsString = Command.ParseParamValue("MapsToRebuildHLODMaps");
                     if (String.IsNullOrEmpty(MapsString) == false)
                     {
-                        var MapNames = new ParamList<string>(MapsString.Split('+'));
+                        var MapNames = new ParamList<string>(MapsString.Split(new char[] { '+' }, StringSplitOptions.RemoveEmptyEntries));
                         foreach (var M in MapNames)
                         {
                             this.MapsToRebuildHLODMaps.Add(M);
@@ -966,7 +966,7 @@ namespace AutomationTool
 					var TitleString = Command.ParseParamValue("TitleID");
 					if (String.IsNullOrEmpty(TitleString) == false)
 					{
-						var TitleIDs = new ParamList<string>(TitleString.Split('+'));
+						var TitleIDs = new ParamList<string>(TitleString.Split(new char[] { '+' }, StringSplitOptions.RemoveEmptyEntries));
 						foreach (var T in TitleIDs)
 						{
 							this.TitleID.Add(T);
@@ -991,7 +991,7 @@ namespace AutomationTool
                     if (ServerConfig != null)
 					{
 						this.ServerConfigsToBuild = new List<UnrealTargetConfiguration>();
-						var Configs = new ParamList<string>(ServerConfig.Split('+'));
+						var Configs = new ParamList<string>(ServerConfig.Split(new char[] { '+' }, StringSplitOptions.RemoveEmptyEntries));
 						foreach (var ConfigName in Configs)
 						{
 							this.ServerConfigsToBuild.Add((UnrealTargetConfiguration)Enum.Parse(typeof(UnrealTargetConfiguration), ConfigName, true));
diff --git a/Engine/Source/Programs/AutomationTool/Scripts/BuildDerivedDataCache.Automation.cs b/Engine/Source/Programs/AutomationTool/Scripts/BuildDerivedDataCache.Automation.cs
--- a/Engine/Source/Programs/AutomationTool/Scripts/BuildDerivedDataCache.Automation.cs
+++ b/Engine/Source/Programs/AutomationTool/Scripts/BuildDerivedDataCache.Automation.cs
@@ -17,7 +17,7 @@ public class BuildDerivedDataCache : BuildCommand
 		// Get the list of platform names
 		string[] FeaturePacks = ParseParamValue("FeaturePacks").Split(';');
 		string TempDir = ParseParamValue("TempDir");
-		UnrealTargetPlatform HostPlatform = (UnrealTargetPlatform)Enum.Parse(typeof(UnrealTargetPlatform), ParseParamValue("HostPlatform"));
+		UnrealTargetPlatform HostPlatform = BuildHostPlatform.Current.Platform;
 		string TargetPlatforms = ParseParamValue("TargetPlatforms");
 		string SavedDir = ParseParamValue("SavedDir");
 
@@ -72,10 +72,13 @@ public class BuildDerivedDataCache : BuildCommand
 						FilteredPlatforms.Add(TargetPlatform);
 					}
 				}
-
+				if(FilteredPlatforms.Count == 0)
+				{
+					LogInformation("Did not find any project specific platforms for FeaturePack {0} out of supplied TargetPlatforms {1}, skipping it!", GameName, ProjectSpecificPlatforms);
+					continue;
+				}
 				ProjectSpecificPlatforms = CommandUtils.CombineCommandletParams(FilteredPlatforms.Distinct().ToArray());
 			}
-
 			CommandUtils.LogInformation("Generating DDC data for {0} on {1}", GameName, ProjectSpecificPlatforms);
 			CommandUtils.DDCCommandlet(FileRef, EditorExe, null, ProjectSpecificPlatforms, "-fill -DDC=CreateInstalledEnginePak -ProjectOnly");
 
