Copy minidumps generated on build machines into the UAT log folder.

diff --git a/Engine/Source/Programs/AutomationTool/AutomationUtils/CommandletUtils.cs b/Engine/Source/Programs/AutomationTool/AutomationUtils/CommandletUtils.cs
--- a/Engine/Source/Programs/AutomationTool/AutomationUtils/CommandletUtils.cs
+++ b/Engine/Source/Programs/AutomationTool/AutomationUtils/CommandletUtils.cs
@@ -42,8 +42,6 @@ namespace AutomationTool
 
 	public partial class CommandUtils
 	{
-		#region Commandlets
-
 		/// <summary>
 		/// Runs Cook commandlet.
 		/// </summary>
@@ -282,9 +280,40 @@ namespace AutomationTool
 				Args += " -UTF8Output";
 				Opts |= ERunOptions.UTF8Output;
 			}
-			var RunResult = Run(EditorExe, Args, Options: Opts, Identifier: Commandlet);
+			IProcessResult RunResult = Run(EditorExe, Args, Options: Opts, Identifier: Commandlet);
 			PopDir();
 
+			// If we're running on a Windows build machine, copy any crash dumps into the log folder
+			if(HostPlatform.Current.HostEditorPlatform == UnrealTargetPlatform.Win64 && IsBuildMachine)
+			{
+				DirectoryInfo CrashesDir = new DirectoryInfo(DirectoryReference.Combine(DirectoryReference.FromFile(ProjectName) ?? CommandUtils.EngineDirectory, "Saved", "Crashes").FullName);
+				if(CrashesDir.Exists)
+				{
+					foreach(DirectoryInfo CrashDir in CrashesDir.EnumerateDirectories())
+					{
+						if(CrashDir.LastWriteTimeUtc > StartTime)
+						{
+							DirectoryInfo OutputCrashesDir = new DirectoryInfo(Path.Combine(CmdEnv.LogFolder, "Crashes", CrashDir.Name));
+							try
+							{
+								CommandUtils.LogInformation("Copying crash data to {0}...", OutputCrashesDir.FullName);
+								OutputCrashesDir.Create();
+
+								foreach(FileInfo CrashFile in CrashDir.EnumerateFiles())
+								{
+									CrashFile.CopyTo(Path.Combine(OutputCrashesDir.FullName, CrashFile.Name));
+								}
+							}
+							catch(Exception Ex)
+							{
+								CommandUtils.LogWarning("Unable to copy crash data; skipping. See log for exception details.");
+								CommandUtils.LogVerbose(Tools.DotNETCommon.ExceptionUtils.FormatExceptionDetails(Ex));
+							}
+						}
+					}
+				}
+			}
+
 			// If we're running on a Mac, dump all the *.crash files that were generated while the editor was running.
 			if(HostPlatform.Current.HostEditorPlatform == UnrealTargetPlatform.Mac)
 			{
@@ -475,7 +504,5 @@ namespace AutomationTool
 			}
 			return ProjectFullPath;
 		}
-
-		#endregion
 	}
 }
