UAT: Allow targets to be built with parallel executor *and* XGE in a single UE4Build instance. Some platforms don't support XGE (but do support parallel exector), so we should still allow them to build in parallel if possible.

diff --git a/Engine/Source/Programs/AutomationTool/AutomationUtils/UE4Build.cs b/Engine/Source/Programs/AutomationTool/AutomationUtils/UE4Build.cs
--- a/Engine/Source/Programs/AutomationTool/AutomationUtils/UE4Build.cs
+++ b/Engine/Source/Programs/AutomationTool/AutomationUtils/UE4Build.cs
@@ -94,7 +94,7 @@ namespace AutomationTool
 			public string OutputCaption;
 		}
 
-		XGEItem XGEPrepareBuildWithUBT(string TargetName, UnrealTargetPlatform Platform, UnrealTargetConfiguration Config, FileReference UprojectPath, bool ForceMonolithic = false, bool ForceNonUnity = false, bool ForceDebugInfo = false, string InAddArgs = "", bool ForceUnity = false)
+		XGEItem XGEPrepareBuildWithUBT(string TargetName, UnrealTargetPlatform Platform, UnrealTargetConfiguration Config, FileReference UprojectPath, string InAddArgs)
 		{
 			string AddArgs = "";
 			if (UprojectPath != null)
@@ -103,22 +103,6 @@ namespace AutomationTool
 			}
 			AddArgs += " -NoUBTMakefiles";
 			AddArgs += " " + InAddArgs;
-			if (ForceMonolithic)
-			{
-				AddArgs += " -monolithic";
-			}
-			if (ForceNonUnity)
-			{
-				AddArgs += " -disableunity";
-			}
-			if (ForceUnity)
-			{
-				AddArgs += " -forceunity";
-			}
-			if (ForceDebugInfo)
-			{
-				AddArgs += " -forcedebuginfo";
-			}
 
 			PrepareUBT();
 
@@ -182,7 +166,7 @@ namespace AutomationTool
 			}
 		}
 
-		void CleanWithUBT(string TargetName, UnrealTargetPlatform Platform, UnrealTargetConfiguration Config, FileReference UprojectPath, bool ForceMonolithic = false, bool ForceNonUnity = false, bool ForceDebugInfo = false, string InAddArgs = "", bool ForceUnity = false)
+		void CleanWithUBT(string TargetName, UnrealTargetPlatform Platform, UnrealTargetConfiguration Config, FileReference UprojectPath, string InAddArgs = "")
 		{
 			string AddArgs = "";
 			if (UprojectPath != null)
@@ -191,22 +175,6 @@ namespace AutomationTool
 			}
 			AddArgs += " -NoUBTMakefiles";
 			AddArgs += " " + InAddArgs;
-			if (ForceMonolithic)
-			{
-				AddArgs += " -monolithic";
-			}
-			if (ForceNonUnity)
-			{
-				AddArgs += " -disableunity";
-			}
-			if (ForceUnity)
-			{
-				AddArgs += " -forceunity";
-			}
-			if (ForceDebugInfo)
-			{
-				AddArgs += " -forcedebuginfo";
-			}
 			if (!TargetName.Equals("UnrealHeaderTool", StringComparison.InvariantCultureIgnoreCase))
 			{
 				AddArgs += " -nobuilduht";
@@ -219,7 +187,7 @@ namespace AutomationTool
 			}
         }
 
-		BuildManifest BuildWithUBT(string TargetName, UnrealTargetPlatform TargetPlatform, UnrealTargetConfiguration Config, FileReference UprojectPath, bool ForceMonolithic = false, bool ForceNonUnity = false, bool ForceDebugInfo = false, bool ForceFlushMac = false, bool DisableXGE = false, string InAddArgs = "", bool ForceUnity = false)
+		BuildManifest BuildWithUBT(string TargetName, UnrealTargetPlatform TargetPlatform, UnrealTargetConfiguration Config, FileReference UprojectPath, bool ForceFlushMac = false, bool DisableXGE = false, string InAddArgs = "")
 		{
 			string AddArgs = "";
 			if (UprojectPath != null)
@@ -228,22 +196,6 @@ namespace AutomationTool
 			}
 			AddArgs += " -NoUBTMakefiles";
 			AddArgs += " " + InAddArgs;
-			if (ForceMonolithic)
-			{
-				AddArgs += " -monolithic";
-			}
-			if (ForceNonUnity)
-			{
-				AddArgs += " -disableunity";
-			}
-			if (ForceUnity)
-			{
-				AddArgs += " -forceunity";
-			}
-			if (ForceDebugInfo)
-			{
-				AddArgs += " -forcedebuginfo";
-			}
 			if (ForceFlushMac)
 			{
 				AddArgs += " -flushmac";
@@ -1043,7 +995,7 @@ namespace AutomationTool
 		/// <param name="InUpdateVersionFiles">True if the version files are to be updated </param>
 		/// <param name="InForceNoXGE">If true will force XGE off</param>
 		/// <param name="InUseParallelExecutor">If true AND XGE not present or not being used then use ParallelExecutor</param>
-		public void Build(BuildAgenda Agenda, bool? InDeleteBuildProducts = null, bool InUpdateVersionFiles = true, bool InForceNoXGE = false, bool InUseParallelExecutor = false, bool InForceNonUnity = false, bool InForceUnity = false, bool InShowProgress = false, int? InChangelistNumberOverride = null, Dictionary<BuildTarget, BuildManifest> InTargetToManifest = null)
+		public void Build(BuildAgenda Agenda, bool? InDeleteBuildProducts = null, bool InUpdateVersionFiles = true, bool InForceNoXGE = false, bool InUseParallelExecutor = false, bool InShowProgress = false, int? InChangelistNumberOverride = null, Dictionary<BuildTarget, BuildManifest> InTargetToManifest = null)
 		{
 			if (!CommandUtils.CmdEnv.HasCapabilityToCompile)
 			{
@@ -1114,20 +1066,13 @@ namespace AutomationTool
 				AddBuildProductsForCSharpProj(Path.Combine(CommandUtils.CmdEnv.LocalRoot, File));
 			}
 
-			bool bForceMonolithic = ParseParam("ForceMonolithic");
-			bool bForceNonUnity = ParseParam("ForceNonUnity") || InForceNonUnity;
-			bool bForceUnity = ParseParam("ForceUnity") || InForceUnity;
-			bool bForceDebugInfo = ParseParam("ForceDebugInfo");
 			string XGEConsole = null;
 			bool bDisableXGE = ParseParam("NoXGE") || InForceNoXGE;
 			bool bCanUseXGE = !bDisableXGE && PlatformExports.TryGetXgConsoleExecutable(out XGEConsole);
 
 			// only run ParallelExecutor if not running XGE (and we've requested ParallelExecutor and it exists)
-			bool bCanUseParallelExecutor = !bCanUseXGE && InUseParallelExecutor && (HostPlatform.Current.HostEditorPlatform == UnrealTargetPlatform.Win64);
+			bool bCanUseParallelExecutor = InUseParallelExecutor && (HostPlatform.Current.HostEditorPlatform == UnrealTargetPlatform.Win64);
 			CommandUtils.LogLog("************************* UE4Build:");
-			CommandUtils.LogLog("************************* ForceMonolithic: {0}", bForceMonolithic);
-			CommandUtils.LogLog("************************* ForceNonUnity:{0} ", bForceNonUnity);
-			CommandUtils.LogLog("************************* ForceDebugInfo: {0}", bForceDebugInfo);
 			CommandUtils.LogLog("************************* UseXGE: {0}", bCanUseXGE);
 			CommandUtils.LogLog("************************* UseParallelExecutor: {0}", bCanUseParallelExecutor);
 
@@ -1137,112 +1082,105 @@ namespace AutomationTool
 				bool bClean = Target.Clean ?? DeleteBuildProducts;
 				if (bClean)
 				{
-					CleanWithUBT(Target.TargetName, Target.Platform, Target.Config, Target.UprojectPath, bForceMonolithic, bForceNonUnity, bForceDebugInfo, Target.UBTArgs, bForceUnity);
+					CleanWithUBT(Target.TargetName, Target.Platform, Target.Config, Target.UprojectPath, Target.UBTArgs);
 				}
 			}
 
 			// Filter the targets into those which can be built in parallel, vs those that must be executed serially
-			List<BuildTarget> NonParallelTargets = new List<BuildTarget>();
+			List<BuildTarget> NormalTargets = new List<BuildTarget>();
+			List<BuildTarget> ParallelXgeTargets = new List<BuildTarget>();
 			List<BuildTarget> ParallelTargets = new List<BuildTarget>();
 			foreach (BuildTarget Target in Agenda.Targets)
 			{
 				if(Target.TargetName == "UnrealHeaderTool")
 				{
-					NonParallelTargets.Insert(0, Target);
+					NormalTargets.Insert(0, Target);
 				}
-				else if(bCanUseXGE)
+				else if(bCanUseXGE && CanUseXGE(Target.Platform))
 				{
-					if(CanUseXGE(Target.Platform))
-					{
-						ParallelTargets.Add(Target);
-					}
-					else
-					{
-						NonParallelTargets.Add(Target);
-					}
+					ParallelXgeTargets.Add(Target);
 				}
-				else if(bCanUseParallelExecutor)
+				else if(bCanUseParallelExecutor && CanUseParallelExecutor(Target.Platform))
 				{
-					if(CanUseParallelExecutor(Target.Platform))
-					{
-						ParallelTargets.Add(Target);
-					}
-					else
-					{
-						NonParallelTargets.Add(Target);
-					}
+					ParallelTargets.Add(Target);
 				}
 				else
 				{
-					NonParallelTargets.Add(Target);
+					NormalTargets.Add(Target);
 				}
 			}
 
 			// Execute all the serial targets
-			foreach(BuildTarget Target in NonParallelTargets)
+			foreach(BuildTarget Target in NormalTargets)
 			{
 				// When building a target for Mac or iOS, use UBT's -flushmac option to clean up the remote builder
 				bool bForceFlushMac = DeleteBuildProducts && (Target.Platform == UnrealTargetPlatform.Mac || Target.Platform == UnrealTargetPlatform.IOS);
-				BuildManifest Manifest = BuildWithUBT(Target.TargetName, Target.Platform, Target.Config, Target.UprojectPath, bForceMonolithic, bForceNonUnity, bForceDebugInfo, bForceFlushMac, bDisableXGE, Target.UBTArgs, bForceUnity);
+				BuildManifest Manifest = BuildWithUBT(Target.TargetName, Target.Platform, Target.Config, Target.UprojectPath, bForceFlushMac, bDisableXGE, Target.UBTArgs);
 				if(InTargetToManifest != null)
 				{
 					InTargetToManifest[Target] = Manifest;
 				}
 			}
 
+			// Execute all the XGE targets
+			if(ParallelXgeTargets.Count > 0)
+			{
+				BuildParallelTargets(ParallelXgeTargets, InShowProgress, XGEConsole, InTargetToManifest);
+			}
+
 			// Execute all the parallel targets
 			if(ParallelTargets.Count > 0)
 			{
-				string TaskFilePath = CommandUtils.CombinePaths(CommandUtils.CmdEnv.LogFolder, @"UAT_XGE.xml");
+				BuildParallelTargets(ParallelTargets, InShowProgress, null, InTargetToManifest);
+			}
+		}
+
+		private void BuildParallelTargets(List<BuildTarget> ParallelTargets, bool InShowProgress, string XGETool, Dictionary<BuildTarget, BuildManifest> InTargetToManifest)
+		{
+			string TaskFilePath = CommandUtils.CombinePaths(CommandUtils.CmdEnv.LogFolder, @"UAT_XGE.xml");
 
-				CommandUtils.LogSetProgress(InShowProgress, "Generating headers...");
+			CommandUtils.LogSetProgress(InShowProgress, "Generating headers...");
 
-				List<XGEItem> XGEItems = new List<XGEItem>();
-				foreach (BuildTarget Target in ParallelTargets)
+			List<XGEItem> XGEItems = new List<XGEItem>();
+			foreach (BuildTarget Target in ParallelTargets)
+			{
+				XGEItem Item = XGEPrepareBuildWithUBT(Target.TargetName, Target.Platform, Target.Config, Target.UprojectPath, Target.UBTArgs);
+				if(InTargetToManifest != null)
 				{
-					XGEItem Item = XGEPrepareBuildWithUBT(Target.TargetName, Target.Platform, Target.Config, Target.UprojectPath, bForceMonolithic, bForceNonUnity, bForceDebugInfo, Target.UBTArgs, bForceUnity);
-					if(InTargetToManifest != null)
-					{
-						InTargetToManifest[Target] = Item.Manifest;
-					}
-					XGEItems.Add(Item);
+					InTargetToManifest[Target] = Item.Manifest;
 				}
+				XGEItems.Add(Item);
+			}
 
-				string XGETool = null;
-				string Args = null;
-				if (bCanUseXGE) 
+			string Args = null;
+			if (XGETool != null) 
+			{
+				Args = "\"" + TaskFilePath + "\" /Rebuild /NoLogo /ShowAgent /ShowTime";
+				if (ParseParam("StopOnErrors"))
 				{
-					XGETool = XGEConsole;
-					Args = "\"" + TaskFilePath + "\" /Rebuild /NoLogo /ShowAgent /ShowTime";
-					if (ParseParam("StopOnErrors"))
-					{
-						Args += " /StopOnErrors";
-					}
-
-					// A bug in the UCRT can cause XGE to hang on VS2015 builds. Figure out if this hang is likely to effect this build and workaround it if able.
-					string XGEVersion = Microsoft.Win32.Registry.GetValue(@"HKEY_LOCAL_MACHINE\SOFTWARE\Wow6432Node\Xoreax\IncrediBuild\Builder", "Version", null) as string;
-					if (XGEVersion != null)
-					{
-						// Per Xoreax support, subtract 1001000 from the registry value to get the build number of the installed XGE.
-						int XGEBuildNumber;
-						if (Int32.TryParse(XGEVersion, out XGEBuildNumber) && XGEBuildNumber - 1001000 >= 1659)
-						{
-							Args += " /no_watchdog_thread";
-						}
-					}
+					Args += " /StopOnErrors";
 				}
 
-				CommandUtils.LogSetProgress(InShowProgress, "Building...");
-				if (!ProcessXGEItems(XGEItems, XGETool, Args, TaskFilePath, InShowProgress))
+				// A bug in the UCRT can cause XGE to hang on VS2015 builds. Figure out if this hang is likely to effect this build and workaround it if able.
+				string XGEVersion = Microsoft.Win32.Registry.GetValue(@"HKEY_LOCAL_MACHINE\SOFTWARE\Wow6432Node\Xoreax\IncrediBuild\Builder", "Version", null) as string;
+				if (XGEVersion != null)
 				{
-					throw new AutomationException("BUILD FAILED: {0} failed, retries not enabled:", XGETool);
+					// Per Xoreax support, subtract 1001000 from the registry value to get the build number of the installed XGE.
+					int XGEBuildNumber;
+					if (Int32.TryParse(XGEVersion, out XGEBuildNumber) && XGEBuildNumber - 1001000 >= 1659)
+					{
+						Args += " /no_watchdog_thread";
+					}
 				}
 			}
 
-			// NOTE: OK, we're done building projects.  You can build additional targets after this.  Call FinalizebuildProducts() when done.
+			CommandUtils.LogSetProgress(InShowProgress, "Building...");
+			if (!ProcessXGEItems(XGEItems, XGETool, Args, TaskFilePath, InShowProgress))
+			{
+				throw new AutomationException("BUILD FAILED: {0} failed, retries not enabled:", XGETool);
+			}
 		}
 
-
 		/// <summary>
 		/// Checks to make sure there was at least one build product, and that all files exist.  Also, logs them all out.
 		/// </summary>
