Merging //UE4/Dev-Main to Dev-Build (//UE4/Dev-Build)

#rnx

diff --git a/Engine/Plugins/Media/AvfMedia/Source/AvfMedia/Private/Player/AvfMediaVideoSampler.cpp b/Engine/Plugins/Media/AvfMedia/Source/AvfMedia/Private/Player/AvfMediaVideoSampler.cpp
--- a/Engine/Plugins/Media/AvfMedia/Source/AvfMedia/Private/Player/AvfMediaVideoSampler.cpp
+++ b/Engine/Plugins/Media/AvfMedia/Source/AvfMedia/Private/Player/AvfMediaVideoSampler.cpp
@@ -288,7 +288,8 @@ void FAvfMediaVideoSampler::Tick()
 				
 				PixelShader->SetParameters(RHICmdList, YTex, UVTex, MediaShaders::YuvToSrgbPs4, true);
 				
-				RHICmdList.SetStreamSource(0, CreateTempMediaVertexBuffer(), 0);
+				FVertexBufferRHIRef VertexBuffer = CreateTempMediaVertexBuffer();
+				RHICmdList.SetStreamSource(0, VertexBuffer, 0);
 				RHICmdList.SetViewport(0, 0, 0.0f, YWidth, YHeight, 1.0f);
 
 				RHICmdList.DrawPrimitive(PT_TriangleStrip, 0, 2, 1);
diff --git a/Engine/Source/Programs/AutomationTool/BuildGraph/Schema.cs b/Engine/Source/Programs/AutomationTool/BuildGraph/Schema.cs
--- a/Engine/Source/Programs/AutomationTool/BuildGraph/Schema.cs
+++ b/Engine/Source/Programs/AutomationTool/BuildGraph/Schema.cs
@@ -654,7 +654,8 @@ namespace AutomationTool
 			XmlSchemaComplexType BadgeType = new XmlSchemaComplexType();
 			BadgeType.Name = GetTypeName(ScriptSchemaStandardType.Badge);
 			BadgeType.Attributes.Add(CreateSchemaAttribute("Name", ScriptSchemaStandardType.Name, XmlSchemaUse.Required));
-			BadgeType.Attributes.Add(CreateSchemaAttribute("Requires", ScriptSchemaStandardType.NameOrTagList, XmlSchemaUse.Required));
+			BadgeType.Attributes.Add(CreateSchemaAttribute("Requires", ScriptSchemaStandardType.NameOrTagList, XmlSchemaUse.Optional));
+			BadgeType.Attributes.Add(CreateSchemaAttribute("Targets", ScriptSchemaStandardType.NameOrTagList, XmlSchemaUse.Optional));
 			BadgeType.Attributes.Add(CreateSchemaAttribute("Project", ScriptSchemaStandardType.BalancedString, XmlSchemaUse.Required));
 			BadgeType.Attributes.Add(CreateSchemaAttribute("Change", ScriptSchemaStandardType.BalancedString, XmlSchemaUse.Optional));
 			BadgeType.Attributes.Add(CreateSchemaAttribute("If", ScriptSchemaStandardType.BalancedString, XmlSchemaUse.Optional));
diff --git a/Engine/Source/Programs/AutomationTool/BuildGraph/Script.cs b/Engine/Source/Programs/AutomationTool/BuildGraph/Script.cs
--- a/Engine/Source/Programs/AutomationTool/BuildGraph/Script.cs
+++ b/Engine/Source/Programs/AutomationTool/BuildGraph/Script.cs
@@ -964,6 +964,7 @@ namespace AutomationTool
 			if (EvaluateCondition(Element) && TryReadObjectName(Element, out Name))
 			{
 				string[] RequiredNames = ReadListAttribute(Element, "Requires");
+				string[] TargetNames = ReadListAttribute(Element, "Targets");
 				string Project = ReadAttribute(Element, "Project");
 				int Change = ReadIntegerAttribute(Element, "Change", 0);
 
@@ -972,6 +973,11 @@ namespace AutomationTool
 				{
 					NewBadge.Nodes.Add(ReferencedNode);
 				}
+				foreach (Node ReferencedNode in ResolveReferences(Element, TargetNames))
+				{
+					NewBadge.Nodes.Add(ReferencedNode);
+					NewBadge.Nodes.UnionWith(ReferencedNode.OrderDependencies);
+				}
 				Graph.Badges.Add(NewBadge);
 			}
 		}
diff --git a/Engine/Source/Runtime/MediaAssets/Private/Misc/MediaTextureResource.cpp b/Engine/Source/Runtime/MediaAssets/Private/Misc/MediaTextureResource.cpp
--- a/Engine/Source/Runtime/MediaAssets/Private/Misc/MediaTextureResource.cpp
+++ b/Engine/Source/Runtime/MediaAssets/Private/Misc/MediaTextureResource.cpp
@@ -633,7 +633,8 @@ void FMediaTextureResource::ConvertSample(const TSharedPtr<IMediaTextureSample,
 		}
 
 		// draw full size quad into render target
-		CommandList.SetStreamSource(0, CreateTempMediaVertexBuffer(), 0);
+		FVertexBufferRHIRef VertexBuffer = CreateTempMediaVertexBuffer();
+		CommandList.SetStreamSource(0, VertexBuffer, 0);
 		// set viewport to RT size
 		CommandList.SetViewport(0, 0, 0.0f, OutputDim.X, OutputDim.Y, 1.0f);
 
