UAT: Fix -Timestamps parameter not being parsed correctly.

#rnx

diff --git a/Engine/Source/Programs/AutomationTool/Program.cs b/Engine/Source/Programs/AutomationTool/Program.cs
--- a/Engine/Source/Programs/AutomationTool/Program.cs
+++ b/Engine/Source/Programs/AutomationTool/Program.cs
@@ -49,7 +49,7 @@ namespace AutomationTool
 			Trace.Listeners.Add(StartupListener);
 
 			// Configure log timestamps
-			Log.IncludeTimestamps = CommandUtils.ParseParam(Arguments, "Timestamps");
+			Log.IncludeTimestamps = CommandUtils.ParseParam(Arguments, "-Timestamps");
 
 			// Enter the main program section
             ExitCode ReturnCode = ExitCode.Success;
