Fix debug info being staged even if it's disabled in packaging settings.

#rnx

diff --git a/Engine/Source/Programs/AutomationTool/Scripts/CopyBuildToStagingDirectory.Automation.cs b/Engine/Source/Programs/AutomationTool/Scripts/CopyBuildToStagingDirectory.Automation.cs
--- a/Engine/Source/Programs/AutomationTool/Scripts/CopyBuildToStagingDirectory.Automation.cs
+++ b/Engine/Source/Programs/AutomationTool/Scripts/CopyBuildToStagingDirectory.Automation.cs
@@ -1203,7 +1203,7 @@ public partial class Project : CommandUtils
 	{
 		CopyManifestFilesToStageDir(SC.FilesToStage.NonUFSFiles, SC.StageDirectory, SC.DebugStageDirectory, "NonUFSFiles", SC.StageTargetPlatform.GetFilesForCRCCheck(), SC.StageTargetPlatform.PlatformType.ToString());
 
-		if (SC.DebugStageDirectory != null)
+		if (!Params.NoDebugInfo)
 		{
 			CopyManifestFilesToStageDir(SC.FilesToStage.NonUFSDebugFiles, SC.DebugStageDirectory, SC.DebugStageDirectory, "DebugFiles", SC.StageTargetPlatform.GetFilesForCRCCheck(), SC.StageTargetPlatform.PlatformType.ToString());
 		}
