UAT: Remove the global -2015 and -IgnoreJunk arguments. The compiler can be selected on a target-by-target basis now, and junk deletion should be fast and reliable.

#rnx

diff --git a/Engine/Source/Programs/AutomationTool/AutomationUtils/Automation.cs b/Engine/Source/Programs/AutomationTool/AutomationUtils/Automation.cs
--- a/Engine/Source/Programs/AutomationTool/AutomationUtils/Automation.cs
+++ b/Engine/Source/Programs/AutomationTool/AutomationUtils/Automation.cs
@@ -102,7 +102,6 @@ namespace AutomationTool
 		public static CommandLineArg NoSubmit = new CommandLineArg("-NoSubmit");
 		public static CommandLineArg NoP4 = new CommandLineArg("-NoP4");
 		public static CommandLineArg P4 = new CommandLineArg("-P4");
-        public static CommandLineArg Preprocess = new CommandLineArg("-Preprocess");
         public static CommandLineArg Compile = new CommandLineArg("-Compile");
         /// <summary>
         /// This command is LEGACY because we used to run UAT.exe to compile scripts by default.
@@ -114,12 +113,10 @@ namespace AutomationTool
         public static CommandLineArg NoCompileEditor = new CommandLineArg("-NoCompileEditor");
 		public static CommandLineArg Help = new CommandLineArg("-Help");
 		public static CommandLineArg List = new CommandLineArg("-List");
-		public static CommandLineArg VS2015 = new CommandLineArg("-2015");
 		public static CommandLineArg NoKill = new CommandLineArg("-NoKill");
 		public static CommandLineArg UTF8Output = new CommandLineArg("-UTF8Output");
         public static CommandLineArg AllowStdOutLogVerbosity = new CommandLineArg("-AllowStdOutLogVerbosity");
         public static CommandLineArg NoAutoSDK = new CommandLineArg("-NoAutoSDK");
-		public static CommandLineArg IgnoreJunk = new CommandLineArg("-ignorejunk");
         /// <summary>
         /// Allows you to use local storage for your root build storage dir (default of P:\Builds (on PC) is changed to Engine\Saved\LocalBuilds). Used for local testing.
         /// </summary>
diff --git a/Engine/Source/Programs/AutomationTool/AutomationUtils/UBTUtils.cs b/Engine/Source/Programs/AutomationTool/AutomationUtils/UBTUtils.cs
--- a/Engine/Source/Programs/AutomationTool/AutomationUtils/UBTUtils.cs
+++ b/Engine/Source/Programs/AutomationTool/AutomationUtils/UBTUtils.cs
@@ -11,11 +11,6 @@ namespace AutomationTool
 
 	public partial class CommandUtils
 	{
-		/// <summary>
-		/// True if UBT has deleted junk files
-		/// </summary>
-		private static bool bJunkDeleted = false;
-
 		/// <summary>
 		/// Runs UBT with the specified commandline. Automatically creates a logfile. When 
 		/// no LogName is specified, the executable name is used as logfile base name.
@@ -30,21 +25,7 @@ namespace AutomationTool
 				throw new AutomationException("Unable to find UBT executable: " + UBTExecutable);
 			}
 
-			if (GlobalCommandLine.VS2015)
-			{
-				CommandLine += " -2015";
-			}
 			CommandLine += " -NoHotReload";
-			if (bJunkDeleted || GlobalCommandLine.IgnoreJunk)
-			{
-				// UBT has already deleted junk files, make sure it doesn't do it again
-				CommandLine += " -ignorejunk";
-			}
-			else
-			{
-				// UBT will delete junk on first run
-				bJunkDeleted = true;
-			}
 
 			string BaseLogName = String.Format("UBT-{0}", String.Join("-", SharedUtils.ParseCommandLine(CommandLine).Where(x => !x.Contains('/') && !x.Contains('\\') && !x.StartsWith("-"))));
 			string LogName;
