UAT: Remove code that forces the -NoHotReload argument onto every UBT invocation; this only makes sense when building, and causes a warning when running in other modes.

diff --git a/Engine/Source/Programs/AutomationTool/AutomationUtils/UBTUtils.cs b/Engine/Source/Programs/AutomationTool/AutomationUtils/UBTUtils.cs
--- a/Engine/Source/Programs/AutomationTool/AutomationUtils/UBTUtils.cs
+++ b/Engine/Source/Programs/AutomationTool/AutomationUtils/UBTUtils.cs
@@ -25,8 +25,6 @@ namespace AutomationTool
 				throw new AutomationException("Unable to find UBT executable: " + UBTExecutable);
 			}
 
-			CommandLine += " -NoHotReload";
-
 			string BaseLogName = String.Format("UBT-{0}", String.Join("-", SharedUtils.ParseCommandLine(CommandLine).Where(x => !x.Contains('/') && !x.Contains('\\') && !x.StartsWith("-"))));
 			string LogName;
 			for(int Attempt = 1;;Attempt++)
diff --git a/Engine/Source/Programs/AutomationTool/AutomationUtils/UE4Build.cs b/Engine/Source/Programs/AutomationTool/AutomationUtils/UE4Build.cs
--- a/Engine/Source/Programs/AutomationTool/AutomationUtils/UE4Build.cs
+++ b/Engine/Source/Programs/AutomationTool/AutomationUtils/UE4Build.cs
@@ -111,7 +111,7 @@ namespace AutomationTool
 
 			ClearExportedXGEXML();
 
-			CommandUtils.RunUBT(CommandUtils.CmdEnv, UBTExecutable: UBTExecutable, Project: UprojectPath, Target: TargetName, Platform: Platform, Config: Config, AdditionalArgs: String.Format("-Manifest={0} -nobuilduht -xgeexport {1}", CommandUtils.MakePathSafeToUseWithCommandLine(ManifestFile.FullName), AddArgs));
+			CommandUtils.RunUBT(CommandUtils.CmdEnv, UBTExecutable: UBTExecutable, Project: UprojectPath, Target: TargetName, Platform: Platform, Config: Config, AdditionalArgs: String.Format("-Manifest={0} -nobuilduht -NoHotReload -xgeexport {1}", CommandUtils.MakePathSafeToUseWithCommandLine(ManifestFile.FullName), AddArgs));
 
 			XGEItem Result = new XGEItem();
 			Result.Platform = Platform;
@@ -183,7 +183,7 @@ namespace AutomationTool
 			PrepareUBT();
 			using(TelemetryStopwatch CleanStopwatch = new TelemetryStopwatch("CleanWithUBT.{0}.{1}.{2}", TargetName, Platform.ToString(), Config))
 			{
-				CommandUtils.RunUBT(CommandUtils.CmdEnv, UBTExecutable: UBTExecutable, Project: UprojectPath, Target: TargetName, Platform: Platform, Config: Config, AdditionalArgs: "-clean" + AddArgs);
+				CommandUtils.RunUBT(CommandUtils.CmdEnv, UBTExecutable: UBTExecutable, Project: UprojectPath, Target: TargetName, Platform: Platform, Config: Config, AdditionalArgs: "-Clean -NoHotReload" + AddArgs);
 			}
         }
 
@@ -210,7 +210,7 @@ namespace AutomationTool
 			FileReference ManifestFile = GetManifestFile(UprojectPath);
 			CommandUtils.DeleteFile(ManifestFile);
 
-			CommandUtils.RunUBT(CommandUtils.CmdEnv, UBTExecutable: UBTExecutable, Project: UprojectPath, Target: TargetName, Platform: TargetPlatform, Config: Config, AdditionalArgs: String.Format("{0} -Manifest={1}", AddArgs, CommandUtils.MakePathSafeToUseWithCommandLine(ManifestFile.FullName)));
+			CommandUtils.RunUBT(CommandUtils.CmdEnv, UBTExecutable: UBTExecutable, Project: UprojectPath, Target: TargetName, Platform: TargetPlatform, Config: Config, AdditionalArgs: String.Format("{0} -Manifest={1} -NoHotReload", AddArgs, CommandUtils.MakePathSafeToUseWithCommandLine(ManifestFile.FullName)));
 
 			BuildManifest Manifest = AddBuildProductsFromManifest(ManifestFile);
 			CommandUtils.DeleteFile(ManifestFile);
