UAT: Remove unused functonality for building C# projects from the UE4Build BuildAgenda.

#rnx

diff --git a/Engine/Source/Programs/AutomationTool/AutomationUtils/UE4Build.cs b/Engine/Source/Programs/AutomationTool/AutomationUtils/UE4Build.cs
--- a/Engine/Source/Programs/AutomationTool/AutomationUtils/UE4Build.cs
+++ b/Engine/Source/Programs/AutomationTool/AutomationUtils/UE4Build.cs
@@ -423,40 +423,12 @@ namespace AutomationTool
 
 		public class BuildAgenda
 		{
-			/// <summary>
-			/// Full .NET solution files that we will compile and include in the build.  Currently we assume the output
-			/// binary file names match the solution file base name, but with various different binary file extensions.
-			/// </summary>
-			public List<string> DotNetSolutions = new List<string>();
-
 			/// <summary>
 			/// .NET .csproj files that will be compiled and included in the build.  Currently we assume the output
 			/// binary file names match the solution file base name, but with various different binary file extensions.
 			/// </summary>
 			public List<string> DotNetProjects = new List<string>();
 
-			/// <summary>
-			/// These are .NET project binary base file names that we want to include and label with the build, but
-			/// we won't be compiling these projects directly ourselves (however, they may be incidentally build or
-			/// published by a different project that we are building.)  We'll look for various .NET binary files with
-			/// this base file name but with different extensions.
-			/// </summary>
-			public List<string> ExtraDotNetFiles = new List<string>();
-
-			/// <summary>
-			/// These are .NET project binary files that are specific to IOS and found in the IOS subdirectory.  We define
-			/// these buildproducts as existing in the DotNET\IOS directory and assume that the output binary file names match
-			/// the solution file base name, but with various binary file extensions
-			/// </summary>
-			public List<string> IOSDotNetProjects = new List<string>();
-
-			/// <summary>
-			/// These are .NET project binary files that are specific to HTML5.  We define these build products as existing in the 
-			/// DotNET directory and assume that the output binary file names match
-			/// the solution file base name, but with various binary file extensions
-			/// </summary>
-			public List<string> HTML5DotNetProjects = new List<string>();
-
 			public string SwarmAgentProject = "";
 			public string SwarmCoordinatorProject = "";
 
@@ -1025,14 +997,7 @@ namespace AutomationTool
 				CommandUtils.BuildSolution(CommandUtils.CmdEnv, SwarmCoordinatorSolution, "Development", "Mixed Platforms");
 				AddSwarmBuildProducts();
 			}
-			
-			foreach (var DotNetSolution in Agenda.DotNetSolutions)
-			{
-				string Solution = Path.Combine(CommandUtils.CmdEnv.LocalRoot, DotNetSolution);
-				CommandUtils.BuildSolution(CommandUtils.CmdEnv, Solution, "Development", "Any CPU");
-				AddBuildProductsForCSharpProj(Solution);
-			}
-			
+				
 			foreach (var DotNetProject in Agenda.DotNetProjects)
 			{
 				string CsProj = Path.Combine(CommandUtils.CmdEnv.LocalRoot, DotNetProject);
@@ -1040,25 +1005,6 @@ namespace AutomationTool
 				AddBuildProductsForCSharpProj(CsProj);
 			}
 
-			foreach (var IOSDotNetProject in Agenda.IOSDotNetProjects)
-			{
-				string IOSCsProj = Path.Combine(CommandUtils.CmdEnv.LocalRoot, IOSDotNetProject);
-				CommandUtils.BuildCSharpProject(CommandUtils.CmdEnv, IOSCsProj);
-				AddIOSBuildProductsForCSharpProj(IOSCsProj);
-			}
-
-			foreach (var HTML5DotNetProject in Agenda.HTML5DotNetProjects)
-			{
-				string CsProj = Path.Combine(CommandUtils.CmdEnv.LocalRoot, HTML5DotNetProject);
-				CommandUtils.BuildCSharpProject(CommandUtils.CmdEnv, CsProj);
-				AddBuildProductsForCSharpProj(CsProj);
-			}
-
-			foreach (var File in Agenda.ExtraDotNetFiles)
-			{
-				AddBuildProductsForCSharpProj(Path.Combine(CommandUtils.CmdEnv.LocalRoot, File));
-			}
-
 			string XGEConsole = null;
 			bool bDisableXGE = ParseParam("NoXGE") || InForceNoXGE;
 			bool bCanUseXGE = !bDisableXGE && PlatformExports.TryGetXgConsoleExecutable(out XGEConsole);
