Fix missing or incorrect copyright notices.

#rnx

diff --git a/Engine/Plugins/Editor/USDImporter/Source/UnrealUSDWrapper/UsdWrapperTest/UsdWrapperTest.cpp b/Engine/Plugins/Editor/USDImporter/Source/UnrealUSDWrapper/UsdWrapperTest/UsdWrapperTest.cpp
--- a/Engine/Plugins/Editor/USDImporter/Source/UnrealUSDWrapper/UsdWrapperTest/UsdWrapperTest.cpp
+++ b/Engine/Plugins/Editor/USDImporter/Source/UnrealUSDWrapper/UsdWrapperTest/UsdWrapperTest.cpp
@@ -1,3 +1,5 @@
+// Copyright 1998-2018 Epic Games, Inc. All Rights Reserved.
+
 // UsdWrapperTest.cpp : Defines the entry point for the console application.
 //
 
diff --git a/Engine/Plugins/Experimental/SkeletalReduction/Source/Private/SkeletalSimplifier.h b/Engine/Plugins/Experimental/SkeletalReduction/Source/Private/SkeletalSimplifier.h
--- a/Engine/Plugins/Experimental/SkeletalReduction/Source/Private/SkeletalSimplifier.h
+++ b/Engine/Plugins/Experimental/SkeletalReduction/Source/Private/SkeletalSimplifier.h
@@ -1,4 +1,4 @@
-// Copyright 1998-2017 Epic Games, Inc. All Rights Reserved.
+// Copyright 1998-2018 Epic Games, Inc. All Rights Reserved.
 
 #pragma once
 
diff --git a/Engine/Plugins/Experimental/SkeletalReduction/Source/Private/SkeletalSimplifierLinearAlgebra.h b/Engine/Plugins/Experimental/SkeletalReduction/Source/Private/SkeletalSimplifierLinearAlgebra.h
--- a/Engine/Plugins/Experimental/SkeletalReduction/Source/Private/SkeletalSimplifierLinearAlgebra.h
+++ b/Engine/Plugins/Experimental/SkeletalReduction/Source/Private/SkeletalSimplifierLinearAlgebra.h
@@ -1,4 +1,4 @@
-// Copyright 1998-2017 Epic Games, Inc. All Rights Reserved.
+// Copyright 1998-2018 Epic Games, Inc. All Rights Reserved.
 
 #pragma once
 
diff --git a/Engine/Plugins/Online/OnlineFramework/Source/Party/Private/Chat/SocialChatRoom.cpp b/Engine/Plugins/Online/OnlineFramework/Source/Party/Private/Chat/SocialChatRoom.cpp
--- a/Engine/Plugins/Online/OnlineFramework/Source/Party/Private/Chat/SocialChatRoom.cpp
+++ b/Engine/Plugins/Online/OnlineFramework/Source/Party/Private/Chat/SocialChatRoom.cpp
@@ -1,4 +1,4 @@
-// Copyright 1998-2016 Epic Games, Inc. All Rights Reserved.
+// Copyright 1998-2018 Epic Games, Inc. All Rights Reserved.
 
 #include "Chat/SocialChatRoom.h"
 #include "SocialToolkit.h"
diff --git a/Engine/Plugins/Online/OnlineFramework/Source/Party/Private/Chat/SocialPartyChatRoom.cpp b/Engine/Plugins/Online/OnlineFramework/Source/Party/Private/Chat/SocialPartyChatRoom.cpp
--- a/Engine/Plugins/Online/OnlineFramework/Source/Party/Private/Chat/SocialPartyChatRoom.cpp
+++ b/Engine/Plugins/Online/OnlineFramework/Source/Party/Private/Chat/SocialPartyChatRoom.cpp
@@ -1,4 +1,4 @@
-// Copyright 1998-2016 Epic Games, Inc. All Rights Reserved.
+// Copyright 1998-2018 Epic Games, Inc. All Rights Reserved.
 
 #include "Chat/SocialPartyChatRoom.h"
 #include "SocialToolkit.h"
diff --git a/Engine/Plugins/Runtime/Oculus/OculusAvatar/Source/OculusAvatar.build.cs b/Engine/Plugins/Runtime/Oculus/OculusAvatar/Source/OculusAvatar.build.cs
--- a/Engine/Plugins/Runtime/Oculus/OculusAvatar/Source/OculusAvatar.build.cs
+++ b/Engine/Plugins/Runtime/Oculus/OculusAvatar/Source/OculusAvatar.build.cs
@@ -1,3 +1,5 @@
+// Copyright 1998-2018 Epic Games, Inc. All Rights Reserved.
+
 using UnrealBuildTool;
 using System;
 using System.IO;
diff --git a/Engine/Plugins/Runtime/Oculus/OculusAvatar/Source/Public/OvrAvatar.cpp b/Engine/Plugins/Runtime/Oculus/OculusAvatar/Source/Public/OvrAvatar.cpp
--- a/Engine/Plugins/Runtime/Oculus/OculusAvatar/Source/Public/OvrAvatar.cpp
+++ b/Engine/Plugins/Runtime/Oculus/OculusAvatar/Source/Public/OvrAvatar.cpp
@@ -1,4 +1,4 @@
-// Fill out your copyright notice in the Description page of Project Settings.
+// Copyright 1998-2018 Epic Games, Inc. All Rights Reserved.
 
 #include "OvrAvatar.h"
 #include "OvrAvatarManager.h"
diff --git a/Engine/Plugins/Runtime/Oculus/OculusAvatar/Source/Public/OvrAvatarManager.cpp b/Engine/Plugins/Runtime/Oculus/OculusAvatar/Source/Public/OvrAvatarManager.cpp
--- a/Engine/Plugins/Runtime/Oculus/OculusAvatar/Source/Public/OvrAvatarManager.cpp
+++ b/Engine/Plugins/Runtime/Oculus/OculusAvatar/Source/Public/OvrAvatarManager.cpp
@@ -1,3 +1,4 @@
+// Copyright 1998-2018 Epic Games, Inc. All Rights Reserved.
 
 #include "OvrAvatarManager.h"
 #include "OvrAvatar.h"
diff --git a/Engine/Plugins/Runtime/Oculus/OculusAvatar/Source/Public/OvrAvatarManager.h b/Engine/Plugins/Runtime/Oculus/OculusAvatar/Source/Public/OvrAvatarManager.h
--- a/Engine/Plugins/Runtime/Oculus/OculusAvatar/Source/Public/OvrAvatarManager.h
+++ b/Engine/Plugins/Runtime/Oculus/OculusAvatar/Source/Public/OvrAvatarManager.h
@@ -1,4 +1,7 @@
+// Copyright 1998-2018 Epic Games, Inc. All Rights Reserved.
+
 #pragma once
+
 #include "Containers/Ticker.h"
 #include "OVR_Avatar.h"
 #include "Containers/Map.h"
diff --git a/Engine/Plugins/Runtime/Oculus/OculusVR/Source/OculusEditor/Private/OculusEditorSettings.cpp b/Engine/Plugins/Runtime/Oculus/OculusVR/Source/OculusEditor/Private/OculusEditorSettings.cpp
--- a/Engine/Plugins/Runtime/Oculus/OculusVR/Source/OculusEditor/Private/OculusEditorSettings.cpp
+++ b/Engine/Plugins/Runtime/Oculus/OculusVR/Source/OculusEditor/Private/OculusEditorSettings.cpp
@@ -1,4 +1,4 @@
-// Fill out your copyright notice in the Description page of Project Settings.
+// Copyright 1998-2018 Epic Games, Inc. All Rights Reserved.
 
 #include "OculusEditorSettings.h"
 
diff --git a/Engine/Source/Editor/Sequencer/Private/MovieSceneSequenceEditor.cpp b/Engine/Source/Editor/Sequencer/Private/MovieSceneSequenceEditor.cpp
--- a/Engine/Source/Editor/Sequencer/Private/MovieSceneSequenceEditor.cpp
+++ b/Engine/Source/Editor/Sequencer/Private/MovieSceneSequenceEditor.cpp
@@ -1,3 +1,4 @@
+// Copyright 1998-2018 Epic Games, Inc. All Rights Reserved.
 
 #include "MovieSceneSequenceEditor.h"
 #include "Modules/ModuleManager.h"
diff --git a/Engine/Source/Programs/AutomationTool/Gauntlet/Framework/FileLockInfo.cs b/Engine/Source/Programs/AutomationTool/Gauntlet/Framework/FileLockInfo.cs
--- a/Engine/Source/Programs/AutomationTool/Gauntlet/Framework/FileLockInfo.cs
+++ b/Engine/Source/Programs/AutomationTool/Gauntlet/Framework/FileLockInfo.cs
@@ -1,4 +1,6 @@
-﻿using System;
+﻿// Copyright 1998-2018 Epic Games, Inc. All Rights Reserved.
+
+using System;
 using System.Collections.Generic;
 using System.Diagnostics;
 using System.Linq;
diff --git a/Engine/Source/Programs/AutomationTool/Gauntlet/Framework/Utils/Gauntlet.MarkdownBuilder.cs b/Engine/Source/Programs/AutomationTool/Gauntlet/Framework/Utils/Gauntlet.MarkdownBuilder.cs
--- a/Engine/Source/Programs/AutomationTool/Gauntlet/Framework/Utils/Gauntlet.MarkdownBuilder.cs
+++ b/Engine/Source/Programs/AutomationTool/Gauntlet/Framework/Utils/Gauntlet.MarkdownBuilder.cs
@@ -1,4 +1,6 @@
-﻿using System;
+﻿// Copyright 1998-2018 Epic Games, Inc. All Rights Reserved.
+
+using System;
 using System.Collections.Generic;
 using System.Linq;
 using System.Text;
diff --git a/Engine/Source/Programs/AutomationTool/Gauntlet/SelfTest/Unreal/Gauntlet.SelfTest.TestUnrealBuildParsing.cs b/Engine/Source/Programs/AutomationTool/Gauntlet/SelfTest/Unreal/Gauntlet.SelfTest.TestUnrealBuildParsing.cs
--- a/Engine/Source/Programs/AutomationTool/Gauntlet/SelfTest/Unreal/Gauntlet.SelfTest.TestUnrealBuildParsing.cs
+++ b/Engine/Source/Programs/AutomationTool/Gauntlet/SelfTest/Unreal/Gauntlet.SelfTest.TestUnrealBuildParsing.cs
@@ -1,4 +1,6 @@
-﻿using Gauntlet;
+﻿// Copyright 1998-2018 Epic Games, Inc. All Rights Reserved.
+
+using Gauntlet;
 using System;
 using System.Collections.Generic;
 using System.Diagnostics;
diff --git a/Engine/Source/Runtime/Engine/Public/Animation/AnimEncodingDecompressionContext.h b/Engine/Source/Runtime/Engine/Public/Animation/AnimEncodingDecompressionContext.h
--- a/Engine/Source/Runtime/Engine/Public/Animation/AnimEncodingDecompressionContext.h
+++ b/Engine/Source/Runtime/Engine/Public/Animation/AnimEncodingDecompressionContext.h
@@ -1,4 +1,4 @@
-// Copyright 2018 Epic Games, Inc. All Rights Reserved.
+// Copyright 1998-2018 Epic Games, Inc. All Rights Reserved.
 
 #pragma once
 
diff --git a/Engine/Source/Runtime/Engine/Public/Animation/AnimEncodingHeapAllocator.h b/Engine/Source/Runtime/Engine/Public/Animation/AnimEncodingHeapAllocator.h
--- a/Engine/Source/Runtime/Engine/Public/Animation/AnimEncodingHeapAllocator.h
+++ b/Engine/Source/Runtime/Engine/Public/Animation/AnimEncodingHeapAllocator.h
@@ -1,4 +1,4 @@
-// Copyright 2018 Epic Games, Inc. All Rights Reserved.
+// Copyright 1998-2018 Epic Games, Inc. All Rights Reserved.
 
 #pragma once
 
