-Add "Depot" global property to BuildGraph
- If "Branch" is //Foo/Bar, "Depot" will be "Foo"

diff --git a/Engine/Source/Programs/AutomationTool/BuildGraph/BuildGraph.cs b/Engine/Source/Programs/AutomationTool/BuildGraph/BuildGraph.cs
--- a/Engine/Source/Programs/AutomationTool/BuildGraph/BuildGraph.cs
+++ b/Engine/Source/Programs/AutomationTool/BuildGraph/BuildGraph.cs
@@ -119,6 +119,7 @@ namespace AutomationTool
 			// Set up the standard properties which build scripts might need
 			Dictionary<string, string> DefaultProperties = new Dictionary<string,string>(StringComparer.InvariantCultureIgnoreCase);
 			DefaultProperties["Branch"] = P4Enabled ? P4Env.Branch : "Unknown";
+			DefaultProperties["Depot"] = P4Enabled ? DefaultProperties["Branch"].Substring(2).Split('/').First() : "Unknown";
 			DefaultProperties["EscapedBranch"] = P4Enabled ? CommandUtils.EscapePath(P4Env.Branch) : "Unknown";
 			DefaultProperties["Change"] = P4Enabled ? P4Env.Changelist.ToString() : "0";
 			DefaultProperties["CodeChange"] = P4Enabled ? P4Env.CodeChangelist.ToString() : "0";
