Fix CIS.

#rnx

diff --git a/Engine/Source/Programs/AutomationTool/AutomationUtils/ProjectUtils.cs b/Engine/Source/Programs/AutomationTool/AutomationUtils/ProjectUtils.cs
--- a/Engine/Source/Programs/AutomationTool/AutomationUtils/ProjectUtils.cs
+++ b/Engine/Source/Programs/AutomationTool/AutomationUtils/ProjectUtils.cs
@@ -318,7 +318,7 @@ namespace AutomationTool
 		/// <returns>Project properties.</returns>
         private static ProjectProperties DetectProjectProperties(FileReference RawProjectPath, List<UnrealTargetPlatform> ClientTargetPlatforms, List<UnrealTargetConfiguration> ClientTargetConfigurations, bool AssetNativizationRequested)
 		{
-			var Properties = new ProjectProperties();
+			ProjectProperties Properties = new ProjectProperties();
 			Properties.RawProjectPath = RawProjectPath;
 
 			// detect if the project is content only, but has non-default build settings
@@ -385,7 +385,7 @@ namespace AutomationTool
 				{
 					if (TargetPlatformType != UnrealTargetPlatform.Unknown)
 					{
-						var Config = ConfigCache.ReadHierarchy(ConfigHierarchyType.Engine, RawProjectPath.Directory, TargetPlatformType);
+						ConfigHierarchy Config = ConfigCache.ReadHierarchy(ConfigHierarchyType.Engine, RawProjectPath.Directory, TargetPlatformType);
 						Properties.EngineConfigs.Add(TargetPlatformType, Config);
 					}
 				}
@@ -394,7 +394,7 @@ namespace AutomationTool
 				{
 					if (TargetPlatformType != UnrealTargetPlatform.Unknown)
 					{
-						var Config = ConfigCache.ReadHierarchy(ConfigHierarchyType.Game, RawProjectPath.Directory, TargetPlatformType);
+						ConfigHierarchy Config = ConfigCache.ReadHierarchy(ConfigHierarchyType.Game, RawProjectPath.Directory, TargetPlatformType);
 						Properties.GameConfigs.Add(TargetPlatformType, Config);
 					}
 				}
@@ -462,8 +462,8 @@ namespace AutomationTool
 			FileReference TargetsDllFilename;
 			string FullProjectPath = null;
 
-			var GameFolders = new List<DirectoryReference>();
-			var RulesFolder = new DirectoryReference(GetRulesAssemblyFolder());
+			List<DirectoryReference> GameFolders = new List<DirectoryReference>();
+			DirectoryReference RulesFolder = new DirectoryReference(GetRulesAssemblyFolder());
 			if (Properties.RawProjectPath != null)
 			{
 				CommandUtils.LogVerbose("Looking for targets for project {0}", Properties.RawProjectPath);
@@ -480,15 +480,15 @@ namespace AutomationTool
 			}
 
 			// the UBT code assumes a certain CWD, but artists don't have this CWD.
-			var SourceDir = CommandUtils.CombinePaths(CommandUtils.CmdEnv.LocalRoot, "Engine", "Source");
+			string SourceDir = CommandUtils.CombinePaths(CommandUtils.CmdEnv.LocalRoot, "Engine", "Source");
 			bool DirPushed = false;
 			if (CommandUtils.DirectoryExists_NoExceptions(SourceDir))
 			{
 				CommandUtils.PushDir(SourceDir);
 				DirPushed = true;
 			}
-			var ExtraSearchDirectories = (ExtraSearchPaths == null)? null : ExtraSearchPaths.Select(x => new DirectoryReference(x)).ToList();
-			var TargetScripts = RulesCompiler.FindAllRulesSourceFiles(RulesCompiler.RulesFileType.Target, GameFolders: GameFolders, ForeignPlugins: null, AdditionalSearchPaths: ExtraSearchDirectories, bIncludeEnterprise: false);
+			List<DirectoryReference> ExtraSearchDirectories = (ExtraSearchPaths == null)? null : ExtraSearchPaths.Select(x => new DirectoryReference(x)).ToList();
+			List<FileReference> TargetScripts = RulesCompiler.FindAllRulesSourceFiles(RulesCompiler.RulesFileType.Target, GameFolders: GameFolders, ForeignPlugins: null, AdditionalSearchPaths: ExtraSearchDirectories, bIncludeEnterprise: false);
 			if (DirPushed)
 			{
 				CommandUtils.PopDir();
@@ -497,8 +497,8 @@ namespace AutomationTool
 			if (!CommandUtils.IsNullOrEmpty(TargetScripts))
 			{
 				// We only care about project target script so filter out any scripts not in the project folder, or take them all if we are just doing engine stuff
-				var ProjectTargetScripts = new List<FileReference>();
-				foreach (var TargetScript in TargetScripts)
+				List<FileReference> ProjectTargetScripts = new List<FileReference>();
+				foreach (FileReference TargetScript in TargetScripts)
 				{
 					if (FullProjectPath == null || TargetScript.IsUnderDirectory(new DirectoryReference(FullProjectPath)))
 					{
@@ -511,7 +511,7 @@ namespace AutomationTool
 			if (!CommandUtils.IsNullOrEmpty(TargetScripts))
 			{
 				CommandUtils.LogVerbose("Found {0} target rule files:", TargetScripts.Count);
-				foreach (var Filename in TargetScripts)
+				foreach (FileReference Filename in TargetScripts)
 				{
 					CommandUtils.LogVerbose("  {0}", Filename);
 				}
@@ -547,7 +547,7 @@ namespace AutomationTool
 		{
 			CommandUtils.LogVerbose("Compiling targets DLL: {0}", TargetsDllFilename);
 
-			var ReferencedAssemblies = new List<string>() 
+			List<string> ReferencedAssemblies = new List<string>() 
 					{ 
 						"System.dll", 
 						"System.Core.dll", 
@@ -555,8 +555,8 @@ namespace AutomationTool
 						typeof(UnrealBuildTool.PlatformExports).Assembly.Location
 					};
 			List<string> PreprocessorDefinitions = RulesAssembly.GetPreprocessorDefinitions();
-			var TargetsDLL = DynamicCompilation.CompileAndLoadAssembly(TargetsDllFilename, TargetScripts, ReferencedAssemblies, PreprocessorDefinitions, DoNotCompile);
-			var AllCompiledTypes = TargetsDLL.GetTypes();
+			Assembly TargetsDLL = DynamicCompilation.CompileAndLoadAssembly(TargetsDllFilename, new HashSet<FileReference>(TargetScripts), ReferencedAssemblies, PreprocessorDefinitions, DoNotCompile);
+			Type[] AllCompiledTypes = TargetsDLL.GetTypes();
 			foreach (Type TargetType in AllCompiledTypes)
 			{
 				// Find TargetRules but skip all "UE4Editor", "UE4Game" targets.
@@ -594,13 +594,13 @@ namespace AutomationTool
 		/// <returns>True if the generated assembly is out of date.</returns>
 		private static bool CheckIfScriptAssemblyIsOutOfDate(FileReference TargetsDllFilename, List<FileReference> TargetScripts)
 		{
-			var bOutOfDate = false;
-			var AssemblyInfo = new FileInfo(TargetsDllFilename.FullName);
+			bool bOutOfDate = false;
+			FileInfo AssemblyInfo = new FileInfo(TargetsDllFilename.FullName);
 			if (AssemblyInfo.Exists)
 			{
-				foreach (var ScriptFilename in TargetScripts)
+				foreach (FileReference ScriptFilename in TargetScripts)
 				{
-					var ScriptInfo = new FileInfo(ScriptFilename.FullName);
+					FileInfo ScriptInfo = new FileInfo(ScriptFilename.FullName);
 					if (ScriptInfo.Exists && ScriptInfo.LastWriteTimeUtc > AssemblyInfo.LastWriteTimeUtc)
 					{
 						bOutOfDate = true;
@@ -623,7 +623,7 @@ namespace AutomationTool
 		private static string GetTargetName(Type TargetRulesType)
 		{
 			const string TargetPostfix = "Target";
-			var Name = TargetRulesType.Name;
+			string Name = TargetRulesType.Name;
 			if (Name.EndsWith(TargetPostfix, StringComparison.InvariantCultureIgnoreCase))
 			{
 				Name = Name.Substring(0, Name.Length - TargetPostfix.Length);
@@ -637,7 +637,7 @@ namespace AutomationTool
 		public static void CleanupFolders()
 		{
 			CommandUtils.LogVerbose("Cleaning up project rules folder");
-			var RulesFolder = GetRulesAssemblyFolder();
+			string RulesFolder = GetRulesAssemblyFolder();
 			if (CommandUtils.DirectoryExists(RulesFolder))
 			{
 				CommandUtils.DeleteDirectoryContents(RulesFolder);
@@ -693,11 +693,11 @@ namespace AutomationTool
                 CommandUtils.LogVerbose("    ShortName:    " + GameName);
 				CommandUtils.LogVerbose("      FilePath          : " + FilePath);
 				CommandUtils.LogVerbose("      bIsCodeBasedProject  : " + (Properties.bIsCodeBasedProject ? "YES" : "NO"));
-                foreach (var HostPlatform in InHostPlatforms)
+                foreach (UnrealTargetPlatform HostPlatform in InHostPlatforms)
                 {
 					CommandUtils.LogVerbose("      For Host : " + HostPlatform.ToString());
 					CommandUtils.LogVerbose("          Targets {0}:", Properties.Targets.Count);
-                    foreach (var ThisTarget in Properties.Targets)
+                    foreach (SingleTargetProperties ThisTarget in Properties.Targets)
                     {
 						CommandUtils.LogVerbose("            TargetName          : " + ThisTarget.TargetName);
 						CommandUtils.LogVerbose("              Type          : " + ThisTarget.Rules.Type);
@@ -706,7 +706,7 @@ namespace AutomationTool
 						CommandUtils.LogVerbose("              bUsesSlate  : " + (ThisTarget.Rules.bUsesSlate ? "YES" : "NO"));
                     }
 					CommandUtils.LogVerbose("      Programs {0}:", Properties.Programs.Count);
-                    foreach (var ThisTarget in Properties.Programs)
+                    foreach (SingleTargetProperties ThisTarget in Properties.Programs)
                     {
 						CommandUtils.LogVerbose("            TargetName                    : " + ThisTarget.TargetName);
                     }
@@ -729,9 +729,9 @@ namespace AutomationTool
             IEnumerable<FileReference> AllProjects = UnrealBuildTool.NativeProjects.EnumerateProjectFiles();
 			using(TelemetryStopwatch SortProjectsStopwatch = new TelemetryStopwatch("SortProjects"))
 			{
-				foreach (var InfoEntry in AllProjects)
+				foreach (FileReference InfoEntry in AllProjects)
 				{
-					var UProject = new BranchUProject(InfoEntry);
+					BranchUProject UProject = new BranchUProject(InfoEntry);
 					if (UProject.Properties.bIsCodeBasedProject)
 					{
 						CodeProjects.Add(UProject);
@@ -758,12 +758,12 @@ namespace AutomationTool
 				BaseEngineProject.Dump(InHostPlatforms);
 
 				CommandUtils.LogVerbose("  {0} Code projects:", CodeProjects.Count);
-				foreach (var Proj in CodeProjects)
+				foreach (BranchUProject Proj in CodeProjects)
 				{
 					Proj.Dump(InHostPlatforms);
 				}
 				CommandUtils.LogVerbose("  {0} Non-Code projects:", NonCodeProjects.Count);
-				foreach (var Proj in NonCodeProjects)
+				foreach (BranchUProject Proj in NonCodeProjects)
 				{
 					Proj.Dump(InHostPlatforms);
 				}
@@ -772,14 +772,14 @@ namespace AutomationTool
 
         public BranchUProject FindGame(string GameName)
         {
-            foreach (var Proj in CodeProjects)
+            foreach (BranchUProject Proj in CodeProjects)
             {
                 if (Proj.GameName.Equals(GameName, StringComparison.InvariantCultureIgnoreCase))
                 {
                     return Proj;
                 }
             }
-            foreach (var Proj in NonCodeProjects)
+            foreach (BranchUProject Proj in NonCodeProjects)
             {
                 if (Proj.GameName.Equals(GameName, StringComparison.InvariantCultureIgnoreCase))
                 {
@@ -799,7 +799,7 @@ namespace AutomationTool
 		}
         public SingleTargetProperties FindProgram(string ProgramName)
         {
-            foreach (var Proj in BaseEngineProject.Properties.Programs)
+            foreach (SingleTargetProperties Proj in BaseEngineProject.Properties.Programs)
             {
                 if (Proj.TargetName.Equals(ProgramName, StringComparison.InvariantCultureIgnoreCase))
                 {
@@ -807,9 +807,9 @@ namespace AutomationTool
                 }
             }
 
-			foreach (var CodeProj in CodeProjects)
+			foreach (BranchUProject CodeProj in CodeProjects)
 			{
-				foreach (var Proj in CodeProj.Properties.Programs)
+				foreach (SingleTargetProperties Proj in CodeProj.Properties.Programs)
 				{
 					if (Proj.TargetName.Equals(ProgramName, StringComparison.InvariantCultureIgnoreCase))
 					{
