Fix exception in BuildPlugin UAT command when building a plugin for UHT.

#rnx

diff --git a/Engine/Source/Programs/AutomationTool/Scripts/BuildPluginCommand.Automation.cs b/Engine/Source/Programs/AutomationTool/Scripts/BuildPluginCommand.Automation.cs
--- a/Engine/Source/Programs/AutomationTool/Scripts/BuildPluginCommand.Automation.cs
+++ b/Engine/Source/Programs/AutomationTool/Scripts/BuildPluginCommand.Automation.cs
@@ -130,7 +130,7 @@ class BuildPlugin : BuildCommand
 			{
 				if (Plugin.SupportedPrograms != null && Plugin.SupportedPrograms.Contains("UnrealHeaderTool"))
 				{
-					CompilePluginWithUBT(null, HostProjectPluginFile, Plugin, "UnrealHeaderTool", TargetType.Program, HostPlatform, UnrealTargetConfiguration.Development, ManifestFileNames, String.Format("{0} -plugin={1}", AdditionalArgs, CommandUtils.MakePathSafeToUseWithCommandLine(HostProjectPluginFile.FullName)));
+					CompilePluginWithUBT(HostProjectFile, HostProjectPluginFile, Plugin, "UnrealHeaderTool", TargetType.Program, HostPlatform, UnrealTargetConfiguration.Development, ManifestFileNames, String.Format("{0} -plugin={1}", AdditionalArgs, CommandUtils.MakePathSafeToUseWithCommandLine(HostProjectPluginFile.FullName)));
 				}
 				CompilePluginWithUBT(HostProjectFile, HostProjectPluginFile, Plugin, "UE4Editor", TargetType.Editor, HostPlatform, UnrealTargetConfiguration.Development, ManifestFileNames, AdditionalArgs);
 			}
