Switch UAT and UBT to use .NET 4.6.2, in order to support long file names (which is enabled via an app.manifest file in each application).

diff --git a/Engine/Source/Programs/AutomationTool/AllDesktop/AllDesktop.Automation.csproj b/Engine/Source/Programs/AutomationTool/AllDesktop/AllDesktop.Automation.csproj
--- a/Engine/Source/Programs/AutomationTool/AllDesktop/AllDesktop.Automation.csproj
+++ b/Engine/Source/Programs/AutomationTool/AllDesktop/AllDesktop.Automation.csproj
@@ -1,5 +1,5 @@
-<?xml version="1.0" encoding="utf-8"?>
-<Project ToolsVersion="4.0" DefaultTargets="Build" xmlns="http://schemas.microsoft.com/developer/msbuild/2003">
+﻿<?xml version="1.0" encoding="utf-8"?>
+<Project ToolsVersion="12.0" DefaultTargets="Build" xmlns="http://schemas.microsoft.com/developer/msbuild/2003">
   <Import Project="$(MSBuildExtensionsPath)\$(MSBuildToolsVersion)\Microsoft.Common.props" Condition="Exists('$(MSBuildExtensionsPath)\$(MSBuildToolsVersion)\Microsoft.Common.props')" />
   <PropertyGroup>
     <Configuration Condition=" '$(Configuration)' == '' ">Development</Configuration>
@@ -9,9 +9,10 @@
     <AppDesignerFolder>Properties</AppDesignerFolder>
     <RootNamespace>AllDesktop.Automation</RootNamespace>
     <AssemblyName>AllDesktop.Automation</AssemblyName>
-    <TargetFrameworkVersion>v4.5</TargetFrameworkVersion>
+    <TargetFrameworkVersion>v4.6.2</TargetFrameworkVersion>
     <FileAlignment>512</FileAlignment>
     <WarningsNotAsErrors>612,618</WarningsNotAsErrors>
+    <TargetFrameworkProfile />
   </PropertyGroup>
   <PropertyGroup Condition=" '$(Configuration)|$(Platform)' == 'Debug|AnyCPU' ">
     <DebugSymbols>true</DebugSymbols>
diff --git a/Engine/Source/Programs/AutomationTool/Android/Android.Automation.csproj b/Engine/Source/Programs/AutomationTool/Android/Android.Automation.csproj
--- a/Engine/Source/Programs/AutomationTool/Android/Android.Automation.csproj
+++ b/Engine/Source/Programs/AutomationTool/Android/Android.Automation.csproj
@@ -1,5 +1,5 @@
-<?xml version="1.0" encoding="utf-8"?>
-<Project ToolsVersion="4.0" DefaultTargets="Build" xmlns="http://schemas.microsoft.com/developer/msbuild/2003">
+﻿<?xml version="1.0" encoding="utf-8"?>
+<Project ToolsVersion="12.0" DefaultTargets="Build" xmlns="http://schemas.microsoft.com/developer/msbuild/2003">
   <Import Project="$(MSBuildExtensionsPath)\$(MSBuildToolsVersion)\Microsoft.Common.props" Condition="Exists('$(MSBuildExtensionsPath)\$(MSBuildToolsVersion)\Microsoft.Common.props')" />
   <PropertyGroup>
     <Configuration Condition=" '$(Configuration)' == '' ">Development</Configuration>
@@ -9,9 +9,10 @@
     <AppDesignerFolder>Properties</AppDesignerFolder>
     <RootNamespace>Android.Automation</RootNamespace>
     <AssemblyName>Android.Automation</AssemblyName>
-    <TargetFrameworkVersion>v4.5</TargetFrameworkVersion>
+    <TargetFrameworkVersion>v4.6.2</TargetFrameworkVersion>
     <FileAlignment>512</FileAlignment>
     <WarningsNotAsErrors>612,618</WarningsNotAsErrors>
+    <TargetFrameworkProfile />
   </PropertyGroup>
   <PropertyGroup Condition=" '$(Configuration)|$(Platform)' == 'Debug|AnyCPU' ">
     <DebugSymbols>true</DebugSymbols>
diff --git a/Engine/Source/Programs/AutomationTool/AutomationTool.csproj b/Engine/Source/Programs/AutomationTool/AutomationTool.csproj
--- a/Engine/Source/Programs/AutomationTool/AutomationTool.csproj
+++ b/Engine/Source/Programs/AutomationTool/AutomationTool.csproj
@@ -1,5 +1,5 @@
-<?xml version="1.0" encoding="utf-8"?>
-<Project ToolsVersion="4.0" DefaultTargets="Build" xmlns="http://schemas.microsoft.com/developer/msbuild/2003">
+﻿<?xml version="1.0" encoding="utf-8"?>
+<Project ToolsVersion="12.0" DefaultTargets="Build" xmlns="http://schemas.microsoft.com/developer/msbuild/2003">
   <PropertyGroup>
     <Configuration Condition=" '$(Configuration)' == '' ">Development</Configuration>
     <Platform Condition=" '$(Platform)' == '' ">AnyCPU</Platform>
@@ -10,7 +10,7 @@
     <AppDesignerFolder>Properties</AppDesignerFolder>
     <RootNamespace>AutomationTool</RootNamespace>
     <AssemblyName>AutomationTool</AssemblyName>
-    <TargetFrameworkVersion>v4.5</TargetFrameworkVersion>
+    <TargetFrameworkVersion>v4.6.2</TargetFrameworkVersion>
     <TargetFrameworkProfile>
     </TargetFrameworkProfile>
     <FileAlignment>512</FileAlignment>
@@ -39,6 +39,9 @@
     </DocumentationFile>
     <Prefer32Bit>false</Prefer32Bit>
   </PropertyGroup>
+  <PropertyGroup>
+    <ApplicationManifest>app.manifest</ApplicationManifest>
+  </PropertyGroup>
   <ItemGroup>
     <Reference Include="Ionic.Zip.Reduced">
       <HintPath>..\..\..\Binaries\DotNET\Ionic.Zip.Reduced.dll</HintPath>
@@ -88,6 +91,9 @@
   </ItemGroup>
   <ItemGroup>
     <None Include="app.config" />
+    <None Include="app.manifest">
+      <SubType>Designer</SubType>
+    </None>
   </ItemGroup>
   <ItemGroup>
     <WCFMetadata Include="Service References\" />
diff --git a/Engine/Source/Programs/AutomationTool/AutomationUtils/AutomationUtils.Automation.csproj b/Engine/Source/Programs/AutomationTool/AutomationUtils/AutomationUtils.Automation.csproj
--- a/Engine/Source/Programs/AutomationTool/AutomationUtils/AutomationUtils.Automation.csproj
+++ b/Engine/Source/Programs/AutomationTool/AutomationUtils/AutomationUtils.Automation.csproj
@@ -1,5 +1,5 @@
-<?xml version="1.0" encoding="utf-8"?>
-<Project ToolsVersion="4.0" DefaultTargets="Build" xmlns="http://schemas.microsoft.com/developer/msbuild/2003">
+﻿<?xml version="1.0" encoding="utf-8"?>
+<Project ToolsVersion="12.0" DefaultTargets="Build" xmlns="http://schemas.microsoft.com/developer/msbuild/2003">
   <Import Project="$(MSBuildExtensionsPath)\$(MSBuildToolsVersion)\Microsoft.Common.props" Condition="Exists('$(MSBuildExtensionsPath)\$(MSBuildToolsVersion)\Microsoft.Common.props')" />
   <PropertyGroup>
     <Configuration Condition=" '$(Configuration)' == '' ">Development</Configuration>
@@ -9,9 +9,10 @@
     <AppDesignerFolder>Properties</AppDesignerFolder>
     <RootNamespace>AutomationUtils.Automation</RootNamespace>
     <AssemblyName>AutomationUtils.Automation</AssemblyName>
-    <TargetFrameworkVersion>v4.5</TargetFrameworkVersion>
+    <TargetFrameworkVersion>v4.6.2</TargetFrameworkVersion>
     <FileAlignment>512</FileAlignment>
     <WarningsNotAsErrors>612,618</WarningsNotAsErrors>
+    <TargetFrameworkProfile />
   </PropertyGroup>
   <PropertyGroup Condition=" '$(Configuration)|$(Platform)' == 'Debug|AnyCPU' ">
     <DebugSymbols>true</DebugSymbols>
diff --git a/Engine/Source/Programs/AutomationTool/BuildGraph/BuildGraph.Automation.csproj b/Engine/Source/Programs/AutomationTool/BuildGraph/BuildGraph.Automation.csproj
--- a/Engine/Source/Programs/AutomationTool/BuildGraph/BuildGraph.Automation.csproj
+++ b/Engine/Source/Programs/AutomationTool/BuildGraph/BuildGraph.Automation.csproj
@@ -1,5 +1,5 @@
-<?xml version="1.0" encoding="utf-8"?>
-<Project ToolsVersion="4.0" DefaultTargets="Build" xmlns="http://schemas.microsoft.com/developer/msbuild/2003">
+﻿<?xml version="1.0" encoding="utf-8"?>
+<Project ToolsVersion="12.0" DefaultTargets="Build" xmlns="http://schemas.microsoft.com/developer/msbuild/2003">
   <Import Project="$(MSBuildExtensionsPath)\$(MSBuildToolsVersion)\Microsoft.Common.props" Condition="Exists('$(MSBuildExtensionsPath)\$(MSBuildToolsVersion)\Microsoft.Common.props')" />
   <PropertyGroup>
     <Configuration Condition=" '$(Configuration)' == '' ">Development</Configuration>
@@ -9,9 +9,10 @@
     <AppDesignerFolder>Properties</AppDesignerFolder>
     <RootNamespace>AutomationTool</RootNamespace>
     <AssemblyName>BuildGraph.Automation</AssemblyName>
-    <TargetFrameworkVersion>v4.5</TargetFrameworkVersion>
+    <TargetFrameworkVersion>v4.6.2</TargetFrameworkVersion>
     <FileAlignment>512</FileAlignment>
     <WarningsNotAsErrors>612,618</WarningsNotAsErrors>
+    <TargetFrameworkProfile />
   </PropertyGroup>
   <PropertyGroup Condition=" '$(Configuration)|$(Platform)' == 'Debug|AnyCPU' ">
     <DebugSymbols>true</DebugSymbols>
diff --git a/Engine/Source/Programs/AutomationTool/Gauntlet/Gauntlet.Automation.csproj b/Engine/Source/Programs/AutomationTool/Gauntlet/Gauntlet.Automation.csproj
--- a/Engine/Source/Programs/AutomationTool/Gauntlet/Gauntlet.Automation.csproj
+++ b/Engine/Source/Programs/AutomationTool/Gauntlet/Gauntlet.Automation.csproj
@@ -1,5 +1,5 @@
 ﻿<?xml version="1.0" encoding="utf-8"?>
-<Project ToolsVersion="4.0" DefaultTargets="Build" xmlns="http://schemas.microsoft.com/developer/msbuild/2003">
+<Project ToolsVersion="12.0" DefaultTargets="Build" xmlns="http://schemas.microsoft.com/developer/msbuild/2003">
   <Import Project="$(MSBuildExtensionsPath)\$(MSBuildToolsVersion)\Microsoft.Common.props" Condition="Exists('$(MSBuildExtensionsPath)\$(MSBuildToolsVersion)\Microsoft.Common.props')" />
   <PropertyGroup>
     <Configuration Condition=" '$(Configuration)' == '' ">Debug</Configuration>
@@ -9,9 +9,10 @@
     <AppDesignerFolder>Properties</AppDesignerFolder>
     <RootNamespace>Gauntlet.Automation</RootNamespace>
     <AssemblyName>Gauntlet.Automation</AssemblyName>
-    <TargetFrameworkVersion>v4.5</TargetFrameworkVersion>
+    <TargetFrameworkVersion>v4.6.2</TargetFrameworkVersion>
     <FileAlignment>512</FileAlignment>
     <WarningsNotAsErrors>612,618</WarningsNotAsErrors>
+    <TargetFrameworkProfile />
   </PropertyGroup>
   <PropertyGroup Condition=" '$(Configuration)|$(Platform)' == 'Debug|AnyCPU' ">
     <DebugSymbols>true</DebugSymbols>
diff --git a/Engine/Source/Programs/AutomationTool/HTML5/HTML5.Automation.csproj b/Engine/Source/Programs/AutomationTool/HTML5/HTML5.Automation.csproj
--- a/Engine/Source/Programs/AutomationTool/HTML5/HTML5.Automation.csproj
+++ b/Engine/Source/Programs/AutomationTool/HTML5/HTML5.Automation.csproj
@@ -1,5 +1,5 @@
-<?xml version="1.0" encoding="utf-8"?>
-<Project ToolsVersion="4.0" DefaultTargets="Build" xmlns="http://schemas.microsoft.com/developer/msbuild/2003">
+﻿<?xml version="1.0" encoding="utf-8"?>
+<Project ToolsVersion="12.0" DefaultTargets="Build" xmlns="http://schemas.microsoft.com/developer/msbuild/2003">
   <Import Project="$(MSBuildExtensionsPath)\$(MSBuildToolsVersion)\Microsoft.Common.props" Condition="Exists('$(MSBuildExtensionsPath)\$(MSBuildToolsVersion)\Microsoft.Common.props')" />
   <PropertyGroup>
     <Configuration Condition=" '$(Configuration)' == '' ">Development</Configuration>
@@ -9,9 +9,10 @@
     <AppDesignerFolder>Properties</AppDesignerFolder>
     <RootNamespace>HTML5.Automation</RootNamespace>
     <AssemblyName>HTML5.Automation</AssemblyName>
-    <TargetFrameworkVersion>v4.5</TargetFrameworkVersion>
+    <TargetFrameworkVersion>v4.6.2</TargetFrameworkVersion>
     <FileAlignment>512</FileAlignment>
     <WarningsNotAsErrors>612,618</WarningsNotAsErrors>
+    <TargetFrameworkProfile />
   </PropertyGroup>
   <PropertyGroup Condition=" '$(Configuration)|$(Platform)' == 'Debug|AnyCPU' ">
     <DebugSymbols>true</DebugSymbols>
diff --git a/Engine/Source/Programs/AutomationTool/IOS/IOS.Automation.csproj b/Engine/Source/Programs/AutomationTool/IOS/IOS.Automation.csproj
--- a/Engine/Source/Programs/AutomationTool/IOS/IOS.Automation.csproj
+++ b/Engine/Source/Programs/AutomationTool/IOS/IOS.Automation.csproj
@@ -1,5 +1,5 @@
-<?xml version="1.0" encoding="utf-8"?>
-<Project ToolsVersion="4.0" DefaultTargets="Build" xmlns="http://schemas.microsoft.com/developer/msbuild/2003">
+﻿<?xml version="1.0" encoding="utf-8"?>
+<Project ToolsVersion="12.0" DefaultTargets="Build" xmlns="http://schemas.microsoft.com/developer/msbuild/2003">
   <Import Project="$(MSBuildExtensionsPath)\$(MSBuildToolsVersion)\Microsoft.Common.props" Condition="Exists('$(MSBuildExtensionsPath)\$(MSBuildToolsVersion)\Microsoft.Common.props')" />
   <PropertyGroup>
     <Configuration Condition=" '$(Configuration)' == '' ">Development</Configuration>
@@ -9,9 +9,10 @@
     <AppDesignerFolder>Properties</AppDesignerFolder>
     <RootNamespace>IOS.Automation</RootNamespace>
     <AssemblyName>IOS.Automation</AssemblyName>
-    <TargetFrameworkVersion>v4.5</TargetFrameworkVersion>
+    <TargetFrameworkVersion>v4.6.2</TargetFrameworkVersion>
     <FileAlignment>512</FileAlignment>
     <WarningsNotAsErrors>612,618</WarningsNotAsErrors>
+    <TargetFrameworkProfile />
   </PropertyGroup>
   <PropertyGroup Condition=" '$(Configuration)|$(Platform)' == 'Debug|AnyCPU' ">
     <DebugSymbols>true</DebugSymbols>
diff --git a/Engine/Source/Programs/AutomationTool/Linux/Linux.Automation.csproj b/Engine/Source/Programs/AutomationTool/Linux/Linux.Automation.csproj
--- a/Engine/Source/Programs/AutomationTool/Linux/Linux.Automation.csproj
+++ b/Engine/Source/Programs/AutomationTool/Linux/Linux.Automation.csproj
@@ -1,5 +1,5 @@
-<?xml version="1.0" encoding="utf-8"?>
-<Project ToolsVersion="4.0" DefaultTargets="Build" xmlns="http://schemas.microsoft.com/developer/msbuild/2003">
+﻿<?xml version="1.0" encoding="utf-8"?>
+<Project ToolsVersion="12.0" DefaultTargets="Build" xmlns="http://schemas.microsoft.com/developer/msbuild/2003">
   <Import Project="$(MSBuildExtensionsPath)\$(MSBuildToolsVersion)\Microsoft.Common.props" Condition="Exists('$(MSBuildExtensionsPath)\$(MSBuildToolsVersion)\Microsoft.Common.props')" />
   <PropertyGroup>
     <Configuration Condition=" '$(Configuration)' == '' ">Developement</Configuration>
@@ -9,9 +9,10 @@
     <AppDesignerFolder>Properties</AppDesignerFolder>
     <RootNamespace>Linux.Automation</RootNamespace>
     <AssemblyName>Linux.Automation</AssemblyName>
-    <TargetFrameworkVersion>v4.5</TargetFrameworkVersion>
+    <TargetFrameworkVersion>v4.6.2</TargetFrameworkVersion>
     <FileAlignment>512</FileAlignment>
     <WarningsNotAsErrors>612,618</WarningsNotAsErrors>
+    <TargetFrameworkProfile />
   </PropertyGroup>
   <PropertyGroup Condition=" '$(Configuration)|$(Platform)' == 'Debug|AnyCPU' ">
     <DebugSymbols>true</DebugSymbols>
diff --git a/Engine/Source/Programs/AutomationTool/Localization/Localization.Automation.csproj b/Engine/Source/Programs/AutomationTool/Localization/Localization.Automation.csproj
--- a/Engine/Source/Programs/AutomationTool/Localization/Localization.Automation.csproj
+++ b/Engine/Source/Programs/AutomationTool/Localization/Localization.Automation.csproj
@@ -1,5 +1,5 @@
-<?xml version="1.0" encoding="utf-8"?>
-<Project ToolsVersion="4.0" DefaultTargets="Build" xmlns="http://schemas.microsoft.com/developer/msbuild/2003">
+﻿<?xml version="1.0" encoding="utf-8"?>
+<Project ToolsVersion="12.0" DefaultTargets="Build" xmlns="http://schemas.microsoft.com/developer/msbuild/2003">
   <Import Project="$(MSBuildExtensionsPath)\$(MSBuildToolsVersion)\Microsoft.Common.props" Condition="Exists('$(MSBuildExtensionsPath)\$(MSBuildToolsVersion)\Microsoft.Common.props')" />
   <PropertyGroup>
     <Configuration Condition=" '$(Configuration)' == '' ">Development</Configuration>
@@ -9,9 +9,10 @@
     <AppDesignerFolder>Properties</AppDesignerFolder>
     <RootNamespace>AutomationTool</RootNamespace>
     <AssemblyName>Localization.Automation</AssemblyName>
-    <TargetFrameworkVersion>v4.5</TargetFrameworkVersion>
+    <TargetFrameworkVersion>v4.6.2</TargetFrameworkVersion>
     <FileAlignment>512</FileAlignment>
     <WarningsNotAsErrors>612,618</WarningsNotAsErrors>
+    <TargetFrameworkProfile />
   </PropertyGroup>
   <PropertyGroup Condition=" '$(Configuration)|$(Platform)' == 'Debug|AnyCPU' ">
     <DebugSymbols>true</DebugSymbols>
diff --git a/Engine/Source/Programs/AutomationTool/Lumin/Lumin.Automation.csproj b/Engine/Source/Programs/AutomationTool/Lumin/Lumin.Automation.csproj
--- a/Engine/Source/Programs/AutomationTool/Lumin/Lumin.Automation.csproj
+++ b/Engine/Source/Programs/AutomationTool/Lumin/Lumin.Automation.csproj
@@ -1,5 +1,5 @@
 ﻿<?xml version="1.0" encoding="utf-8"?>
-<Project ToolsVersion="4.0" DefaultTargets="Build" xmlns="http://schemas.microsoft.com/developer/msbuild/2003">
+<Project ToolsVersion="12.0" DefaultTargets="Build" xmlns="http://schemas.microsoft.com/developer/msbuild/2003">
   <Import Project="$(MSBuildExtensionsPath)\$(MSBuildToolsVersion)\Microsoft.Common.props" Condition="Exists('$(MSBuildExtensionsPath)\$(MSBuildToolsVersion)\Microsoft.Common.props')" />
   <PropertyGroup>
     <Configuration Condition=" '$(Configuration)' == '' ">Debug</Configuration>
@@ -8,10 +8,11 @@
     <AppDesignerFolder>Properties</AppDesignerFolder>
     <RootNamespace>Lumin.Automation</RootNamespace>
     <AssemblyName>Lumin.Automation</AssemblyName>
-    <TargetFrameworkVersion>v4.5</TargetFrameworkVersion>
+    <TargetFrameworkVersion>v4.6.2</TargetFrameworkVersion>
     <FileAlignment>512</FileAlignment>
     <ProjectGuid>{17779110-5CAB-43A0-9DFB-4DF67DF29B42}</ProjectGuid>
     <WarningsNotAsErrors>612,618</WarningsNotAsErrors>
+    <TargetFrameworkProfile />
   </PropertyGroup>
   <PropertyGroup Condition=" '$(Configuration)|$(Platform)' == 'Debug|AnyCPU' ">
     <DebugSymbols>true</DebugSymbols>
diff --git a/Engine/Source/Programs/AutomationTool/Mac/Mac.Automation.csproj b/Engine/Source/Programs/AutomationTool/Mac/Mac.Automation.csproj
--- a/Engine/Source/Programs/AutomationTool/Mac/Mac.Automation.csproj
+++ b/Engine/Source/Programs/AutomationTool/Mac/Mac.Automation.csproj
@@ -1,5 +1,5 @@
-<?xml version="1.0" encoding="utf-8"?>
-<Project ToolsVersion="4.0" DefaultTargets="Build" xmlns="http://schemas.microsoft.com/developer/msbuild/2003">
+﻿<?xml version="1.0" encoding="utf-8"?>
+<Project ToolsVersion="12.0" DefaultTargets="Build" xmlns="http://schemas.microsoft.com/developer/msbuild/2003">
   <Import Project="$(MSBuildExtensionsPath)\$(MSBuildToolsVersion)\Microsoft.Common.props" Condition="Exists('$(MSBuildExtensionsPath)\$(MSBuildToolsVersion)\Microsoft.Common.props')" />
   <PropertyGroup>
     <Configuration Condition=" '$(Configuration)' == '' ">Development</Configuration>
@@ -9,9 +9,10 @@
     <AppDesignerFolder>Properties</AppDesignerFolder>
     <RootNamespace>Mac.Automation</RootNamespace>
     <AssemblyName>Mac.Automation</AssemblyName>
-    <TargetFrameworkVersion>v4.5</TargetFrameworkVersion>
+    <TargetFrameworkVersion>v4.6.2</TargetFrameworkVersion>
     <FileAlignment>512</FileAlignment>
     <WarningsNotAsErrors>612,618</WarningsNotAsErrors>
+    <TargetFrameworkProfile />
   </PropertyGroup>
   <PropertyGroup Condition=" '$(Configuration)|$(Platform)' == 'Debug|AnyCPU' ">
     <DebugSymbols>true</DebugSymbols>
diff --git a/Engine/Source/Programs/AutomationTool/OneSkyLocalization/OneSkyLocalization.Automation.csproj b/Engine/Source/Programs/AutomationTool/OneSkyLocalization/OneSkyLocalization.Automation.csproj
--- a/Engine/Source/Programs/AutomationTool/OneSkyLocalization/OneSkyLocalization.Automation.csproj
+++ b/Engine/Source/Programs/AutomationTool/OneSkyLocalization/OneSkyLocalization.Automation.csproj
@@ -1,5 +1,5 @@
-<?xml version="1.0" encoding="utf-8"?>
-<Project ToolsVersion="4.0" DefaultTargets="Build" xmlns="http://schemas.microsoft.com/developer/msbuild/2003">
+﻿<?xml version="1.0" encoding="utf-8"?>
+<Project ToolsVersion="12.0" DefaultTargets="Build" xmlns="http://schemas.microsoft.com/developer/msbuild/2003">
   <Import Project="$(MSBuildExtensionsPath)\$(MSBuildToolsVersion)\Microsoft.Common.props" Condition="Exists('$(MSBuildExtensionsPath)\$(MSBuildToolsVersion)\Microsoft.Common.props')" />
   <PropertyGroup>
     <Configuration Condition=" '$(Configuration)' == '' ">Development</Configuration>
@@ -9,9 +9,10 @@
     <AppDesignerFolder>Properties</AppDesignerFolder>
     <RootNamespace>AutomationTool</RootNamespace>
     <AssemblyName>OneSkyLocalization.Automation</AssemblyName>
-    <TargetFrameworkVersion>v4.5</TargetFrameworkVersion>
+    <TargetFrameworkVersion>v4.6.2</TargetFrameworkVersion>
     <FileAlignment>512</FileAlignment>
     <WarningsNotAsErrors>612,618</WarningsNotAsErrors>
+    <TargetFrameworkProfile />
   </PropertyGroup>
   <PropertyGroup Condition=" '$(Configuration)|$(Platform)' == 'Debug|AnyCPU' ">
     <DebugSymbols>true</DebugSymbols>
diff --git a/Engine/Source/Programs/AutomationTool/Scripts/AutomationScripts.Automation.csproj b/Engine/Source/Programs/AutomationTool/Scripts/AutomationScripts.Automation.csproj
--- a/Engine/Source/Programs/AutomationTool/Scripts/AutomationScripts.Automation.csproj
+++ b/Engine/Source/Programs/AutomationTool/Scripts/AutomationScripts.Automation.csproj
@@ -1,5 +1,5 @@
-<?xml version="1.0" encoding="utf-8"?>
-<Project ToolsVersion="4.0" DefaultTargets="Build" xmlns="http://schemas.microsoft.com/developer/msbuild/2003">
+﻿<?xml version="1.0" encoding="utf-8"?>
+<Project ToolsVersion="12.0" DefaultTargets="Build" xmlns="http://schemas.microsoft.com/developer/msbuild/2003">
   <Import Project="$(MSBuildExtensionsPath)\$(MSBuildToolsVersion)\Microsoft.Common.props" Condition="Exists('$(MSBuildExtensionsPath)\$(MSBuildToolsVersion)\Microsoft.Common.props')" />
   <PropertyGroup>
     <Configuration Condition=" '$(Configuration)' == '' ">Development</Configuration>
@@ -9,9 +9,10 @@
     <AppDesignerFolder>Properties</AppDesignerFolder>
     <RootNamespace>AutomationTool</RootNamespace>
     <AssemblyName>AutomationScripts.Automation</AssemblyName>
-    <TargetFrameworkVersion>v4.5</TargetFrameworkVersion>
+    <TargetFrameworkVersion>v4.6.2</TargetFrameworkVersion>
     <FileAlignment>512</FileAlignment>
     <WarningsNotAsErrors>612,618</WarningsNotAsErrors>
+    <TargetFrameworkProfile />
   </PropertyGroup>
   <PropertyGroup Condition=" '$(Configuration)|$(Platform)' == 'Debug|AnyCPU' ">
     <DebugSymbols>true</DebugSymbols>
diff --git a/Engine/Source/Programs/AutomationTool/Scripts/CheckCsprojDotNetVersion.cs b/Engine/Source/Programs/AutomationTool/Scripts/CheckCsprojDotNetVersion.cs
--- a/Engine/Source/Programs/AutomationTool/Scripts/CheckCsprojDotNetVersion.cs
+++ b/Engine/Source/Programs/AutomationTool/Scripts/CheckCsprojDotNetVersion.cs
@@ -7,17 +7,21 @@ using AutomationTool;
 using UnrealBuildTool;
 using System.Text.RegularExpressions;
 using Tools.DotNETCommon;
+using System.Linq;
 
 class CheckCsprojDotNetVersion : BuildCommand
 {
 	public override void ExecuteBuild()
 	{
 		// just calling this DesiredTargetVersion because TargetVersion and TargetedVersion get super confusing.
-		string DesiredTargetVersion = this.ParseParamValue("TargetVersion", null);
-		if (string.IsNullOrEmpty(DesiredTargetVersion))
+		string DesiredTargetVersionParam = this.ParseParamValue("TargetVersion", null);
+		if (string.IsNullOrEmpty(DesiredTargetVersionParam))
 		{
 			throw new AutomationException("-TargetVersion was not specified.");
 		}
+
+		string[] DesiredTargetVersions = DesiredTargetVersionParam.Split('+');
+
 		CommandUtils.LogInformation("Scanning for all csproj's...");
 		// Check for all csproj's in the engine dir
 		DirectoryReference EngineDir = CommandUtils.EngineDirectory;
@@ -53,9 +57,9 @@ class CheckCsprojDotNetVersion : BuildCommand
 			{
 				string TargetedVersion = Match.Groups[1].Value;
 				// make sure we match, throw warning otherwise
-				if (!DesiredTargetVersion.Equals(TargetedVersion, StringComparison.InvariantCultureIgnoreCase))
+				if (!DesiredTargetVersions.Any(DesiredTargetVersion => DesiredTargetVersion.Equals(TargetedVersion, StringComparison.InvariantCultureIgnoreCase)))
 				{
-					CommandUtils.LogWarning("Targeted Framework version for project: {0} was not {1}! Targeted Version: {2}", CsProj, DesiredTargetVersion, TargetedVersion);
+					CommandUtils.LogWarning("Targeted Framework version for project: {0} was not {1}! Targeted Version: {2}", CsProj, String.Join("/", DesiredTargetVersions), TargetedVersion);
 				}
 			}
 			// if we don't have a TargetFrameworkVersion, check for the existence of TargetFrameworkProfile.
@@ -108,9 +112,9 @@ class CheckCsprojDotNetVersion : BuildCommand
 				}
 
 				// make sure the versions match
-				if (!(DesiredTargetVersion.Equals(Version1String, StringComparison.InvariantCultureIgnoreCase)))
+				if (!(DesiredTargetVersions.Any(DesiredTargetVersion => DesiredTargetVersion.Equals(Version1String, StringComparison.InvariantCultureIgnoreCase))))
 				{
-					CommandUtils.LogWarning("Targeted Framework version for project: {0} was not {1}! Targeted Version: {2}", CsProj, DesiredTargetVersion, Version1String);
+					CommandUtils.LogWarning("Targeted Framework version for project: {0} was not {1}! Targeted Version: {2}", CsProj, String.Join("/", DesiredTargetVersions), Version1String);
 				}
 			}
 		}
diff --git a/Engine/Source/Programs/AutomationTool/TVOS/TVOS.Automation.csproj b/Engine/Source/Programs/AutomationTool/TVOS/TVOS.Automation.csproj
--- a/Engine/Source/Programs/AutomationTool/TVOS/TVOS.Automation.csproj
+++ b/Engine/Source/Programs/AutomationTool/TVOS/TVOS.Automation.csproj
@@ -1,5 +1,5 @@
-<?xml version="1.0" encoding="utf-8"?>
-<Project ToolsVersion="4.0" DefaultTargets="Build" xmlns="http://schemas.microsoft.com/developer/msbuild/2003">
+﻿<?xml version="1.0" encoding="utf-8"?>
+<Project ToolsVersion="12.0" DefaultTargets="Build" xmlns="http://schemas.microsoft.com/developer/msbuild/2003">
   <Import Project="$(MSBuildExtensionsPath)\$(MSBuildToolsVersion)\Microsoft.Common.props" Condition="Exists('$(MSBuildExtensionsPath)\$(MSBuildToolsVersion)\Microsoft.Common.props')" />
   <PropertyGroup>
     <Configuration Condition=" '$(Configuration)' == '' ">Development</Configuration>
@@ -9,9 +9,10 @@
     <AppDesignerFolder>Properties</AppDesignerFolder>
     <RootNamespace>TVOS.Automation</RootNamespace>
     <AssemblyName>TVOS.Automation</AssemblyName>
-    <TargetFrameworkVersion>v4.5</TargetFrameworkVersion>
+    <TargetFrameworkVersion>v4.6.2</TargetFrameworkVersion>
     <FileAlignment>512</FileAlignment>
     <WarningsNotAsErrors>612,618</WarningsNotAsErrors>
+    <TargetFrameworkProfile />
   </PropertyGroup>
   <PropertyGroup Condition=" '$(Configuration)|$(Platform)' == 'Debug|AnyCPU' ">
     <DebugSymbols>true</DebugSymbols>
diff --git a/Engine/Source/Programs/AutomationTool/Win/Win.Automation.csproj b/Engine/Source/Programs/AutomationTool/Win/Win.Automation.csproj
--- a/Engine/Source/Programs/AutomationTool/Win/Win.Automation.csproj
+++ b/Engine/Source/Programs/AutomationTool/Win/Win.Automation.csproj
@@ -1,5 +1,5 @@
-<?xml version="1.0" encoding="utf-8"?>
-<Project ToolsVersion="4.0" DefaultTargets="Build" xmlns="http://schemas.microsoft.com/developer/msbuild/2003">
+﻿<?xml version="1.0" encoding="utf-8"?>
+<Project ToolsVersion="12.0" DefaultTargets="Build" xmlns="http://schemas.microsoft.com/developer/msbuild/2003">
   <Import Project="$(MSBuildExtensionsPath)\$(MSBuildToolsVersion)\Microsoft.Common.props" Condition="Exists('$(MSBuildExtensionsPath)\$(MSBuildToolsVersion)\Microsoft.Common.props')" />
   <PropertyGroup>
     <Configuration Condition=" '$(Configuration)' == '' ">Development</Configuration>
@@ -9,9 +9,10 @@
     <AppDesignerFolder>Properties</AppDesignerFolder>
     <RootNamespace>Win.Automation</RootNamespace>
     <AssemblyName>Win.Automation</AssemblyName>
-    <TargetFrameworkVersion>v4.5</TargetFrameworkVersion>
+    <TargetFrameworkVersion>v4.6.2</TargetFrameworkVersion>
     <FileAlignment>512</FileAlignment>
     <WarningsNotAsErrors>612,618</WarningsNotAsErrors>
+    <TargetFrameworkProfile />
   </PropertyGroup>
   <PropertyGroup Condition=" '$(Configuration)|$(Platform)' == 'Debug|AnyCPU' ">
     <DebugSymbols>true</DebugSymbols>
diff --git a/Engine/Source/Programs/AutomationTool/XLocLocalization/XLocLocalization.Automation.csproj b/Engine/Source/Programs/AutomationTool/XLocLocalization/XLocLocalization.Automation.csproj
--- a/Engine/Source/Programs/AutomationTool/XLocLocalization/XLocLocalization.Automation.csproj
+++ b/Engine/Source/Programs/AutomationTool/XLocLocalization/XLocLocalization.Automation.csproj
@@ -1,5 +1,5 @@
-<?xml version="1.0" encoding="utf-8"?>
-<Project ToolsVersion="4.0" DefaultTargets="Build" xmlns="http://schemas.microsoft.com/developer/msbuild/2003">
+﻿<?xml version="1.0" encoding="utf-8"?>
+<Project ToolsVersion="12.0" DefaultTargets="Build" xmlns="http://schemas.microsoft.com/developer/msbuild/2003">
   <Import Project="$(MSBuildExtensionsPath)\$(MSBuildToolsVersion)\Microsoft.Common.props" Condition="Exists('$(MSBuildExtensionsPath)\$(MSBuildToolsVersion)\Microsoft.Common.props')" />
   <PropertyGroup>
     <Configuration Condition=" '$(Configuration)' == '' ">Development</Configuration>
@@ -9,9 +9,10 @@
     <AppDesignerFolder>Properties</AppDesignerFolder>
     <RootNamespace>AutomationTool</RootNamespace>
     <AssemblyName>XLocLocalization.Automation</AssemblyName>
-    <TargetFrameworkVersion>v4.5</TargetFrameworkVersion>
+    <TargetFrameworkVersion>v4.6.2</TargetFrameworkVersion>
     <FileAlignment>512</FileAlignment>
     <WarningsNotAsErrors>612,618</WarningsNotAsErrors>
+    <TargetFrameworkProfile />
   </PropertyGroup>
   <PropertyGroup Condition=" '$(Configuration)|$(Platform)' == 'Debug|AnyCPU' ">
     <DebugSymbols>true</DebugSymbols>
diff --git a/Engine/Source/Programs/AutomationTool/app.config b/Engine/Source/Programs/AutomationTool/app.config
--- a/Engine/Source/Programs/AutomationTool/app.config
+++ b/Engine/Source/Programs/AutomationTool/app.config
@@ -1,13 +1,13 @@
-﻿<?xml version="1.0" encoding="utf-8"?>
+<?xml version="1.0" encoding="utf-8"?>
 <configuration>
 	<startup>
-		<supportedRuntime version="v4.0" sku=".NETFramework,Version=v4.5" />
+		<supportedRuntime version="v4.0" sku=".NETFramework,Version=v4.6.2"/>
 	</startup>
 	<runtime>
 		<assemblyBinding xmlns="urn:schemas-microsoft-com:asm.v1">
 			<dependentAssembly>
-				<assemblyIdentity name="Newtonsoft.Json" publicKeyToken="30ad4fe6b2a6aeed" culture="neutral" />
-				<bindingRedirect oldVersion="0.0.0.0-9.0.0.0" newVersion="9.0.0.0" />
+				<assemblyIdentity name="Newtonsoft.Json" publicKeyToken="30ad4fe6b2a6aeed" culture="neutral"/>
+				<bindingRedirect oldVersion="0.0.0.0-9.0.0.0" newVersion="9.0.0.0"/>
 			</dependentAssembly>
 		</assemblyBinding>
 	</runtime>
diff --git a/Engine/Source/Programs/AutomationTool/app.manifest b/Engine/Source/Programs/AutomationTool/app.manifest
new file mode 100644
--- /dev/null
+++ b/Engine/Source/Programs/AutomationTool/app.manifest
@@ -0,0 +1,30 @@
+<?xml version="1.0" encoding="utf-8"?>
+<assembly manifestVersion="1.0" xmlns="urn:schemas-microsoft-com:asm.v1">
+	<trustInfo xmlns="urn:schemas-microsoft-com:asm.v2">
+		<security>
+			<requestedPrivileges xmlns="urn:schemas-microsoft-com:asm.v3">
+				<requestedExecutionLevel level="asInvoker" uiAccess="false" />
+			</requestedPrivileges>
+		</security>
+	</trustInfo>
+	<compatibility xmlns="urn:schemas-microsoft-com:compatibility.v1">
+		<application>
+			<!-- Windows 7 -->
+			<supportedOS Id="{35138b9a-5d96-4fbd-8e2d-a2440225f93a}" />
+
+			<!-- Windows 8 -->
+			<supportedOS Id="{4a2f28e3-53b9-4441-ba9c-d69d4a4a6e38}" />
+
+			<!-- Windows 8.1 -->
+			<supportedOS Id="{1f676c76-80e1-4239-95bb-83d0f6d0da78}" />
+
+			<!-- Windows 10 -->
+			<supportedOS Id="{8e0f7a12-bfb3-4fe8-b9a5-48fd50a15a9a}" />
+		</application>
+	</compatibility>
+	<application xmlns="urn:schemas-microsoft-com:asm.v3">
+		<windowsSettings>
+			<longPathAware xmlns="http://schemas.microsoft.com/SMI/2016/WindowsSettings">true</longPathAware>
+		</windowsSettings>
+	</application>
+</assembly>
diff --git a/Engine/Source/Programs/AutomationToolLauncher/AutomationToolLauncher.csproj b/Engine/Source/Programs/AutomationToolLauncher/AutomationToolLauncher.csproj
--- a/Engine/Source/Programs/AutomationToolLauncher/AutomationToolLauncher.csproj
+++ b/Engine/Source/Programs/AutomationToolLauncher/AutomationToolLauncher.csproj
@@ -9,9 +9,10 @@
     <AppDesignerFolder>Properties</AppDesignerFolder>
     <RootNamespace>AutomationToolLauncher</RootNamespace>
     <AssemblyName>AutomationToolLauncher</AssemblyName>
-    <TargetFrameworkVersion>v4.5</TargetFrameworkVersion>
+    <TargetFrameworkVersion>v4.6.2</TargetFrameworkVersion>
     <FileAlignment>512</FileAlignment>
     <AutoGenerateBindingRedirects>true</AutoGenerateBindingRedirects>
+    <TargetFrameworkProfile />
   </PropertyGroup>
   <PropertyGroup Condition=" '$(Configuration)|$(Platform)' == 'Debug|AnyCPU' ">
     <PlatformTarget>AnyCPU</PlatformTarget>
@@ -36,6 +37,9 @@
     <Prefer32Bit>false</Prefer32Bit>
     <AllowUnsafeBlocks>false</AllowUnsafeBlocks>
   </PropertyGroup>
+  <PropertyGroup>
+    <ApplicationManifest>app.manifest</ApplicationManifest>
+  </PropertyGroup>
   <ItemGroup>
     <Reference Include="System" />
     <Reference Include="System.Core" />
@@ -54,6 +58,9 @@
   </ItemGroup>
   <ItemGroup>
     <None Include="app.config" />
+    <None Include="app.manifest">
+      <SubType>Designer</SubType>
+    </None>
   </ItemGroup>
   <Import Project="$(MSBuildToolsPath)\Microsoft.CSharp.targets" />
   <!-- To modify your build process, add your task inside one of the targets below and uncomment it. 
diff --git a/Engine/Source/Programs/AutomationToolLauncher/app.config b/Engine/Source/Programs/AutomationToolLauncher/app.config
--- a/Engine/Source/Programs/AutomationToolLauncher/app.config
+++ b/Engine/Source/Programs/AutomationToolLauncher/app.config
@@ -1,14 +1,14 @@
 <?xml version="1.0"?>
 <configuration>
 	<startup>
-		<supportedRuntime version="v4.0" sku=".NETFramework,Version=v4.5" />
+		<supportedRuntime version="v4.0" sku=".NETFramework,Version=v4.6.2"/>
 	</startup>
 	<runtime>
 		<assemblyBinding xmlns="urn:schemas-microsoft-com:asm.v1">
 			<dependentAssembly>
-				<assemblyIdentity name="Newtonsoft.Json" publicKeyToken="30ad4fe6b2a6aeed" culture="neutral" />
-				<bindingRedirect oldVersion="0.0.0.0-8.0.0.0" newVersion="9.0.0.0" />
+				<assemblyIdentity name="Newtonsoft.Json" publicKeyToken="30ad4fe6b2a6aeed" culture="neutral"/>
+				<bindingRedirect oldVersion="0.0.0.0-8.0.0.0" newVersion="9.0.0.0"/>
 			</dependentAssembly>
 		</assemblyBinding>
 	</runtime>
-</configuration>
\ No newline at end of file
+</configuration>
diff --git a/Engine/Source/Programs/AutomationToolLauncher/app.manifest b/Engine/Source/Programs/AutomationToolLauncher/app.manifest
new file mode 100644
--- /dev/null
+++ b/Engine/Source/Programs/AutomationToolLauncher/app.manifest
@@ -0,0 +1,30 @@
+<?xml version="1.0" encoding="utf-8"?>
+<assembly manifestVersion="1.0" xmlns="urn:schemas-microsoft-com:asm.v1">
+	<trustInfo xmlns="urn:schemas-microsoft-com:asm.v2">
+		<security>
+			<requestedPrivileges xmlns="urn:schemas-microsoft-com:asm.v3">
+				<requestedExecutionLevel level="asInvoker" uiAccess="false" />
+			</requestedPrivileges>
+		</security>
+	</trustInfo>
+	<compatibility xmlns="urn:schemas-microsoft-com:compatibility.v1">
+		<application>
+			<!-- Windows 7 -->
+			<supportedOS Id="{35138b9a-5d96-4fbd-8e2d-a2440225f93a}" />
+
+			<!-- Windows 8 -->
+			<supportedOS Id="{4a2f28e3-53b9-4441-ba9c-d69d4a4a6e38}" />
+
+			<!-- Windows 8.1 -->
+			<supportedOS Id="{1f676c76-80e1-4239-95bb-83d0f6d0da78}" />
+
+			<!-- Windows 10 -->
+			<supportedOS Id="{8e0f7a12-bfb3-4fe8-b9a5-48fd50a15a9a}" />
+		</application>
+	</compatibility>
+	<application xmlns="urn:schemas-microsoft-com:asm.v3">
+		<windowsSettings>
+			<longPathAware xmlns="http://schemas.microsoft.com/SMI/2016/WindowsSettings">true</longPathAware>
+		</windowsSettings>
+	</application>
+</assembly>
diff --git a/Engine/Source/Programs/DotNETCommon/DotNETUtilities/DotNETUtilities.csproj b/Engine/Source/Programs/DotNETCommon/DotNETUtilities/DotNETUtilities.csproj
--- a/Engine/Source/Programs/DotNETCommon/DotNETUtilities/DotNETUtilities.csproj
+++ b/Engine/Source/Programs/DotNETCommon/DotNETUtilities/DotNETUtilities.csproj
@@ -1,5 +1,5 @@
 ﻿<?xml version="1.0" encoding="utf-8"?>
-<Project ToolsVersion="4.0" DefaultTargets="Build" xmlns="http://schemas.microsoft.com/developer/msbuild/2003">
+<Project ToolsVersion="12.0" DefaultTargets="Build" xmlns="http://schemas.microsoft.com/developer/msbuild/2003">
   <Import Project="$(MSBuildExtensionsPath)\$(MSBuildToolsVersion)\Microsoft.Common.props" Condition="Exists('$(MSBuildExtensionsPath)\$(MSBuildToolsVersion)\Microsoft.Common.props')" />
   <PropertyGroup>
     <Configuration Condition=" '$(Configuration)' == '' ">Development</Configuration>
@@ -9,8 +9,9 @@
     <AppDesignerFolder>Properties</AppDesignerFolder>
     <RootNamespace>DotNETUtilities</RootNamespace>
     <AssemblyName>DotNETUtilities</AssemblyName>
-    <TargetFrameworkVersion>v4.5</TargetFrameworkVersion>
+    <TargetFrameworkVersion>v4.6.2</TargetFrameworkVersion>
     <FileAlignment>512</FileAlignment>
+    <TargetFrameworkProfile />
   </PropertyGroup>
   <PropertyGroup Condition=" '$(Configuration)|$(Platform)' == 'Debug|AnyCPU' ">
     <DebugSymbols>true</DebugSymbols>
diff --git a/Engine/Source/Programs/UnrealBuildTool/UnrealBuildTool.csproj b/Engine/Source/Programs/UnrealBuildTool/UnrealBuildTool.csproj
--- a/Engine/Source/Programs/UnrealBuildTool/UnrealBuildTool.csproj
+++ b/Engine/Source/Programs/UnrealBuildTool/UnrealBuildTool.csproj
@@ -1,5 +1,5 @@
 ﻿<?xml version="1.0" encoding="utf-8"?>
-<Project DefaultTargets="Build" xmlns="http://schemas.microsoft.com/developer/msbuild/2003" ToolsVersion="4.0">
+<Project DefaultTargets="Build" xmlns="http://schemas.microsoft.com/developer/msbuild/2003" ToolsVersion="12.0">
   <PropertyGroup>
     <Configuration Condition=" '$(Configuration)' == '' ">Development</Configuration>
     <Platform Condition=" '$(Platform)' == '' ">AnyCPU</Platform>
@@ -12,7 +12,7 @@
     <AssemblyName>UnrealBuildTool</AssemblyName>
     <RunPostBuildEvent>OnOutputUpdated</RunPostBuildEvent>
     <IsWebBootstrapper>false</IsWebBootstrapper>
-    <TargetFrameworkVersion>v4.5</TargetFrameworkVersion>
+    <TargetFrameworkVersion>v4.6.2</TargetFrameworkVersion>
     <FileUpgradeFlags>
     </FileUpgradeFlags>
     <OldToolsVersion>3.5</OldToolsVersion>
@@ -59,6 +59,9 @@
     <Prefer32Bit>false</Prefer32Bit>
     <DocumentationFile>..\..\..\Binaries\DotNET\UnrealBuildTool.XML</DocumentationFile>
   </PropertyGroup>
+  <PropertyGroup>
+    <ApplicationManifest>app.manifest</ApplicationManifest>
+  </PropertyGroup>
   <ItemGroup>
     <Reference Include="Ionic.Zip.Reduced">
       <HintPath>..\..\..\Binaries\DotNET\Ionic.Zip.Reduced.dll</HintPath>
@@ -390,6 +393,7 @@
   </ItemGroup>
   <ItemGroup>
     <None Include="app.config" />
+    <None Include="app.manifest" />
   </ItemGroup>
   <Import Project="$(MSBuildBinPath)\Microsoft.CSharp.targets" />
   <!-- To modify your build process, add your task inside one of the targets below and uncomment it. 
diff --git a/Engine/Source/Programs/UnrealBuildTool/app.config b/Engine/Source/Programs/UnrealBuildTool/app.config
--- a/Engine/Source/Programs/UnrealBuildTool/app.config
+++ b/Engine/Source/Programs/UnrealBuildTool/app.config
@@ -1,3 +1,3 @@
 <?xml version="1.0"?>
 <configuration>
-<startup><supportedRuntime version="v4.0" sku=".NETFramework,Version=v4.5"/></startup></configuration>
+<startup><supportedRuntime version="v4.0" sku=".NETFramework,Version=v4.6.2"/></startup></configuration>
diff --git a/Engine/Source/Programs/UnrealBuildTool/app.manifest b/Engine/Source/Programs/UnrealBuildTool/app.manifest
new file mode 100644
--- /dev/null
+++ b/Engine/Source/Programs/UnrealBuildTool/app.manifest
@@ -0,0 +1,30 @@
+<?xml version="1.0" encoding="utf-8"?>
+<assembly manifestVersion="1.0" xmlns="urn:schemas-microsoft-com:asm.v1">
+	<trustInfo xmlns="urn:schemas-microsoft-com:asm.v2">
+		<security>
+			<requestedPrivileges xmlns="urn:schemas-microsoft-com:asm.v3">
+				<requestedExecutionLevel level="asInvoker" uiAccess="false" />
+			</requestedPrivileges>
+		</security>
+	</trustInfo>
+	<compatibility xmlns="urn:schemas-microsoft-com:compatibility.v1">
+		<application>
+			<!-- Windows 7 -->
+			<supportedOS Id="{35138b9a-5d96-4fbd-8e2d-a2440225f93a}" />
+
+			<!-- Windows 8 -->
+			<supportedOS Id="{4a2f28e3-53b9-4441-ba9c-d69d4a4a6e38}" />
+
+			<!-- Windows 8.1 -->
+			<supportedOS Id="{1f676c76-80e1-4239-95bb-83d0f6d0da78}" />
+
+			<!-- Windows 10 -->
+			<supportedOS Id="{8e0f7a12-bfb3-4fe8-b9a5-48fd50a15a9a}" />
+		</application>
+	</compatibility>
+	<application xmlns="urn:schemas-microsoft-com:asm.v3">
+		<windowsSettings>
+			<longPathAware xmlns="http://schemas.microsoft.com/SMI/2016/WindowsSettings">true</longPathAware>
+		</windowsSettings>
+	</application>
+</assembly>
