UAT: Simplify implementation of CommandUtils.ZipFiles overload that takes a FileFilter argument.

#rnx

diff --git a/Engine/Source/Programs/AutomationTool/AutomationUtils/CommandUtils.cs b/Engine/Source/Programs/AutomationTool/AutomationUtils/CommandUtils.cs
--- a/Engine/Source/Programs/AutomationTool/AutomationUtils/CommandUtils.cs
+++ b/Engine/Source/Programs/AutomationTool/AutomationUtils/CommandUtils.cs
@@ -2093,40 +2093,8 @@ namespace AutomationTool
 		/// <param name="BaseDirectory">Base directory to store relative paths in the zip file to</param>
 		public static void ZipFiles(FileReference ZipFileName, DirectoryReference BaseDirectory, FileFilter Filter)
 		{
-			// Ionic.Zip.Zip64Option.Always option produces broken archives on Mono, so we use system zip tool instead
-			if (Utils.IsRunningOnMono)
-			{
-				CommandUtils.CreateDirectory(Path.GetDirectoryName(ZipFileName.FullName));
-				CommandUtils.PushDir(BaseDirectory.FullName);
-				string FilesList = "";
-				foreach (FileReference FilteredFile in Filter.ApplyToDirectory(BaseDirectory, true))
-				{
-					FilesList += " \"" + FilteredFile.MakeRelativeTo(BaseDirectory) + "\"";
-					if (FilesList.Length > 32000)
-					{
-						CommandUtils.RunAndLog(CommandUtils.CmdEnv, "zip", "-g -q \"" + ZipFileName + "\"" + FilesList);
-						FilesList = "";
-					}
-				}
-				if (FilesList.Length > 0)
-				{
-					CommandUtils.RunAndLog(CommandUtils.CmdEnv, "zip", "-g -q \"" + ZipFileName + "\"" + FilesList);
-				}
-				CommandUtils.PopDir();
-			}
-			else
-			{
-				using (Ionic.Zip.ZipFile Zip = new Ionic.Zip.ZipFile(Encoding.UTF8))
-				{
-					Zip.UseZip64WhenSaving = Ionic.Zip.Zip64Option.Always;
-					foreach (FileReference FilteredFile in Filter.ApplyToDirectory(BaseDirectory, true))
-					{
-						Zip.AddFile(FilteredFile.FullName, Path.GetDirectoryName(FilteredFile.MakeRelativeTo(BaseDirectory)));
-					}
-					CommandUtils.CreateDirectory(Path.GetDirectoryName(ZipFileName.FullName));
-					Zip.Save(ZipFileName.FullName);
-				}
-			}
+			List<FileReference> Files = Filter.ApplyToDirectory(BaseDirectory, true);
+			ZipFiles(ZipFileName, BaseDirectory, Files);
 		}
 
 		/// <summary>
@@ -2140,7 +2108,7 @@ namespace AutomationTool
 			// Ionic.Zip.Zip64Option.Always option produces broken archives on Mono, so we use system zip tool instead
 			if (Utils.IsRunningOnMono)
 			{
-				CommandUtils.CreateDirectory(ZipFile.Directory.FullName);
+				CommandUtils.CreateDirectory(ZipFile.Directory);
  				CommandUtils.PushDir(BaseDirectory.FullName);
  				string FilesList = "";
 				foreach(FileReference File in Files)
@@ -2160,14 +2128,16 @@ namespace AutomationTool
 			}
 			else
 			{
-				Ionic.Zip.ZipFile Zip = new Ionic.Zip.ZipFile(Encoding.UTF8);
-				Zip.UseZip64WhenSaving = Ionic.Zip.Zip64Option.Always;
-				foreach(FileReference File in Files)
+				using(Ionic.Zip.ZipFile Zip = new Ionic.Zip.ZipFile(Encoding.UTF8))
 				{
-					Zip.AddFile(File.FullName, Path.GetDirectoryName(File.MakeRelativeTo(BaseDirectory)));
+					Zip.UseZip64WhenSaving = Ionic.Zip.Zip64Option.Always;
+					foreach(FileReference File in Files)
+					{
+						Zip.AddFile(File.FullName, Path.GetDirectoryName(File.MakeRelativeTo(BaseDirectory)));
+					}
+					CommandUtils.CreateDirectory(ZipFile.Directory);
+					Zip.Save(ZipFile.FullName);
 				}
-				CommandUtils.CreateDirectory(ZipFile.Directory.FullName);
-				Zip.Save(ZipFile.FullName);
 			}
 		}
 
