Force C# projects to be compiled against the NET 4.5 framework on Mac/Linux, because we don't have 4.6.2 libraries on the version we currently bundle.

#rnx

diff --git a/Engine/Build/BatchFiles/Mac/GenerateProjectFiles.sh b/Engine/Build/BatchFiles/Mac/GenerateProjectFiles.sh
--- a/Engine/Build/BatchFiles/Mac/GenerateProjectFiles.sh
+++ b/Engine/Build/BatchFiles/Mac/GenerateProjectFiles.sh
@@ -22,7 +22,7 @@ fi
 . "$BASE_PATH/SetupMono.sh" $BASE_PATH
 
 if [ -f "$BASE_PATH/../../../Source/Programs/UnrealBuildTool/UnrealBuildTool.csproj" ]; then
-	xbuild "$BASE_PATH/../../../Source/Programs/UnrealBuildTool/UnrealBuildTool.csproj" /property:Configuration="Development" /verbosity:quiet /nologo |grep -i error
+	xbuild "$BASE_PATH/../../../Source/Programs/UnrealBuildTool/UnrealBuildTool.csproj" /property:Configuration="Development" /property:TargetFrameworkVersion=v4.5 /verbosity:quiet /nologo |grep -i error
 fi
 
 # pass all parameters to UBT
diff --git a/Engine/Build/BatchFiles/RunUAT.sh b/Engine/Build/BatchFiles/RunUAT.sh
--- a/Engine/Build/BatchFiles/RunUAT.sh
+++ b/Engine/Build/BatchFiles/RunUAT.sh
@@ -77,7 +77,7 @@ if [ "$UATCompileArg" = "-compile" ]; then
 		fi
 		echo Compiling AutomationTool with xbuild
 
-		ARGS="/p:Configuration=Development /p:Platform=AnyCPU /verbosity:quiet /nologo"
+		ARGS="/p:Configuration=Development /p:Platform=AnyCPU /p:TargetFrameworkVersion=v4.5 /verbosity:quiet /nologo"
 		ARGS="${ARGS} /p:TargetFrameworkProfile="
 
 		echo "xbuild Source/Programs/AutomationTool/AutomationTool.csproj $ARGS"
diff --git a/Engine/Source/Developer/DesktopPlatform/Private/DesktopPlatformBase.cpp b/Engine/Source/Developer/DesktopPlatform/Private/DesktopPlatformBase.cpp
--- a/Engine/Source/Developer/DesktopPlatform/Private/DesktopPlatformBase.cpp
+++ b/Engine/Source/Developer/DesktopPlatform/Private/DesktopPlatformBase.cpp
@@ -1218,7 +1218,7 @@ bool FDesktopPlatformBase::BuildUnrealBuildTool(const FString& RootDir, FOutputD
 	{
 		FString ScriptPath = FPaths::ConvertRelativePathToFull(RootDir / TEXT("Engine/Build/BatchFiles/Mac/RunXBuild.sh"));
 		CompilerExecutableFilename = TEXT("/bin/sh");
-		CmdLineParams = FString::Printf(TEXT("\"%s\" /property:Configuration=Development %s"), *ScriptPath, *CsProjLocation);
+		CmdLineParams = FString::Printf(TEXT("\"%s\" /property:Configuration=Development /property:TargetFrameworkVersion=v4.5 %s"), *ScriptPath, *CsProjLocation);
 	}
 	else if (PLATFORM_LINUX)
 	{
diff --git a/Engine/Source/Programs/AutomationTool/AutomationUtils/ScriptCompiler.cs b/Engine/Source/Programs/AutomationTool/AutomationUtils/ScriptCompiler.cs
--- a/Engine/Source/Programs/AutomationTool/AutomationUtils/ScriptCompiler.cs
+++ b/Engine/Source/Programs/AutomationTool/AutomationUtils/ScriptCompiler.cs
@@ -243,7 +243,7 @@ namespace AutomationTool
 				throw new AutomationException(String.Format("Unable to build Project {0}. Project file not found.", ProjectFile));
 			}
 
-			var CmdLine = String.Format("\"{0}\" /verbosity:quiet /nologo /target:Build /property:Configuration={1} /property:Platform=AnyCPU /p:TreatWarningsAsErrors=false /p:NoWarn=\"612,618,672\" /p:BuildProjectReferences=true",
+			string CmdLine = String.Format("\"{0}\" /verbosity:quiet /nologo /target:Build /property:Configuration={1} /property:Platform=AnyCPU /property:TargetFrameworkVersion=v4.5 /p:TreatWarningsAsErrors=false /p:NoWarn=\"612,618,672\" /p:BuildProjectReferences=true",
 				ProjectFile, BuildConfig);
 
 			// Compile the project
