Force NET 4.5 when compiling C# projects on Linux, since the bundled version of Mono doesn't support 4.6.2 yet.

diff --git a/Engine/Source/Programs/AutomationTool/AutomationUtils/BuildUtils.cs b/Engine/Source/Programs/AutomationTool/AutomationUtils/BuildUtils.cs
--- a/Engine/Source/Programs/AutomationTool/AutomationUtils/BuildUtils.cs
+++ b/Engine/Source/Programs/AutomationTool/AutomationUtils/BuildUtils.cs
@@ -1,4 +1,4 @@
-﻿// Copyright 1998-2019 Epic Games, Inc. All Rights Reserved.
+// Copyright 1998-2019 Epic Games, Inc. All Rights Reserved.
 using System;
 using System.Collections.Generic;
 using System.Text;
@@ -27,11 +27,15 @@ namespace AutomationTool
 			{
 				throw new AutomationException("Project {0} does not exist!", Project);
 			}
-			var RunArguments = MakePathSafeToUseWithCommandLine(Project);
+			string RunArguments = MakePathSafeToUseWithCommandLine(Project);
 			if (!String.IsNullOrEmpty(Arguments))
 			{
 				RunArguments += " " + Arguments;
 			}
+			if(HostPlatform.Current.HostEditorPlatform == UnrealBuildTool.UnrealTargetPlatform.Linux)
+			{
+				RunArguments += " /property:TargetFrameworkVersion=v4.5"; // Need support for 4.6.2 with bundled Mono before removing this.
+			}
 			RunAndLog(Env, Env.MsBuildExe, RunArguments, LogName);
 		}
 
