PR #3874: Added option to Automation Tool BuildPlugin command for checking include structure

#rnx

diff --git a/Engine/Source/Programs/AutomationTool/Scripts/BuildPluginCommand.Automation.cs b/Engine/Source/Programs/AutomationTool/Scripts/BuildPluginCommand.Automation.cs
--- a/Engine/Source/Programs/AutomationTool/Scripts/BuildPluginCommand.Automation.cs
+++ b/Engine/Source/Programs/AutomationTool/Scripts/BuildPluginCommand.Automation.cs
@@ -8,12 +8,14 @@ using System.Reflection;
 using AutomationTool;
 using UnrealBuildTool;
 using Tools.DotNETCommon;
+using System.Text;
 
 [Help("Builds a plugin, and packages it for distribution")]
 [Help("Plugin", "Specify the path to the descriptor file for the plugin that should be packaged")]
 [Help("NoHostPlatform", "Prevent compiling for the editor platform on the host")]
 [Help("TargetPlatforms", "Specify a list of target platforms to build, separated by '+' characters (eg. -TargetPlatforms=Win32+Win64). Default is all the Rocket target platforms.")]
 [Help("Package", "The path which the build artifacts should be packaged to, ready for distribution.")]
+[Help("StrictIncludes", "Disables precompiled headers and unity build in order to check all source files have self-contained headers.")]
 [Help("Unversioned", "Do not embed the current engine version into the descriptor")]
 class BuildPlugin : BuildCommand
 {
@@ -40,6 +42,9 @@ class BuildPlugin : BuildCommand
 			throw new AutomationException("Missing -Package=... argument");
 		}
 
+		// Option for verifying that all include directive s
+		bool bStrictIncludes = ParseParam("StrictIncludes");
+
 		// Make sure the packaging directory is valid
 		DirectoryReference PackageDir = new DirectoryReference(PackageParam);
 		if (PluginFile.IsUnderDirectory(PackageDir))
@@ -86,10 +91,18 @@ class BuildPlugin : BuildCommand
 		CommandUtils.LogInformation("Reading plugin from {0}...", HostProjectPluginFile);
 		PluginDescriptor Plugin = PluginDescriptor.FromFile(HostProjectPluginFile);
 
+		// Get the arguments for the compile
+		StringBuilder AdditionalArgs = new StringBuilder();
+		if (bStrictIncludes)
+		{
+			CommandUtils.LogInformation("Building with precompiled headers and unity disabled");
+			AdditionalArgs.Append(" -NoPCH -NoSharedPCH -DisableUnity");
+		}
+
 		// Compile the plugin for all the target platforms
 		List<UnrealTargetPlatform> HostPlatforms = ParseParam("NoHostPlatform")? new List<UnrealTargetPlatform>() : new List<UnrealTargetPlatform> { BuildHostPlatform.Current.Platform };
 		List<UnrealTargetPlatform> TargetPlatforms = GetTargetPlatforms(this, BuildHostPlatform.Current.Platform);
-		FileReference[] BuildProducts = CompilePlugin(HostProjectFile, HostProjectPluginFile, Plugin, HostPlatforms, TargetPlatforms, "");
+		FileReference[] BuildProducts = CompilePlugin(HostProjectFile, HostProjectPluginFile, Plugin, HostPlatforms, TargetPlatforms, AdditionalArgs.ToString());
 
 		// Package up the final plugin data
 		PackagePlugin(HostProjectPluginFile, BuildProducts, PackageDir, ParseParam("unversioned"));
@@ -144,8 +157,8 @@ class BuildPlugin : BuildCommand
 			{
 				if(Plugin.SupportsTargetPlatform(TargetPlatform))
 				{
-					CompilePluginWithUBT(HostProjectFile, HostProjectPluginFile, Plugin, "UE4Game", TargetType.Game, TargetPlatform, UnrealTargetConfiguration.Development, ManifestFileNames, null);
-					CompilePluginWithUBT(HostProjectFile, HostProjectPluginFile, Plugin, "UE4Game", TargetType.Game, TargetPlatform, UnrealTargetConfiguration.Shipping, ManifestFileNames, null);
+					CompilePluginWithUBT(HostProjectFile, HostProjectPluginFile, Plugin, "UE4Game", TargetType.Game, TargetPlatform, UnrealTargetConfiguration.Development, ManifestFileNames, AdditionalArgs);
+					CompilePluginWithUBT(HostProjectFile, HostProjectPluginFile, Plugin, "UE4Game", TargetType.Game, TargetPlatform, UnrealTargetConfiguration.Shipping, ManifestFileNames, AdditionalArgs);
 				}
 			}
 		}
