Added UE4Editor-Cmd for Mac

diff --git a/Engine/Source/Developer/LauncherServices/Private/Profiles/LauncherProfile.h b/Engine/Source/Developer/LauncherServices/Private/Profiles/LauncherProfile.h
--- a/Engine/Source/Developer/LauncherServices/Private/Profiles/LauncherProfile.h
+++ b/Engine/Source/Developer/LauncherServices/Private/Profiles/LauncherProfile.h
@@ -1873,6 +1873,20 @@ public:
 					EditorExe.Empty();
 				}
 			}
+#elif PLATFORM_MAC
+			// turn UE4editor into UE4editor-cmd
+			if (!FPaths::GetBaseFilename(EditorExe).EndsWith("-cmd", ESearchCase::IgnoreCase))
+			{
+				FString NewExeName = EditorExe + "-Cmd";
+				if (FPaths::FileExists(NewExeName))
+				{
+					EditorExe = NewExeName;
+				}
+				else
+				{
+					EditorExe.Empty();
+				}
+			}
 #endif
 		}
 		else
diff --git a/Engine/Source/Editor/UnrealEd/Private/UnrealEdMisc.cpp b/Engine/Source/Editor/UnrealEd/Private/UnrealEdMisc.cpp
--- a/Engine/Source/Editor/UnrealEd/Private/UnrealEdMisc.cpp
+++ b/Engine/Source/Editor/UnrealEd/Private/UnrealEdMisc.cpp
@@ -1801,6 +1801,16 @@ FString FUnrealEdMisc::GetExecutableForCommandlets() const
 			ExecutableName = NewExeName;
 		}
 	}
+#elif PLATFORM_MAC
+	// turn UE4editor into UE4editor-cmd
+	if (!FPaths::GetBaseFilename(ExecutableName).EndsWith("-cmd", ESearchCase::IgnoreCase))
+	{
+		FString NewExeName = ExecutableName + "-Cmd";
+		if (FPaths::FileExists(NewExeName))
+		{
+			ExecutableName = NewExeName;
+		}
+	}
 #endif
 	return ExecutableName;
 }
diff --git a/Engine/Source/Programs/AutomationTool/AutomationUtils/MacHostPlatform.cs b/Engine/Source/Programs/AutomationTool/AutomationUtils/MacHostPlatform.cs
--- a/Engine/Source/Programs/AutomationTool/AutomationUtils/MacHostPlatform.cs
+++ b/Engine/Source/Programs/AutomationTool/AutomationUtils/MacHostPlatform.cs
@@ -27,7 +27,7 @@ namespace AutomationTool
 			int CmdExeIndex = UE4Exe.IndexOf("-Cmd.exe");
 			if (CmdExeIndex != -1)
 			{
-				UE4Exe = UE4Exe.Substring(0, CmdExeIndex);
+				UE4Exe = UE4Exe.Substring(0, CmdExeIndex + 4);
 			}
 			else
 			{
@@ -37,17 +37,15 @@ namespace AutomationTool
 					UE4Exe = UE4Exe.Substring(0, CmdExeIndex);
 				}
 			}
-			int AppNameIndex = CommandUtils.CmdEnv.LocalRoot.IndexOf("/" + UE4Exe + ".app/Contents/");
-			string EditorExe;
-			if (AppNameIndex != -1) // this is an app with the engine stored inside the app's bundle
+
+			if (UE4Exe.EndsWith("-Cmd", StringComparison.OrdinalIgnoreCase))
 			{
-				EditorExe = CommandUtils.CmdEnv.LocalRoot.Substring(0, AppNameIndex + UE4Exe.Length + 5) + "/Contents/MacOS/" + UE4Exe;
+				return CommandUtils.CombinePaths(CommandUtils.CmdEnv.LocalRoot, RelativeBinariesFolder, UE4Exe);
 			}
 			else
 			{
-				EditorExe = CommandUtils.CombinePaths(CommandUtils.CmdEnv.LocalRoot, RelativeBinariesFolder, UE4Exe + ".app/Contents/MacOS", UE4Exe);
+				return CommandUtils.CombinePaths(CommandUtils.CmdEnv.LocalRoot, RelativeBinariesFolder, UE4Exe + ".app/Contents/MacOS", UE4Exe);
 			}
-			return EditorExe;
 		}
 
 		public override string LocalBuildsLogFolder
diff --git a/Engine/Source/Programs/UnrealBuildTool/Configuration/UEBuildBinary.cs b/Engine/Source/Programs/UnrealBuildTool/Configuration/UEBuildBinary.cs
--- a/Engine/Source/Programs/UnrealBuildTool/Configuration/UEBuildBinary.cs
+++ b/Engine/Source/Programs/UnrealBuildTool/Configuration/UEBuildBinary.cs
@@ -541,7 +541,13 @@ namespace UnrealBuildTool
 		/// <returns></returns>
 		public static FileReference GetAdditionalConsoleAppPath(FileReference BinaryPath)
 		{
-			return FileReference.Combine(BinaryPath.Directory, BinaryPath.GetFileNameWithoutExtension() + "-Cmd" + BinaryPath.GetExtension());
+			DirectoryReference Directory = BinaryPath.Directory;
+			if (Directory.FullName.EndsWith(".app/Contents/MacOS"))
+			{
+				Directory = Directory.ParentDirectory.ParentDirectory.ParentDirectory;
+			}
+
+			return FileReference.Combine(Directory, BinaryPath.GetFileNameWithoutExtension() + "-Cmd" + BinaryPath.GetExtension());
 		}
 
 		/// <summary>
diff --git a/Engine/Source/Programs/UnrealBuildTool/Configuration/UEBuildTarget.cs b/Engine/Source/Programs/UnrealBuildTool/Configuration/UEBuildTarget.cs
--- a/Engine/Source/Programs/UnrealBuildTool/Configuration/UEBuildTarget.cs
+++ b/Engine/Source/Programs/UnrealBuildTool/Configuration/UEBuildTarget.cs
@@ -2769,7 +2769,7 @@ namespace UnrealBuildTool
 			Binary.AddModule(LaunchModule);
 
 			// Create an additional console app for the editor
-			if (Platform == UnrealTargetPlatform.Win64 && Configuration != UnrealTargetConfiguration.Shipping && TargetType == TargetType.Editor)
+			if ((Platform == UnrealTargetPlatform.Win64 || Platform == UnrealTargetPlatform.Mac) && Configuration != UnrealTargetConfiguration.Shipping && TargetType == TargetType.Editor)
 			{
 				Binary.bBuildAdditionalConsoleApp = true;
 			}
diff --git a/Engine/Source/Runtime/Core/Private/Modules/BuildVersion.cpp b/Engine/Source/Runtime/Core/Private/Modules/BuildVersion.cpp
--- a/Engine/Source/Runtime/Core/Private/Modules/BuildVersion.cpp
+++ b/Engine/Source/Runtime/Core/Private/Modules/BuildVersion.cpp
@@ -26,7 +26,7 @@ FString FBuildVersion::GetDefaultFileName()
 FString FBuildVersion::GetFileNameForCurrentExecutable()
 {
 	FString AppExecutableName = FPlatformProcess::ExecutableName();
-#if PLATFORM_WINDOWS
+#if PLATFORM_WINDOWS || PLATFORM_MAC
 	if(AppExecutableName.EndsWith(TEXT("-Cmd")))
 	{
 		AppExecutableName = AppExecutableName.Left(AppExecutableName.Len() - 4);
diff --git a/Engine/Source/Runtime/Launch/Private/Mac/LaunchMac.cpp b/Engine/Source/Runtime/Launch/Private/Mac/LaunchMac.cpp
--- a/Engine/Source/Runtime/Launch/Private/Mac/LaunchMac.cpp
+++ b/Engine/Source/Runtime/Launch/Private/Mac/LaunchMac.cpp
@@ -339,6 +339,8 @@ static int32 MacOSVersionCompare(const NSOperatingSystemVersion& VersionA, const
 
 @end
 
+extern bool GIsConsoleExecutable;
+
 INT32_MAIN_INT32_ARGC_TCHAR_ARGV()
 {
 	for (int32 Option = 1; Option < ArgC; Option++)
@@ -367,6 +369,8 @@ INT32_MAIN_INT32_ARGC_TCHAR_ARGV()
 	GSavedCommandLine += TEXT(" -installed");
 #endif
 
+	GIsConsoleExecutable = !FCString::Strstr(ArgV[0], TEXT(".app/Contents/MacOS/"));
+
 	// convert $'s to " because Xcode swallows the " and this will allow -execcmds= to be usable from xcode
 	GSavedCommandLine = GSavedCommandLine.Replace(TEXT("$"), TEXT("\""));
 
