Fix incorrect path to receipt on TVOS.

#rnx

diff --git a/Engine/Source/Programs/AutomationTool/TVOS/TVOSPlatform.Automation.cs b/Engine/Source/Programs/AutomationTool/TVOS/TVOSPlatform.Automation.cs
--- a/Engine/Source/Programs/AutomationTool/TVOS/TVOSPlatform.Automation.cs
+++ b/Engine/Source/Programs/AutomationTool/TVOS/TVOSPlatform.Automation.cs
@@ -43,11 +43,11 @@ public class TVOSPlatform : IOSPlatform
 		FileReference TargetReceiptFileName;
 		if (bIsUE4Game)
 		{
-			TargetReceiptFileName = TargetReceipt.GetDefaultPath(InEngineDir, "UE4Game", UnrealTargetPlatform.IOS, Config, "");
+			TargetReceiptFileName = TargetReceipt.GetDefaultPath(InEngineDir, "UE4Game", UnrealTargetPlatform.TVOS, Config, "");
 		}
 		else
 		{
-			TargetReceiptFileName = TargetReceipt.GetDefaultPath(ProjectDirectory, TargetName, UnrealTargetPlatform.IOS, Config, "");
+			TargetReceiptFileName = TargetReceipt.GetDefaultPath(ProjectDirectory, TargetName, UnrealTargetPlatform.TVOS, Config, "");
 		}
 		return TVOSExports.GeneratePList(ProjectFile, Config, ProjectDirectory, bIsUE4Game, GameName, ProjectName, InEngineDir, AppDirectory, TargetReceiptFileName, out bSupportsPortrait, out bSupportsLandscape, out bSkipIcons);
 	}
