BuildGraph: Allow the commandlet task to specify a project name, without specifying a path.

diff --git a/Engine/Source/Programs/AutomationTool/BuildGraph/Tasks/CommandletTask.cs b/Engine/Source/Programs/AutomationTool/BuildGraph/Tasks/CommandletTask.cs
--- a/Engine/Source/Programs/AutomationTool/BuildGraph/Tasks/CommandletTask.cs
+++ b/Engine/Source/Programs/AutomationTool/BuildGraph/Tasks/CommandletTask.cs
@@ -2,6 +2,7 @@
 
 using System;
 using System.Collections.Generic;
+using System.IO;
 using System.Linq;
 using System.Text;
 using System.Threading.Tasks;
@@ -73,7 +74,19 @@ namespace AutomationTool.Tasks
 			FileReference ProjectFile = null;
 			if(!String.IsNullOrEmpty(Parameters.Project))
 			{
-				ProjectFile = ResolveFile(Parameters.Project);
+				if(Parameters.Project.EndsWith(".uproject", StringComparison.OrdinalIgnoreCase))
+				{
+					ProjectFile = ResolveFile(Parameters.Project);
+				}
+				else
+				{
+					ProjectFile = NativeProjects.EnumerateProjectFiles().FirstOrDefault(x => x.GetFileNameWithoutExtension().Equals(Parameters.Project, StringComparison.OrdinalIgnoreCase));
+				}
+
+				if(ProjectFile == null || !FileReference.Exists(ProjectFile))
+				{
+					throw new BuildException("Unable to resolve project '{0}'", Parameters.Project);
+				}
 			}
 
 			// Get the path to the editor, and check it exists
