UAT: Force sync the latest Version.h and Metadata.cs files before updating them with version info.

diff --git a/Engine/Source/Programs/AutomationTool/AutomationUtils/UE4Build.cs b/Engine/Source/Programs/AutomationTool/AutomationUtils/UE4Build.cs
--- a/Engine/Source/Programs/AutomationTool/AutomationUtils/UE4Build.cs
+++ b/Engine/Source/Programs/AutomationTool/AutomationUtils/UE4Build.cs
@@ -322,6 +322,12 @@ namespace AutomationTool
 		{
 			FileReference BuildVersionFile = BuildVersion.GetDefaultFileName();
 
+			// Get the revision to sync files to before 
+			if(CommandUtils.P4Enabled && ChangelistNumber > 0)
+			{
+				CommandUtils.P4.Sync(String.Format("-f \"{0}@{1}\"", BuildVersionFile, ChangelistNumber), false, false);
+			}
+
 			BuildVersion Version;
 			if(!BuildVersion.TryRead(BuildVersionFile, out Version))
 			{
@@ -363,10 +369,16 @@ namespace AutomationTool
 
             {
                 // Use Version.h data to update MetaData.cs so the assemblies match the engine version.
-                FileReference MetaDataFile = FileReference.Combine(CommandUtils.EngineDirectory, "Source", "Programs", "DotNETCommon", "MetaData.cs");
+				FileReference MetaDataFile = FileReference.Combine(CommandUtils.EngineDirectory, "Source", "Programs", "DotNETCommon", "MetaData.cs");
 
 				if (bDoUpdateVersionFiles)
                 {
+					// Get the revision to sync files to before 
+					if(CommandUtils.P4Enabled && ChangelistNumber > 0)
+					{
+						CommandUtils.P4.Sync(String.Format("-f \"{0}@{1}\"", MetaDataFile, ChangelistNumber), false, false);
+					}
+
                     // Get the MAJOR/MINOR/PATCH from the Engine Version file, as it is authoritative. The rest we get from the P4Env.
                     string NewInformationalVersion = String.Format("{0}.{1}.{2}-{3}+{4}", Version.MajorVersion, Version.MinorVersion, Version.PatchVersion, Version.Changelist.ToString(), Version.BranchName);
 
