Android Automation
Don't rebuild OBB unless it's dependencies are out of date.  This keeps it from being redeployed with every run which is very time consuming.

diff --git a/Engine/Source/Programs/AutomationTool/Android/AndroidPlatform.Automation.cs b/Engine/Source/Programs/AutomationTool/Android/AndroidPlatform.Automation.cs
--- a/Engine/Source/Programs/AutomationTool/Android/AndroidPlatform.Automation.cs
+++ b/Engine/Source/Programs/AutomationTool/Android/AndroidPlatform.Automation.cs
@@ -295,64 +295,80 @@ public class AndroidPlatform : Platform
 
 		string LocalObbName = SC.StageDirectory.FullName+".obb";
 
-		// Always delete the target OBB file if it exists
-		if (File.Exists(LocalObbName))
+		FileFilter ObbFileFilter = new FileFilter(FileFilterType.Include);
+		ConfigHierarchy EngineIni = ConfigCache.ReadHierarchy(ConfigHierarchyType.Engine, DirectoryReference.FromFile(Params.RawProjectPath), UnrealTargetPlatform.Android);
+		List<string> ObbFilters;
+		EngineIni.GetArray("/Script/AndroidRuntimeSettings.AndroidRuntimeSettings", "ObbFilters", out ObbFilters);
+		if (ObbFilters != null)
 		{
-			File.Delete(LocalObbName);
+			ObbFileFilter.AddRules(ObbFilters);
 		}
 
-		// Now create the OBB as a ZIP archive.
-		LogInformation("Creating {0} from {1}", LocalObbName, SC.StageDirectory);
-		using (ZipFile ObbFile = new ZipFile(LocalObbName))
-		{
-			ObbFile.CompressionMethod = CompressionMethod.None;
-			ObbFile.CompressionLevel = Ionic.Zlib.CompressionLevel.None;
-			ObbFile.UseZip64WhenSaving = Ionic.Zip.Zip64Option.Never;
-
-			int ObbFileCount = 0;
-			ObbFile.AddProgress +=
-				delegate(object sender, AddProgressEventArgs e)
-				{
-					if (e.EventType == ZipProgressEventType.Adding_AfterAddEntry)
-					{
-						ObbFileCount += 1;
-						LogInformation("[{0}/{1}] Adding {2} to OBB",
-							ObbFileCount, e.EntriesTotal,
-							e.CurrentEntry.FileName);
-					}
-				};
-
-
+		string StageDirectoryPath = Path.Combine(SC.StageDirectory.FullName, SC.ShortProjectName);
+		List<FileReference> FilesToObb = ObbFileFilter.ApplyToDirectory(new DirectoryReference(StageDirectoryPath), true);
 
-			FileFilter ObbFileFilter = new FileFilter(FileFilterType.Include);
-			ConfigHierarchy EngineIni = ConfigCache.ReadHierarchy(ConfigHierarchyType.Engine, DirectoryReference.FromFile(Params.RawProjectPath), UnrealTargetPlatform.Android);
-			List<string> ObbFilters;
-			EngineIni.GetArray("/Script/AndroidRuntimeSettings.AndroidRuntimeSettings", "ObbFilters", out ObbFilters);
-			if (ObbFilters != null)
+		bool OBBNeedsUpdate = false;
+		System.DateTime OBBTimeStamp = File.GetLastWriteTimeUtc(LocalObbName);
+		foreach (FileReference FileToObb in FilesToObb)
+		{
+			System.DateTime FileTimeStamp = File.GetLastWriteTimeUtc(FileToObb.FullName);
+			if(FileTimeStamp > OBBTimeStamp)
 			{
-				ObbFileFilter.AddRules(ObbFilters);
+				OBBNeedsUpdate = true;
+				break;
 			}
+		}
 
-			string StageDirectoryPath = Path.Combine(SC.StageDirectory.FullName, SC.ShortProjectName);
-			List<FileReference> FilesToObb = ObbFileFilter.ApplyToDirectory(new DirectoryReference(StageDirectoryPath), true);
+		if (!OBBNeedsUpdate)
+		{
+			LogInformation("OBB is up to date: " + LocalObbName);
+		}
+		else
+		{
 
-			foreach (FileReference FileToObb in FilesToObb)
+			// Always delete the target OBB file if it exists
+			if (File.Exists(LocalObbName))
 			{
-				string DestinationPath = Path.GetDirectoryName(FileToObb.FullName).Replace(StageDirectoryPath, SC.ShortProjectName);
-				ObbFile.AddFile(FileToObb.FullName, DestinationPath);
+				File.Delete(LocalObbName);
 			}
-			
-
 
-			// ObbFile.AddDirectory(SC.StageDirectory+"/"+SC.ShortProjectName, SC.ShortProjectName);
-			try
+			// Now create the OBB as a ZIP archive.
+			LogInformation("Creating {0} from {1}", LocalObbName, SC.StageDirectory);
+			using (ZipFile ObbFile = new ZipFile(LocalObbName))
 			{
-				ObbFile.Save();
-			}
-			catch (Exception)
-			{
-				LogInformation("Failed to build OBB: " + LocalObbName);
-				throw new AutomationException(ExitCode.Error_AndroidOBBError, "Stage Failed. Could not build OBB {0}. The file may be too big to fit in an OBB (2 GiB limit)", LocalObbName);
+				ObbFile.CompressionMethod = CompressionMethod.None;
+				ObbFile.CompressionLevel = Ionic.Zlib.CompressionLevel.None;
+				ObbFile.UseZip64WhenSaving = Ionic.Zip.Zip64Option.Never;
+
+				int ObbFileCount = 0;
+				ObbFile.AddProgress +=
+					delegate (object sender, AddProgressEventArgs e)
+					{
+						if (e.EventType == ZipProgressEventType.Adding_AfterAddEntry)
+						{
+							ObbFileCount += 1;
+							LogInformation("[{0}/{1}] Adding {2} to OBB",
+								ObbFileCount, e.EntriesTotal,
+								e.CurrentEntry.FileName);
+						}
+					};
+
+				foreach (FileReference FileToObb in FilesToObb)
+				{
+					string DestinationPath = Path.GetDirectoryName(FileToObb.FullName).Replace(StageDirectoryPath, SC.ShortProjectName);
+					ObbFile.AddFile(FileToObb.FullName, DestinationPath);
+				}
+
+				// ObbFile.AddDirectory(SC.StageDirectory+"/"+SC.ShortProjectName, SC.ShortProjectName);
+				try
+				{
+					ObbFile.Save();
+				}
+				catch (Exception)
+				{
+					LogInformation("Failed to build OBB: " + LocalObbName);
+					throw new AutomationException(ExitCode.Error_AndroidOBBError, "Stage Failed. Could not build OBB {0}. The file may be too big to fit in an OBB (2 GiB limit)", LocalObbName);
+				}
 			}
 		}
 
