Gauntlet: Do not parse ADB process output for EC jobs

diff --git a/Engine/Source/Programs/AutomationTool/Gauntlet/Platform/Android/Gauntlet.TargetDeviceAndroid.cs b/Engine/Source/Programs/AutomationTool/Gauntlet/Platform/Android/Gauntlet.TargetDeviceAndroid.cs
--- a/Engine/Source/Programs/AutomationTool/Gauntlet/Platform/Android/Gauntlet.TargetDeviceAndroid.cs
+++ b/Engine/Source/Programs/AutomationTool/Gauntlet/Platform/Android/Gauntlet.TargetDeviceAndroid.cs
@@ -1064,15 +1064,18 @@ namespace Gauntlet
 			{
 				Log.Verbose("Running ADB Command: adb {0}", Args);
 			}
-			
-			IProcessResult Process = AndroidPlatform.RunAdbCommand(null, null, Args, null, RunOptions);
 
-			if (Wait)
+			using (var PauseEC = new ScopedSuspendECErrorParsing())
 			{
-				Process.WaitForExit();
+				IProcessResult Process = AndroidPlatform.RunAdbCommand(null, null, Args, null, RunOptions);
+
+				if (Wait)
+				{
+					Process.WaitForExit();
+				}
 			}
 			
-return Process;
+			return Process;
 		}
 
 		public void AllowDeviceSleepState(bool bAllowSleep)
