Fix Android device compile error

diff --git a/Engine/Source/Programs/AutomationTool/Gauntlet/Platform/Android/Gauntlet.TargetDeviceAndroid.cs b/Engine/Source/Programs/AutomationTool/Gauntlet/Platform/Android/Gauntlet.TargetDeviceAndroid.cs
--- a/Engine/Source/Programs/AutomationTool/Gauntlet/Platform/Android/Gauntlet.TargetDeviceAndroid.cs
+++ b/Engine/Source/Programs/AutomationTool/Gauntlet/Platform/Android/Gauntlet.TargetDeviceAndroid.cs
@@ -1065,9 +1065,11 @@ namespace Gauntlet
 				Log.Verbose("Running ADB Command: adb {0}", Args);
 			}
 
+			IProcessResult Process;
+
 			using (var PauseEC = new ScopedSuspendECErrorParsing())
 			{
-				IProcessResult Process = AndroidPlatform.RunAdbCommand(null, null, Args, null, RunOptions);
+				Process = AndroidPlatform.RunAdbCommand(null, null, Args, null, RunOptions);
 
 				if (Wait)
 				{
