Make Android log suspending optional to reduce spam

diff --git a/Engine/Source/Programs/AutomationTool/Gauntlet/Platform/Android/Gauntlet.TargetDeviceAndroid.cs b/Engine/Source/Programs/AutomationTool/Gauntlet/Platform/Android/Gauntlet.TargetDeviceAndroid.cs
--- a/Engine/Source/Programs/AutomationTool/Gauntlet/Platform/Android/Gauntlet.TargetDeviceAndroid.cs
+++ b/Engine/Source/Programs/AutomationTool/Gauntlet/Platform/Android/Gauntlet.TargetDeviceAndroid.cs
@@ -485,7 +485,7 @@ namespace Gauntlet
 					if (!IsExistingDevice)
 					{
 						// disconnect
-						RunAdbGlobalCommand(string.Format("disconnect {0}", DeviceName));
+						RunAdbGlobalCommand(string.Format("disconnect {0}", DeviceName), true, false, true);
 
 						Log.Info("Disconnected {0}", DeviceName);
 					}					
@@ -1006,14 +1006,14 @@ namespace Gauntlet
 		/// <param name="Wait"></param>
 		/// <param name="Input"></param>
 		/// <returns></returns>
-		public IProcessResult RunAdbDeviceCommand(string Args, bool Wait=true, bool bShouldLogCommand = false)
+		public IProcessResult RunAdbDeviceCommand(string Args, bool Wait=true, bool bShouldLogCommand = false, bool bPauseErrorParsing = false)
 		{
 			if (string.IsNullOrEmpty(DeviceName) == false)
 			{
 				Args = string.Format("-s {0} {1}", DeviceName, Args);
 			}
 
-			return RunAdbGlobalCommand(Args, Wait, bShouldLogCommand);
+			return RunAdbGlobalCommand(Args, Wait, bShouldLogCommand, bPauseErrorParsing);
 		}
 
 		/// <summary>
@@ -1047,7 +1047,7 @@ namespace Gauntlet
 		/// <param name="Args"></param>
 		/// <param name="Wait"></param>
 		/// <returns></returns>
-		public static IProcessResult RunAdbGlobalCommand(string Args, bool Wait = true, bool bShouldLogCommand = false)
+		public static IProcessResult RunAdbGlobalCommand(string Args, bool Wait = true, bool bShouldLogCommand = false, bool bPauseErrorParsing = false)
 		{
 			CommandUtils.ERunOptions RunOptions = CommandUtils.ERunOptions.AppMustExist | CommandUtils.ERunOptions.NoWaitForExit;
 
@@ -1067,7 +1067,7 @@ namespace Gauntlet
 
 			IProcessResult Process;
 
-			using (var PauseEC = new ScopedSuspendECErrorParsing())
+			using (bPauseErrorParsing ? new ScopedSuspendECErrorParsing() : null)
 			{
 				Process = AndroidPlatform.RunAdbCommand(null, null, Args, null, RunOptions);
 
@@ -1083,7 +1083,7 @@ namespace Gauntlet
 		public void AllowDeviceSleepState(bool bAllowSleep)
 		{
 			string CommandLine = "shell svc power stayon " + (bAllowSleep ? "false" : "usb");
-			RunAdbDeviceCommand(CommandLine);
+			RunAdbDeviceCommand(CommandLine, true, false, true);
 		}
 
 		public void KillRunningProcess(string AndroidPackageName)
