Fix Gauntlet automation DLLs being output to the {WorkspaceRoot}\Binaries\DotNET directory rather than Engine\Binaries\DotNET.

Seems to get copied into the correct place because it's referenced by GauntletExtras anyway, just creates some junk in the workspace root directory.

diff --git a/Engine/Source/Programs/AutomationTool/Gauntlet/Gauntlet.Automation.csproj b/Engine/Source/Programs/AutomationTool/Gauntlet/Gauntlet.Automation.csproj
--- a/Engine/Source/Programs/AutomationTool/Gauntlet/Gauntlet.Automation.csproj
+++ b/Engine/Source/Programs/AutomationTool/Gauntlet/Gauntlet.Automation.csproj
@@ -18,7 +18,7 @@
     <DebugSymbols>true</DebugSymbols>
     <DebugType>full</DebugType>
     <Optimize>false</Optimize>
-    <OutputPath>..\..\..\..\..\Binaries\DotNET\AutomationScripts\</OutputPath>
+    <OutputPath>..\..\..\..\Binaries\DotNET\AutomationScripts\</OutputPath>
     <DefineConstants>DEBUG;TRACE</DefineConstants>
     <ErrorReport>prompt</ErrorReport>
     <WarningLevel>4</WarningLevel>
