UAT: Prevent temp-target files from having their timestamps updated even when the contents don't change. Causes unnecessary rebuilds when working with content-only projects.

diff --git a/Engine/Source/Programs/AutomationTool/AutomationUtils/ProjectUtils.cs b/Engine/Source/Programs/AutomationTool/AutomationUtils/ProjectUtils.cs
--- a/Engine/Source/Programs/AutomationTool/AutomationUtils/ProjectUtils.cs
+++ b/Engine/Source/Programs/AutomationTool/AutomationUtils/ProjectUtils.cs
@@ -278,8 +278,8 @@ namespace AutomationTool
 			string ProjectName = RawProjectPath.GetFileNameWithoutExtension();
 
 			// Create a target.cs file
-			FileReference TargetLocation = FileReference.Combine(TempDir, ProjectName + ".Target.cs");
-			using (StreamWriter Writer = new StreamWriter(TargetLocation.FullName))
+			MemoryStream TargetStream = new MemoryStream();
+			using (StreamWriter Writer = new StreamWriter(TargetStream))
 			{
 				Writer.WriteLine("using UnrealBuildTool;");
 				Writer.WriteLine();
@@ -292,10 +292,12 @@ namespace AutomationTool
 				Writer.WriteLine("\t}");
 				Writer.WriteLine("}");
 			}
+			FileReference TargetLocation = FileReference.Combine(TempDir, ProjectName + ".Target.cs");
+			FileReference.WriteAllBytesIfDifferent(TargetLocation, TargetStream.ToArray());
 
 			// Create a build.cs file
-			FileReference ModuleLocation = FileReference.Combine(TempDir, ProjectName + ".Build.cs");
-			using (StreamWriter Writer = new StreamWriter(ModuleLocation.FullName))
+			MemoryStream ModuleStream = new MemoryStream();
+			using (StreamWriter Writer = new StreamWriter(ModuleStream))
 			{
 				Writer.WriteLine("using UnrealBuildTool;");
 				Writer.WriteLine();
@@ -310,16 +312,20 @@ namespace AutomationTool
 				Writer.WriteLine("\t}");
 				Writer.WriteLine("}");
 			}
+			FileReference ModuleLocation = FileReference.Combine(TempDir, ProjectName + ".Build.cs");
+			FileReference.WriteAllBytesIfDifferent(ModuleLocation, ModuleStream.ToArray());
 
 			// Create a main module cpp file
-			FileReference SourceFileLocation = FileReference.Combine(TempDir, ProjectName + ".cpp");
-			using (StreamWriter Writer = new StreamWriter(SourceFileLocation.FullName))
+			MemoryStream SourceFileStream = new MemoryStream();
+			using (StreamWriter Writer = new StreamWriter(SourceFileStream))
 			{
 				Writer.WriteLine("#include \"CoreTypes.h\"");
 				Writer.WriteLine("#include \"Modules/ModuleManager.h\"");
 				Writer.WriteLine();
 				Writer.WriteLine("IMPLEMENT_PRIMARY_GAME_MODULE(FDefaultModuleImpl, {0}, \"{0}\");", ProjectName);
 			}
+			FileReference SourceFileLocation = FileReference.Combine(TempDir, ProjectName + ".cpp");
+			FileReference.WriteAllBytesIfDifferent(SourceFileLocation, SourceFileStream.ToArray());
 		}
 
 		/// <summary>
