Remove other references to RPCUtility

diff --git a/Engine/Source/Programs/AutomationTool/AutomationTool.csproj b/Engine/Source/Programs/AutomationTool/AutomationTool.csproj
--- a/Engine/Source/Programs/AutomationTool/AutomationTool.csproj
+++ b/Engine/Source/Programs/AutomationTool/AutomationTool.csproj
@@ -47,10 +47,6 @@
       <HintPath>..\..\..\Binaries\DotNET\Ionic.Zip.Reduced.dll</HintPath>
       <Private>False</Private>
     </Reference>
-    <Reference Include="RPCUtility">
-      <HintPath>..\..\..\Binaries\DotNET\RPCUtility.exe</HintPath>
-      <Private>False</Private>
-    </Reference>
     <Reference Include="System" />
     <Reference Include="System.Core" />
     <Reference Include="System.Management" />
diff --git a/Engine/Source/Programs/AutomationTool/IOS/IOS.Automation.csproj b/Engine/Source/Programs/AutomationTool/IOS/IOS.Automation.csproj
--- a/Engine/Source/Programs/AutomationTool/IOS/IOS.Automation.csproj
+++ b/Engine/Source/Programs/AutomationTool/IOS/IOS.Automation.csproj
@@ -44,10 +44,6 @@
       <HintPath>..\..\..\..\Binaries\DotNET\Ionic.Zip.Reduced.dll</HintPath>
       <Private>False</Private>
     </Reference>
-    <Reference Include="RPCUtility">
-      <HintPath>..\..\..\..\Binaries\DotNET\RPCUtility.exe</HintPath>
-      <Private>False</Private>
-    </Reference>
     <Reference Include="System" />
     <Reference Include="System.Core" />
     <Reference Include="System.Drawing" />
diff --git a/Engine/Source/Programs/AutomationTool/TVOS/TVOS.Automation.csproj b/Engine/Source/Programs/AutomationTool/TVOS/TVOS.Automation.csproj
--- a/Engine/Source/Programs/AutomationTool/TVOS/TVOS.Automation.csproj
+++ b/Engine/Source/Programs/AutomationTool/TVOS/TVOS.Automation.csproj
@@ -40,10 +40,6 @@
       <HintPath>..\..\..\..\Binaries\DotNET\Ionic.Zip.Reduced.dll</HintPath>
       <Private>False</Private>
     </Reference>
-    <Reference Include="RPCUtility">
-      <HintPath>..\..\..\..\Binaries\DotNET\RPCUtility.exe</HintPath>
-      <Private>False</Private>
-    </Reference>
     <Reference Include="System" />
     <Reference Include="System.Core" />
     <Reference Include="System.Drawing" />
