[AutomationTool] Add support to override the launch command line per platform.

diff --git a/Engine/Source/Programs/AutomationTool/AutomationUtils/Platform.cs b/Engine/Source/Programs/AutomationTool/AutomationUtils/Platform.cs
--- a/Engine/Source/Programs/AutomationTool/AutomationUtils/Platform.cs
+++ b/Engine/Source/Programs/AutomationTool/AutomationUtils/Platform.cs
@@ -427,6 +427,16 @@ namespace AutomationTool
 			get { return false; }
 		}
 
+		/// <summary>
+		/// Gets extra launch commandline arguments for this platform.
+		/// </summary>
+		/// <param name="Params"> ProjectParams </param>
+		/// <returns>Launch platform string.</returns>
+		public virtual string GetLaunchExtraCommandLine(ProjectParams Params)
+		{
+			return "";
+		}
+
 		/// <summary>
 		/// True if this platform can write to the abslog path that's on the host desktop.
 		/// </summary>
diff --git a/Engine/Source/Programs/AutomationTool/Scripts/RunProjectCommand.Automation.cs b/Engine/Source/Programs/AutomationTool/Scripts/RunProjectCommand.Automation.cs
--- a/Engine/Source/Programs/AutomationTool/Scripts/RunProjectCommand.Automation.cs
+++ b/Engine/Source/Programs/AutomationTool/Scripts/RunProjectCommand.Automation.cs
@@ -881,10 +881,7 @@ public partial class Project : CommandUtils
 		{
 			TempCmdLine += "-nullrhi ";
 		}
-		if (Params.Deploy && !Params.CookOnTheFly && (SC.StageTargetPlatform.PlatformType == UnrealTargetPlatform.PS4))
-		{
-			TempCmdLine += "-deployedbuild ";
-		}
+		TempCmdLine += SC.StageTargetPlatform.GetLaunchExtraCommandLine(Params);
 
 		TempCmdLine += "-CrashForUAT ";
 		TempCmdLine += Params.RunCommandline;
