Skip android test installation batch files during gauntlet build discovery

diff --git a/Engine/Source/Programs/AutomationTool/Gauntlet/Platform/Android/Gauntlet.AndroidBuildSource.cs b/Engine/Source/Programs/AutomationTool/Gauntlet/Platform/Android/Gauntlet.AndroidBuildSource.cs
--- a/Engine/Source/Programs/AutomationTool/Gauntlet/Platform/Android/Gauntlet.AndroidBuildSource.cs
+++ b/Engine/Source/Programs/AutomationTool/Gauntlet/Platform/Android/Gauntlet.AndroidBuildSource.cs
@@ -71,8 +71,16 @@ namespace Gauntlet
 
 				Match Info = Regex.Match(Fi.Name, RegEx, RegexOptions.IgnoreCase);
 
-				if (Info.Success == false)
+				bool TestInstall = Fi.Name.EndsWith("_TEST.bat", StringComparison.OrdinalIgnoreCase);
+
+				// filter out non-matching or test installation batch files
+				// test installation scripts are intended to be manually invoked
+				if (Info.Success == false || TestInstall)
 				{
+					if (TestInstall)
+					{
+						Log.Verbose("Ignoring test installation batch file {0}", Fi.Name);
+					}
 					continue;
 				}
 
