NonUFSDebugFiles were getting copied to the staging directory first, but did not respect the "DeployLowerCaseFilenames" setting, and were incorrectly determining the captalization of the Engine/Content directory. Copy them after we copy UFSFiles.

diff --git a/Engine/Source/Programs/AutomationTool/Scripts/CopyBuildToStagingDirectory.Automation.cs b/Engine/Source/Programs/AutomationTool/Scripts/CopyBuildToStagingDirectory.Automation.cs
--- a/Engine/Source/Programs/AutomationTool/Scripts/CopyBuildToStagingDirectory.Automation.cs
+++ b/Engine/Source/Programs/AutomationTool/Scripts/CopyBuildToStagingDirectory.Automation.cs
@@ -1262,11 +1262,6 @@ public partial class Project : CommandUtils
 	{
 		CopyManifestFilesToStageDir(SC.FilesToStage.NonUFSFiles, SC.StageDirectory, SC.DebugStageDirectory, "NonUFSFiles", SC.StageTargetPlatform.GetFilesForCRCCheck(), SC.StageTargetPlatform.PlatformType.ToString());
 
-		if (!Params.NoDebugInfo)
-		{
-			CopyManifestFilesToStageDir(SC.FilesToStage.NonUFSDebugFiles, SC.DebugStageDirectory, SC.DebugStageDirectory, "DebugFiles", SC.StageTargetPlatform.GetFilesForCRCCheck(), SC.StageTargetPlatform.PlatformType.ToString());
-		}
-
 		bool bStageUnrealFileSystemFiles = !Params.CookOnTheFly && !Params.UsePak(SC.StageTargetPlatform) && !Params.FileServer;
 		if (bStageUnrealFileSystemFiles)
 		{
@@ -1285,6 +1280,13 @@ public partial class Project : CommandUtils
 			}
 			CopyManifestFilesToStageDir(UFSFiles, SC.StageDirectory, SC.DebugStageDirectory, "UFSFiles", SC.StageTargetPlatform.GetFilesForCRCCheck(), SC.StageTargetPlatform.PlatformType.ToString());
 		}
+
+		// Copy debug files last
+		// they do not respect the DeployLowerCaseFilenames() setting, but if copied to a case-insensitive staging directory first they determine the casing for outer directories (like Engine/Content) 
+		if (!Params.NoDebugInfo)
+		{
+			CopyManifestFilesToStageDir(SC.FilesToStage.NonUFSDebugFiles, SC.DebugStageDirectory, SC.DebugStageDirectory, "DebugFiles", SC.StageTargetPlatform.GetFilesForCRCCheck(), SC.StageTargetPlatform.PlatformType.ToString());
+		}
 	}
 
 	// Pak file rules to apply to a list of files
