BuildGraph: Fix CopyTask printing source and destination files to the log the wrong way around.

diff --git a/Engine/Source/Programs/AutomationTool/BuildGraph/Tasks/CopyTask.cs b/Engine/Source/Programs/AutomationTool/BuildGraph/Tasks/CopyTask.cs
--- a/Engine/Source/Programs/AutomationTool/BuildGraph/Tasks/CopyTask.cs
+++ b/Engine/Source/Programs/AutomationTool/BuildGraph/Tasks/CopyTask.cs
@@ -159,7 +159,7 @@ namespace BuildGraph.Tasks
 			CommandUtils.LogInformation("Copying {0} file{1} from {2} to {3}...", FilePairs.Length, (FilePairs.Length == 1)? "" : "s", SourcePattern.BaseDirectory, TargetPattern.BaseDirectory);
 			foreach(KeyValuePair<FileReference, FileReference> FilePair in FilePairs)
 			{
-				CommandUtils.LogLog("  {0} -> {1}", FilePair.Key, FilePair.Value);
+				CommandUtils.LogLog("  {0} -> {1}", FilePair.Value, FilePair.Key);
 			}
 			CommandUtils.ThreadedCopyFiles(FilePairs.Select(x => x.Value.FullName).ToList(), FilePairs.Select(x => x.Key.FullName).ToList(), bQuiet: true);
 
