Hack to work around AgentInterface being output by multiple nodes, since it's a runtime dependency of several targets.

diff --git a/Engine/Source/Programs/AutomationTool/BuildGraph/BuildGraph.cs b/Engine/Source/Programs/AutomationTool/BuildGraph/BuildGraph.cs
--- a/Engine/Source/Programs/AutomationTool/BuildGraph/BuildGraph.cs
+++ b/Engine/Source/Programs/AutomationTool/BuildGraph/BuildGraph.cs
@@ -678,7 +678,7 @@ namespace AutomationTool
 				foreach(FileReference File in Pair.Value.Files.Select(x => x.ToFileReference(RootDir)))
 				{
 					TempStorageBlock CurrentStorageBlock;
-					if(FileToStorageBlock.TryGetValue(File, out CurrentStorageBlock))
+					if(FileToStorageBlock.TryGetValue(File, out CurrentStorageBlock) && !TempStorage.IsDuplicateBuildProduct(File))
 					{
 						LogError("File '{0}' was produced by {1} and {2}", File, InputStorageBlock, CurrentStorageBlock);
 					}
diff --git a/Engine/Source/Programs/AutomationTool/BuildGraph/TempStorage.cs b/Engine/Source/Programs/AutomationTool/BuildGraph/TempStorage.cs
--- a/Engine/Source/Programs/AutomationTool/BuildGraph/TempStorage.cs
+++ b/Engine/Source/Programs/AutomationTool/BuildGraph/TempStorage.cs
@@ -1,4 +1,4 @@
-﻿// Copyright 1998-2019 Epic Games, Inc. All Rights Reserved.
+// Copyright 1998-2019 Epic Games, Inc. All Rights Reserved.
 
 using System;
 using System.Collections.Concurrent;
@@ -183,8 +183,16 @@ namespace AutomationTool
 			// Check the size matches
 			if(Info.Length != Length)
 			{
-				Message = String.Format("File size differs from manifest - {0} is {1} bytes, expected {2} bytes", RelativePath, Info.Length, Length);
-				return false;
+				if(TempStorage.IsDuplicateBuildProduct(LocalFile))
+				{
+					Message = String.Format("Ignored file size mismatch for {0} - was {1} bytes, expected {2} bytes", RelativePath, Info.Length, Length);
+					return true;
+				}
+				else
+				{
+					Message = String.Format("File size differs from manifest - {0} is {1} bytes, expected {2} bytes", RelativePath, Info.Length, Length);
+					return false;
+				}
 			}
 
 			// Check the timestamp of the file matches. On FAT filesystems writetime has a two seconds resolution (see http://msdn.microsoft.com/en-us/library/windows/desktop/ms724290%28v=vs.85%29.aspx)
@@ -192,7 +200,7 @@ namespace AutomationTool
 			if(TimeDifference.TotalSeconds < -2 || TimeDifference.TotalSeconds > +2)
 			{
 				DateTime ExpectedLocal = new DateTime(LastWriteTimeUtcTicks, DateTimeKind.Utc).ToLocalTime();
-				if(RequireMatchingTimestamps())
+				if(RequireMatchingTimestamps() && !TempStorage.IsDuplicateBuildProduct(LocalFile))
 				{
 					Message = String.Format("File date/time mismatch for {0} - was {1}, expected {2}, TimeDifference {3}", RelativePath, Info.LastWriteTime, ExpectedLocal, TimeDifference);
 					return false;
@@ -1090,6 +1098,22 @@ namespace AutomationTool
 		{
 			return FileReference.Combine(GetDirectoryForNode(BaseDir, NodeName), "Complete");
 		}
+
+		/// <summary>
+		/// Checks whether the given path is whitelisted as a build product that can be produced by more than one node (timestamps may be modified, etc..). Used to suppress
+		/// warnings about build products being overwritten.
+		/// </summary>
+		/// <param name="LocalFile">File name to check</param>
+		/// <returns>True if the given path may be output by multiple build products</returns>
+		public static bool IsDuplicateBuildProduct(FileReference LocalFile)
+		{
+			string FileName = LocalFile.GetFileName();
+			if(FileName.Equals("AgentInterface.dll", StringComparison.OrdinalIgnoreCase) || FileName.Equals("AgentInterface.pdb", StringComparison.OrdinalIgnoreCase))
+			{
+				return true;
+			}
+			return false;
+		}
 	}
 
 	/// <summary>
