Added dynamic resolution to HealthSnapshots. Also, fixed a whitespace issue with RHI Thread and added FrameTime to UnrealSnapshotSummary.ToString()

#rnx

diff --git a/Engine/Source/Programs/AutomationTool/Gauntlet/Unreal/Utils/Gauntlet.UnrealSnapshotParser.cs b/Engine/Source/Programs/AutomationTool/Gauntlet/Unreal/Utils/Gauntlet.UnrealSnapshotParser.cs
--- a/Engine/Source/Programs/AutomationTool/Gauntlet/Unreal/Utils/Gauntlet.UnrealSnapshotParser.cs
+++ b/Engine/Source/Programs/AutomationTool/Gauntlet/Unreal/Utils/Gauntlet.UnrealSnapshotParser.cs
@@ -21,6 +21,7 @@ namespace Gauntlet
 		public float AvgFps;
 		public float Hitches;
 		public float AvgHitches;
+		public float DynamicRes;
 		public float GTTime;
 		public float RTTime;
 		public float GPUTime;
@@ -49,6 +50,7 @@ namespace Gauntlet
 				SB.AppendFormat("AvgFps:\t\t{0:0.00}\n", AvgFps);
 				SB.AppendFormat("HPM:\t\t{0:0.00}\n", Hitches);
 				SB.AppendFormat("AvgH:\t\t{0:0.00}ms\n", AvgHitches);
+				SB.AppendFormat("DynRes:\t\t{0:0.00}\n", DynamicRes);
 				SB.AppendFormat("GT:\t\t{0:0.00}ms\n", GTTime);
 				SB.AppendFormat("RT:\t\t{0:0.00}ms\n", RTTime);
 				SB.AppendFormat("GPU:\t\t{0:0.00}ms\n", GPUTime);
@@ -72,7 +74,6 @@ namespace Gauntlet
 				Name = Groups[1];
 			});
 
-
 			RegexUtil.MatchAndApplyGroups(InContent, @"CPU Memory:[\s\w]+?([\d\.]+)MB,[\s\w]+?([\d\.]+)MB", (Groups) =>
 			{
 				CpuUsedMemory += Convert.ToSingle(Groups[1]);
@@ -118,11 +119,16 @@ namespace Gauntlet
 				FTTime = Convert.ToSingle(Groups[1]);
 			});
 
-			RegexUtil.MatchAndApplyGroups(InContent, @"RHIT:Avg\s([\d\.]+)ms,", (Groups) =>
+			RegexUtil.MatchAndApplyGroups(InContent, @"RHIT:.* Avg\s([\d\.]+)ms,", (Groups) =>
 			{
 				RHIT = Convert.ToSingle(Groups[1]);
 			});
 
+			RegexUtil.MatchAndApplyGroups(InContent, @"DynRes:.* Avg:\s([\d\.]+)%,", (Groups) =>
+			{
+				DynamicRes = Convert.ToSingle(Groups[1]);
+			});
+
 			RegexUtil.MatchAndApplyGroups(InContent, @"DrawCalls:[\w\s:]+?([\d\.]+),[\w\s:]+?([\d\.]+),[\w\s:]+?([\d\.]+)", (Groups) =>
 			{
 				DrawCalls = Convert.ToInt32(Groups[1]);
@@ -150,6 +156,7 @@ namespace Gauntlet
 		public List<float> AvgFps;
 		public List<float> Hitches;
 		public List<float> AvgHitches;
+		public List<float> DynamicRes;
 		public List<float> GTTime;
 		public List<float> RTTime;
 		public List<float> GPUTime;
@@ -170,6 +177,7 @@ namespace Gauntlet
 			AvgFps = new List<float>();
 			Hitches = new List<float>();
 			AvgHitches = new List<float>();
+			DynamicRes = new List<float>();
 			GTTime = new List<float>();
 			RTTime = new List<float>();
 			GPUTime = new List<float>();
@@ -222,6 +230,11 @@ namespace Gauntlet
 					Hitches.Add(Snap.Hitches);
 					AvgHitches.Add(Snap.AvgHitches);
 
+					if (Snap.DynamicRes > 0)
+					{
+						DynamicRes.Add(Snap.DynamicRes);
+					}
+
 					if (Snap.GTTime > 0)
 					{
 						GTTime.Add(Snap.GTTime);
@@ -310,18 +323,32 @@ namespace Gauntlet
 				{
 					SB.AppendFormat("AvgH:\t{0:0.00}ms (Min: {1:0.00}ms, Max: {2:0.00}ms)\n", AvgHitches.Average(), AvgHitches.Min(), AvgHitches.Max());
 				}
+
+				if (DynamicRes.Count > 0)
+				{
+					SB.AppendFormat("DynRes:\t{0:0.00} (Min: {1:0.00}, Max: {2:0.00})\n", DynamicRes.Average(), DynamicRes.Min(), DynamicRes.Max());
+				}
+
 				if (GTTime.Count > 0)
 				{
 					SB.AppendFormat("GT:\t{0:0.00}ms (Min: {1:0.00}ms, Max: {2:0.00}ms)\n", GTTime.Average(), GTTime.Min(), GTTime.Max());
 				}
+
 				if (RTTime.Count > 0)
 				{
 					SB.AppendFormat("RT:\t{0:0.00}ms (Min: {1:0.00}ms, Max: {2:0.00}ms)\n", RTTime.Average(), RTTime.Min(), RTTime.Max());
 				}
+
 				if (GPUTime.Count > 0)
 				{
 					SB.AppendFormat("GPU:\t{0:0.00}ms (Min: {1:0.00}ms, Max: {2:0.00}ms)\n", GPUTime.Average(), GPUTime.Min(), GPUTime.Max());
 				}
+
+				if (FTTime.Count > 0)
+				{
+					SB.AppendFormat("FT:\t{0:0.00}ms (Min: {1:0.00}ms, Max: {2:0.00}ms)\n", FTTime.Average(), FTTime.Min(), FTTime.Max());
+				}
+
 				if (RHIT.Count > 0)
 				{
 					SB.AppendFormat("RHIT:\t{0:0.00}ms (Min: {1:0.00}ms, Max: {2:0.00}ms)\n", RHIT.Average(), RHIT.Min(), RHIT.Max());
diff --git a/Engine/Source/Runtime/Engine/Private/ProfilingDebugging/HealthSnapshot.cpp b/Engine/Source/Runtime/Engine/Private/ProfilingDebugging/HealthSnapshot.cpp
--- a/Engine/Source/Runtime/Engine/Private/ProfilingDebugging/HealthSnapshot.cpp
+++ b/Engine/Source/Runtime/Engine/Private/ProfilingDebugging/HealthSnapshot.cpp
@@ -127,6 +127,10 @@ void FHealthSnapshot::CapturePerformanceStats(const FPerformanceTrackingChart* G
 		FrameTime.Max = GameplayFPSChart->FrametimeHistogram.GetMaxOfAllMeasures();
 		FrameTime.Avg = GameplayFPSChart->FrametimeHistogram.GetAverageOfAllMeasures();
 
+		DynamicResolution.Min = GameplayFPSChart->DynamicResHistogram.GetMinOfAllMeasures();
+		DynamicResolution.Max = GameplayFPSChart->DynamicResHistogram.GetMaxOfAllMeasures();
+		DynamicResolution.Avg = GameplayFPSChart->DynamicResHistogram.GetAverageOfAllMeasures();
+
 		// What % of frames were bound
 		GameThread.PercentFramesBound = FramesCounted > 0 ? (GameplayFPSChart->NumFramesBound_GameThread * 100.0) / FramesCounted : 0;
 
@@ -174,10 +178,11 @@ void FHealthSnapshot::DumpStats(FOutputDevice& Ar, FName CategoryName)
 	{
 		Ar.CategorizedLogf(CategoryName, ELogVerbosity::Log, TEXT("MeasuredPerfTime: %.02f secs"), MeasuredPerfTime);
 		Ar.CategorizedLogf(CategoryName, ELogVerbosity::Log, TEXT("MVP: %.02f%%, AvgFPS:%.2f, HitchesPerMinute: %.2f, Avg Hitch %.02fms"), MVP, AvgFPS, HitchesPerMinute, AvgHitchTime * 1000);
+		Ar.CategorizedLogf(CategoryName, ELogVerbosity::Log, TEXT("DynRes: Avg: %.02f%%, Max: %.02f%%, Min: %.02f%%"), DynamicResolution.Avg, DynamicResolution.Max, DynamicResolution.Min);
 		Ar.CategorizedLogf(CategoryName, ELogVerbosity::Log, TEXT("FT: Avg: %.02fms, Max: %.02fms, Min: %.02fms"), FrameTime.Avg * 1000, FrameTime.Max * 1000, FrameTime.Min * 1000);
 		Ar.CategorizedLogf(CategoryName, ELogVerbosity::Log, TEXT("GT:  Avg %.02fms, Hitches/Min: %.02f, Bound Frames: %.02f%%"), GameThread.AvgTime * 1000, GameThread.HitchesPerMinute, GameThread.PercentFramesBound);
 		Ar.CategorizedLogf(CategoryName, ELogVerbosity::Log, TEXT("RT:  Avg %.02fms, Hitches/Min: %.02f, Bound Frames: %.02f%%"), RenderThread.AvgTime * 1000, RenderThread.HitchesPerMinute, RenderThread.PercentFramesBound);
-		Ar.CategorizedLogf(CategoryName, ELogVerbosity::Log, TEXT("RHIT:Avg %.02fms, Hitches/Min: %.02f, Bound Frames: %.02f%%"), RHIThread.AvgTime * 1000, RHIThread.HitchesPerMinute, RHIThread.PercentFramesBound);
+		Ar.CategorizedLogf(CategoryName, ELogVerbosity::Log, TEXT("RHIT: Avg %.02fms, Hitches/Min: %.02f, Bound Frames: %.02f%%"), RHIThread.AvgTime * 1000, RHIThread.HitchesPerMinute, RHIThread.PercentFramesBound);
 		Ar.CategorizedLogf(CategoryName, ELogVerbosity::Log, TEXT("GPU: Avg %.02fms, Hitches/Min: %.02f, Bound Frames: %.02f%%"), GPU.AvgTime * 1000, GPU.HitchesPerMinute, GPU.PercentFramesBound);
 		Ar.CategorizedLogf(CategoryName, ELogVerbosity::Log, TEXT("DrawCalls: Avg: %d, Max: %d, Min: %d"), DrawCalls.Avg, DrawCalls.Max, DrawCalls.Min);
 		Ar.CategorizedLogf(CategoryName, ELogVerbosity::Log, TEXT("DrawnPrims: Avg: %d, Max: %d, Min: %d"), PrimitivesDrawn.Avg, PrimitivesDrawn.Max, PrimitivesDrawn.Min);
diff --git a/Engine/Source/Runtime/Engine/Public/ProfilingDebugging/HealthSnapshot.h b/Engine/Source/Runtime/Engine/Public/ProfilingDebugging/HealthSnapshot.h
--- a/Engine/Source/Runtime/Engine/Public/ProfilingDebugging/HealthSnapshot.h
+++ b/Engine/Source/Runtime/Engine/Public/ProfilingDebugging/HealthSnapshot.h
@@ -159,6 +159,7 @@ public:
 	FMmaStat<int> DrawCalls;
 	FMmaStat<int> PrimitivesDrawn;
 	FMmaStat<double> FrameTime;
+	FMmaStat<double> DynamicResolution;
 
 	/** Title of the Snapshot */
 	FString Title;
