Remove the CommandUtils.Log function, so we don't shadow the global Log class.

diff --git a/Engine/Source/Programs/AutomationTool/AutomationUtils/CommandUtils.cs b/Engine/Source/Programs/AutomationTool/AutomationUtils/CommandUtils.cs
--- a/Engine/Source/Programs/AutomationTool/AutomationUtils/CommandUtils.cs
+++ b/Engine/Source/Programs/AutomationTool/AutomationUtils/CommandUtils.cs
@@ -123,29 +123,6 @@ namespace AutomationTool
 
 		#region Logging
 
-		/// <summary>
-		/// Writes formatted text to log (with LogEventType.Console).
-		/// </summary>
-		/// <param name="Format">Format string</param>
-		/// <param name="Args">Parameters</param>
-		[MethodImplAttribute(MethodImplOptions.NoInlining)]
-		[Obsolete("CommandUtils.Log() has been deprecated. Use CommandUtils.LogInformation() instead.", false)]
-		public static void Log(string Format, params object[] Args)
-		{
-			Tools.DotNETCommon.Log.WriteLine(1, Tools.DotNETCommon.LogEventType.Console, Format, Args);
-		}
-
-		/// <summary>
-		/// Writes formatted text to log (with LogEventType.Console).
-		/// </summary>
-		/// <param name="Message">Text</param>
-		[MethodImplAttribute(MethodImplOptions.NoInlining)]
-		[Obsolete("CommandUtils.Log() has been deprecated. Use CommandUtils.LogInformation() instead.", false)]
-		public static void Log(string Message)
-		{
-			Tools.DotNETCommon.Log.WriteLine(1, Tools.DotNETCommon.LogEventType.Console, Message);
-		}
-
 		/// <summary>
 		/// Writes formatted text to log (with LogEventType.Console).
 		/// </summary>
