Use case-insensitive comparison for package names

diff --git a/Engine/Source/Programs/AutomationTool/Scripts/CopyBuildToStagingDirectory.Automation.cs b/Engine/Source/Programs/AutomationTool/Scripts/CopyBuildToStagingDirectory.Automation.cs
--- a/Engine/Source/Programs/AutomationTool/Scripts/CopyBuildToStagingDirectory.Automation.cs
+++ b/Engine/Source/Programs/AutomationTool/Scripts/CopyBuildToStagingDirectory.Automation.cs
@@ -1365,6 +1365,8 @@ public partial class Project : CommandUtils
 		List<PakFileRules> RulesList = new List<PakFileRules>();
 		foreach (string SectionName in PakRulesConfig.SectionNames)
 		{
+			//LogInformation("Building PakFileRules for Section {0}", SectionName);
+
 			string PlatformString, TargetString, PakString;
 
 			if (PakRulesConfig.TryGetValue(SectionName, "Platforms", out PlatformString))
@@ -1375,12 +1377,12 @@ public partial class Project : CommandUtils
 				// Check platform string
 				foreach (string Platform in PlatformStrings)
 				{
-					if (SC.StageTargetPlatform.PlatformType.ToString() == Platform)
+					if (SC.StageTargetPlatform.PlatformType.ToString().Equals(Platform, StringComparison.OrdinalIgnoreCase))
 					{
 						bMatches = true;
 						break;
 					}
-					else if (SC.StageTargetPlatform.IniPlatformType.ToString() == Platform)
+					else if (SC.StageTargetPlatform.IniPlatformType.ToString().Equals(Platform, StringComparison.OrdinalIgnoreCase))
 					{
 						bMatches = true;
 						break;
@@ -1389,6 +1391,7 @@ public partial class Project : CommandUtils
 
 				if (!bMatches)
 				{
+					//LogInformation("No matching platform for PakFileRules for Section {0} : {1}, {2}", SectionName, SC.StageTargetPlatform.PlatformType.ToString(), SC.StageTargetPlatform.IniPlatformType.ToString());
 					continue;
 				}
 			}
@@ -1455,6 +1458,7 @@ public partial class Project : CommandUtils
 				PakRules.Filter = new FileFilter();
 				foreach (string FileFilter in FilesEnumberable)
 				{
+					//LogInformation("Adding to PakFileRules for Section {0} : {1}", SectionName, FileFilter);
 					PakRules.Filter.AddRule(FileFilter);
 				}
 
