Fixed issue with UBTArgs including quotes.

diff --git a/Engine/Source/Programs/AutomationTool/Scripts/BuildProjectCommand.Automation.cs b/Engine/Source/Programs/AutomationTool/Scripts/BuildProjectCommand.Automation.cs
--- a/Engine/Source/Programs/AutomationTool/Scripts/BuildProjectCommand.Automation.cs
+++ b/Engine/Source/Programs/AutomationTool/Scripts/BuildProjectCommand.Automation.cs
@@ -124,7 +124,10 @@ public partial class Project : CommandUtils
 
 		if (string.IsNullOrEmpty(Params.UbtArgs) == false)
 		{
-			AdditionalArgs += " " + Params.UbtArgs;
+			string Arg = Params.UbtArgs;
+			Arg = Arg.TrimStart(new char[] { '\"' });
+			Arg = Arg.TrimEnd(new char[] { '\"' });
+			AdditionalArgs += " " + Arg;
 		}
 
 		if (Params.MapFile)
