BuildGraph: Allow using variables when including other scripts.

diff --git a/Engine/Source/Programs/AutomationTool/BuildGraph/Script.cs b/Engine/Source/Programs/AutomationTool/BuildGraph/Script.cs
--- a/Engine/Source/Programs/AutomationTool/BuildGraph/Script.cs
+++ b/Engine/Source/Programs/AutomationTool/BuildGraph/Script.cs
@@ -595,7 +595,7 @@ namespace AutomationTool
 		{
 			if (EvaluateCondition(Element))
 			{
-				FileReference Script = FileReference.Combine(BaseDir, Element.GetAttribute("Script"));
+				FileReference Script = FileReference.Combine(BaseDir, ReadAttribute(Element, "Script"));
 				if (!FileReference.Exists(Script))
 				{
 					LogError(Element, "Cannot find included script '{0}'", Script.FullName);
