Use accessor function on crypto settings in UAT to determine if signing is enabled, rather than reading the flag directly. Makes sure that signing is correctly disabled if the setting is enabled but no key was set

diff --git a/Engine/Source/Programs/AutomationTool/AutomationUtils/ProjectUtils.cs b/Engine/Source/Programs/AutomationTool/AutomationUtils/ProjectUtils.cs
--- a/Engine/Source/Programs/AutomationTool/AutomationUtils/ProjectUtils.cs
+++ b/Engine/Source/Programs/AutomationTool/AutomationUtils/ProjectUtils.cs
@@ -176,7 +176,7 @@ namespace AutomationTool
 				foreach (UnrealTargetPlatform ClientPlatform in ClientTargetPlatforms)
 				{
 					EncryptionAndSigning.CryptoSettings Settings = EncryptionAndSigning.ParseCryptoSettings(RawProjectPath.Directory, ClientPlatform);
-					if (Settings.IsAnyEncryptionEnabled() || Settings.bEnablePakSigning)
+					if (Settings.IsAnyEncryptionEnabled() || Settings.IsPakSigningEnabled())
 					{
 						Reason = "encryption/signing is enabled";
 						return true;
