Fix log output showing CL being synced for iterative cooks.

diff --git a/Engine/Source/Programs/AutomationTool/Scripts/SharedCookedBuild.cs b/Engine/Source/Programs/AutomationTool/Scripts/SharedCookedBuild.cs
--- a/Engine/Source/Programs/AutomationTool/Scripts/SharedCookedBuild.cs
+++ b/Engine/Source/Programs/AutomationTool/Scripts/SharedCookedBuild.cs
@@ -1,4 +1,4 @@
-﻿// Copyright 1998-2019 Epic Games, Inc. All Rights Reserved.
+// Copyright 1998-2019 Epic Games, Inc. All Rights Reserved.
 using System;
 using System.Collections.Generic;
 using System.IO;
@@ -166,7 +166,7 @@ public class SharedCookedBuild
 			}
 
 			FinalCookedBuildPath = BestBuild;
-			CommandUtils.LogInformation("Selected build: {0}");
+			CommandUtils.LogInformation("Selected build: {0}", BestCLNumber);
 			return true;
 		}
 
