Remove superfluous errors when a recursive UAT command fails, or a child process fails. They should already output their own error messages, and adding additional errors is just noise.

diff --git a/Engine/Source/Programs/AutomationTool/AutomationUtils/AutomationException.cs b/Engine/Source/Programs/AutomationTool/AutomationUtils/AutomationException.cs
--- a/Engine/Source/Programs/AutomationTool/AutomationUtils/AutomationException.cs
+++ b/Engine/Source/Programs/AutomationTool/AutomationUtils/AutomationException.cs
@@ -1,4 +1,4 @@
-﻿// Copyright 1998-2019 Epic Games, Inc. All Rights Reserved.
+// Copyright 1998-2019 Epic Games, Inc. All Rights Reserved.
 using System;
 using System.Collections.Generic;
 using System.Linq;
@@ -63,6 +63,27 @@ namespace AutomationTool
 		Error_AndroidOBBError = 155,
 	};
 
+	/// <summary>
+	/// How the exception should be output
+	/// </summary>
+	public enum AutomationExceptionOutputFormat
+	{
+		/// <summary>
+		/// Do not output any message; just output the exception details to the log.
+		/// </summary>
+		Silent,
+
+		/// <summary>
+		/// Output a brief summary to the console, and the exception details to the log
+		/// </summary>
+		Minimal,
+
+		/// <summary>
+		/// Treat the exception as an error
+		/// </summary>
+		Error
+	}
+
     /// <summary>
     /// Exception class used by the AutomationTool to throw exceptions. Allows setting an exit code that will be passed to the entry routine to return to the system on program exit.
     /// If no exit code is given, Error_Unkonwn is used.
@@ -70,6 +91,7 @@ namespace AutomationTool
     public class AutomationException : System.Exception
 	{
         public ExitCode ErrorCode = ExitCode.Error_Unknown;
+		public AutomationExceptionOutputFormat OutputFormat = AutomationExceptionOutputFormat.Error;
 
 		public AutomationException(string Msg)
 			:base(Msg)
diff --git a/Engine/Source/Programs/AutomationTool/AutomationUtils/ProcessUtils.cs b/Engine/Source/Programs/AutomationTool/AutomationUtils/ProcessUtils.cs
--- a/Engine/Source/Programs/AutomationTool/AutomationUtils/ProcessUtils.cs
+++ b/Engine/Source/Programs/AutomationTool/AutomationUtils/ProcessUtils.cs
@@ -961,7 +961,7 @@ namespace AutomationTool
             if (Result.ExitCode > MaxSuccessCode || Result.ExitCode < 0)
             {
                 throw new CommandFailedException((ExitCode)Result.ExitCode, String.Format("Command failed (Result:{3}): {0} {1}. See logfile for details: '{2}' ",
-                                                App, CommandLine, Path.GetFileName(Logfile), Result.ExitCode));
+                                                App, CommandLine, Path.GetFileName(Logfile), Result.ExitCode)){ OutputFormat = AutomationExceptionOutputFormat.Minimal };
             }
             if (!String.IsNullOrEmpty(Result.Output))
             {
@@ -1050,7 +1050,7 @@ namespace AutomationTool
 			IProcessResult Result = Run(CmdEnv.UATExe, CommandLine, null, ERunOptions.Default, EnvironmentVars, Identifier: Identifier);
 			if (Result.ExitCode != 0)
 			{
-				throw new CommandFailedException(String.Format("Recursive UAT command failed (exit code {0})", Result.ExitCode));
+				throw new CommandFailedException(String.Format("Recursive UAT command failed (exit code {0})", Result.ExitCode)){ OutputFormat = AutomationExceptionOutputFormat.Silent };
 			}
 		}
 
diff --git a/Engine/Source/Programs/AutomationTool/Program.cs b/Engine/Source/Programs/AutomationTool/Program.cs
--- a/Engine/Source/Programs/AutomationTool/Program.cs
+++ b/Engine/Source/Programs/AutomationTool/Program.cs
@@ -89,8 +89,22 @@ namespace AutomationTool
 			}
 			catch (AutomationException Ex)
 			{
+				// Output the message in the desired format
+				if(Ex.OutputFormat == AutomationExceptionOutputFormat.Silent)
+				{
+					Log.TraceLog("{0}", ExceptionUtils.FormatExceptionDetails(Ex));
+				}
+				else if(Ex.OutputFormat == AutomationExceptionOutputFormat.Minimal)
+				{
+					Log.TraceInformation("{0}", Ex.Message);
+					Log.TraceLog("{0}", ExceptionUtils.FormatExceptionDetails(Ex));
+				}
+				else
+				{
+					Log.WriteException(Ex, LogUtils.FinalLogFileName);
+				}
+
 				// Take the exit code from the exception
-				Log.WriteException(Ex, LogUtils.FinalLogFileName);
 				ReturnCode = Ex.ErrorCode;
 			}
 			catch (Exception Ex)
