Simplify error messages for UE4Build, so we don't get unnecessary lines tagged by EC.

diff --git a/Engine/Source/Programs/AutomationTool/AutomationUtils/UBTUtils.cs b/Engine/Source/Programs/AutomationTool/AutomationUtils/UBTUtils.cs
--- a/Engine/Source/Programs/AutomationTool/AutomationUtils/UBTUtils.cs
+++ b/Engine/Source/Programs/AutomationTool/AutomationUtils/UBTUtils.cs
@@ -47,7 +47,7 @@ namespace AutomationTool
 			IProcessResult Result = Run(UBTExecutable, CommandLine, Options: ERunOptions.AllowSpew | ERunOptions.NoStdOutCapture);
 			if(Result.ExitCode != 0)
 			{
-				throw new AutomationException((ExitCode)Result.ExitCode, "UnrealBuildTool failed. See log for more details. ({0})", CommandUtils.CombinePaths(Env.FinalLogFolder, LogName));
+				throw new AutomationException((ExitCode)Result.ExitCode, "UnrealBuildTool failed. See log for more details. ({0})", CommandUtils.CombinePaths(Env.FinalLogFolder, LogName)) { OutputFormat = AutomationExceptionOutputFormat.Minimal };
 			}
 		}
 
diff --git a/Engine/Source/Programs/AutomationTool/AutomationUtils/UE4Build.cs b/Engine/Source/Programs/AutomationTool/AutomationUtils/UE4Build.cs
--- a/Engine/Source/Programs/AutomationTool/AutomationUtils/UE4Build.cs
+++ b/Engine/Source/Programs/AutomationTool/AutomationUtils/UE4Build.cs
@@ -12,6 +12,20 @@ using System.Linq;
 
 namespace AutomationTool
 {
+	class UE4BuildException : AutomationException
+	{
+		public UE4BuildException(string Message)
+			: base("BUILD FAILED: " + Message)
+		{
+			OutputFormat = AutomationExceptionOutputFormat.Minimal;
+		}
+
+		public UE4BuildException(string Format, params object[] Args)
+			: this(String.Format(Format, Args))
+		{
+		}
+	}
+
 	[Help("ForceMonolithic", "Toggle to combined the result into one executable")]
 	[Help("ForceDebugInfo", "Forces debug info even in development builds")]
 	[Help("NoXGE", "Toggle to disable the distributed build process")]
@@ -40,7 +54,7 @@ namespace AutomationTool
 			string File = CommandUtils.CombinePaths(InFile);
 			if (!CommandUtils.FileExists(File) && !CommandUtils.DirectoryExists(File))
 			{
-				throw new AutomationException("BUILD FAILED specified file to AddBuildProduct {0} does not exist.", File);
+				throw new UE4BuildException("Specified file to AddBuildProduct {0} does not exist.", File);
 			}
 			if (!HasBuildProduct(InFile))
 			{
@@ -60,7 +74,7 @@ namespace AutomationTool
 		{
 			if (!FileReference.Exists(ManifestFile))
 			{
-				throw new AutomationException("BUILD FAILED UBT Manifest {0} does not exist.", ManifestFile);
+				throw new UE4BuildException("UBT Manifest {0} does not exist.", ManifestFile);
 			}
 
 			BuildManifest Manifest = CommandUtils.ReadManifest(ManifestFile);
@@ -68,7 +82,7 @@ namespace AutomationTool
 			{
 				if (!CommandUtils.FileExists_NoExceptions(Item) && !CommandUtils.DirectoryExists_NoExceptions(Item))
 				{
-					throw new AutomationException("BUILD FAILED {0} was in manifest but was not produced.", Item);
+					throw new UE4BuildException("{0} was in manifest but was not produced.", Item);
 				}
 				AddBuildProduct(Item);
 			}
@@ -79,7 +93,7 @@ namespace AutomationTool
 		{			
 			if (CommandUtils.FileExists(UBTExecutable) == false)
 			{
-				throw new AutomationException("UBT does not exist in {0}.", UBTExecutable);
+				throw new UE4BuildException("UBT does not exist in {0}.", UBTExecutable);
 			}
 		}
 
@@ -127,7 +141,7 @@ namespace AutomationTool
 			{
 				if (!CommandUtils.FileExists_NoExceptions(XGEFile))
 				{
-					throw new AutomationException("BUILD FAILED: Couldn't find file: {0}", XGEFile);
+					throw new UE4BuildException("Couldn't find file: {0}", XGEFile);
 				}
 				int FileNum = 0;
 				string OutFile;
@@ -160,7 +174,7 @@ namespace AutomationTool
 			{
 				if (!CommandUtils.FileExists_NoExceptions(ManifestItem))
 				{
-					throw new AutomationException("BUILD FAILED {0} was in manifest but was not produced.", ManifestItem);
+					throw new UE4BuildException("{0} was in manifest but was not produced.", ManifestItem);
 				}
 				AddBuildProduct(ManifestItem);
 			}
@@ -567,7 +581,7 @@ namespace AutomationTool
 				}
 				if (!CommandUtils.FileExists(TaskFilePath))
 				{
-					throw new AutomationException("Unable to find xge xml: " + TaskFilePath);
+					throw new UE4BuildException("Unable to find xge xml: " + TaskFilePath);
 				}
 
 				CombineXGEStopwatch.Finish();
@@ -614,7 +628,7 @@ namespace AutomationTool
 							}
 							else if (SuccesCode != 0)
 							{
-								throw new AutomationException("Command failed (Result:{3}): {0} {1}. See logfile for details: '{2}' ",
+								throw new UE4BuildException("Command failed (Result:{3}): {0} {1}. See logfile for details: '{2}' ",
 																XGETool, Args, Path.GetFileName(LogFile), SuccesCode);
 							}
 							CommandUtils.LogVerbose("{0} {1} Done *******", XGETool, Args);
@@ -657,7 +671,7 @@ namespace AutomationTool
 				{
 					if (!CommandUtils.FileExists_NoExceptions(XGEFile))
 					{
-						throw new AutomationException("BUILD FAILED: Couldn't find file: {0}", XGEFile);
+						throw new UE4BuildException("Couldn't find file: {0}", XGEFile);
 					}
 					var TargetFile = TaskFilePath + "." + Path.GetFileName(XGEFile);
 					CommandUtils.CopyFile(XGEFile, TargetFile);
@@ -901,7 +915,7 @@ namespace AutomationTool
 											string ActionJob = ThisAction + string.Format("_{0}", Job);
 											if (!ActionToAction.ContainsKey(ActionJob))
 											{
-												throw new AutomationException("Action not found '{0}' in {1}->{2}", ActionJob, XGEFile, TargetFile);
+												throw new UE4BuildException("Action not found '{0}' in {1}->{2}", ActionJob, XGEFile, TargetFile);
 												// the XGE schedule is not topologically sorted. Hmmm. Well, we make a scary assumption here that this 
 											}
 											NewDepends += ActionToAction[ActionJob];
@@ -976,7 +990,7 @@ namespace AutomationTool
 		{
 			if (!CommandUtils.CmdEnv.HasCapabilityToCompile)
 			{
-				throw new AutomationException("You are attempting to compile on a machine that does not have a supported compiler!");
+				throw new UE4BuildException("You are attempting to compile on a machine that does not have a supported compiler!");
 			}
 			bool DeleteBuildProducts = InDeleteBuildProducts.HasValue ? InDeleteBuildProducts.Value : ParseParam("Clean");
 			if (InUpdateVersionFiles)
@@ -1128,7 +1142,7 @@ namespace AutomationTool
 			CommandUtils.LogSetProgress(InShowProgress, "Building...");
 			if (!ProcessXGEItems(XGEItems, XGETool, Args, TaskFilePath, InShowProgress))
 			{
-				throw new AutomationException("BUILD FAILED: {0} failed, retries not enabled:", XGETool);
+				throw new UE4BuildException("{0} failed, retries not enabled:", XGETool);
 			}
 		}
 
@@ -1151,7 +1165,7 @@ namespace AutomationTool
 					{
 						if (!CommandUtils.FileExists(Product) && !CommandUtils.DirectoryExists(Product))
 						{
-							throw new AutomationException("BUILD FAILED {0} was a build product but no longer exists", Product);
+							throw new UE4BuildException("{0} was a build product but no longer exists", Product);
 						}
 						CommandUtils.LogLog(Product);
 					}
@@ -1173,7 +1187,7 @@ namespace AutomationTool
 				CommandUtils.P4.Sync("-f -k " + CommandUtils.MakePathSafeToUseWithCommandLine(File) + "#head"); // sync the file without overwriting local one
 				if (!CommandUtils.FileExists(File))
 				{
-					throw new AutomationException("BUILD FAILED {0} was a build product but no longer exists", File);
+					throw new UE4BuildException("{0} was a build product but no longer exists", File);
 				}
 
 				CommandUtils.P4.ReconcileNoDeletes(WorkingCL, CommandUtils.MakePathSafeToUseWithCommandLine(File));
@@ -1237,7 +1251,7 @@ namespace AutomationTool
 				var UBTProduct = CommandUtils.CombinePaths(UBTLocation, UBTFile);
 				if (!CommandUtils.FileExists_NoExceptions(UBTProduct))
 				{
-					throw new AutomationException("Cannot add UBT to the build products because {0} does not exist.", UBTProduct);
+					throw new UE4BuildException("Cannot add UBT to the build products because {0} does not exist.", UBTProduct);
 				}
 				AddBuildProduct(UBTProduct);
 			}
@@ -1261,7 +1275,7 @@ namespace AutomationTool
 				var OutputFile = CommandUtils.CombinePaths(DotNetOutputLocation, UATFile);
 				if (!CommandUtils.FileExists_NoExceptions(OutputFile))
 				{
-					throw new AutomationException("Cannot add UAT to the build products because {0} does not exist.", OutputFile);
+					throw new UE4BuildException("Cannot add UAT to the build products because {0} does not exist.", OutputFile);
 				}
 				AddBuildProduct(OutputFile);
 			}
@@ -1279,7 +1293,7 @@ namespace AutomationTool
 			var UATScripts = Directory.GetFiles(UATScriptsLocation, "*" + ScriptsPostfix, SearchOption.AllDirectories);
 			if (CommandUtils.IsNullOrEmpty(UATScripts))
 			{
-				throw new AutomationException("No automation scripts found in {0}. Cannot add UAT files to the build products.", UATScriptsLocation);
+				throw new UE4BuildException("No automation scripts found in {0}. Cannot add UAT files to the build products.", UATScriptsLocation);
 			}
 
 			var UATFiles = new List<string>(new string[] 
@@ -1298,7 +1312,7 @@ namespace AutomationTool
 				var OutputFile = CommandUtils.CombinePaths(DotNetOutputLocation, UATFile);
 				if (!CommandUtils.FileExists_NoExceptions(OutputFile))
 				{
-					throw new AutomationException("Cannot add UAT to the build products because {0} does not exist.", OutputFile);
+					throw new UE4BuildException("Cannot add UAT to the build products because {0} does not exist.", OutputFile);
 				}
 				AddBuildProduct(OutputFile);
 			}
@@ -1308,7 +1322,7 @@ namespace AutomationTool
 			{
 				if (!CommandUtils.FileExists_NoExceptions(UATScriptFilePath))
 				{
-					throw new AutomationException("Cannot add UAT to the build products because {0} does not exist.", UATScriptFilePath);
+					throw new UE4BuildException("Cannot add UAT to the build products because {0} does not exist.", UATScriptFilePath);
 				}
 
 				AddBuildProduct(UATScriptFilePath);
diff --git a/Engine/Source/Programs/AutomationTool/Program.cs b/Engine/Source/Programs/AutomationTool/Program.cs
--- a/Engine/Source/Programs/AutomationTool/Program.cs
+++ b/Engine/Source/Programs/AutomationTool/Program.cs
@@ -96,7 +96,7 @@ namespace AutomationTool
 				}
 				else if(Ex.OutputFormat == AutomationExceptionOutputFormat.Minimal)
 				{
-					Log.TraceInformation("{0}", Ex.Message);
+					Log.TraceInformation("{0}", Ex.ToString().Replace("\n", "\n  "));
 					Log.TraceLog("{0}", ExceptionUtils.FormatExceptionDetails(Ex));
 				}
 				else
