Moved param validation to common code
Added -flavor check to BuildCookFort

diff --git a/Engine/Source/Programs/AutomationTool/AutomationUtils/BuildCommand.cs b/Engine/Source/Programs/AutomationTool/AutomationUtils/BuildCommand.cs
--- a/Engine/Source/Programs/AutomationTool/AutomationUtils/BuildCommand.cs
+++ b/Engine/Source/Programs/AutomationTool/AutomationUtils/BuildCommand.cs
@@ -226,6 +226,27 @@ namespace AutomationTool
 			}
 		}
 
+		/// <summary>
+		/// Checks that all of the required params are present, throws an exception if not
+		/// </summary>
+		/// <param name="Args"></param>
+		public void CheckParamsArePresent(params string[] Args)
+		{
+			List<string> MissingParams = new List<string>();
+			foreach (string Arg in Args)
+			{
+				if (ParseParamValue(Arg, null) == null)
+				{
+					MissingParams.Add(Arg);
+				}
+			}
+
+			if (MissingParams.Count > 0)
+			{
+				throw new AutomationException("Params {0} are missing but required. Required params are {1}", string.Join(",", MissingParams), string.Join(",", Args));
+			}
+		}
+
 		/// <summary>
 		/// Build command entry point.  Throws AutomationExceptions on failure.
 		/// </summary>
