Add support for BPT ChunkDeltaOptimise 'DiffAbortThreshold' commandline param to automation code.

diff --git a/Engine/Source/Programs/AutomationTool/AutomationUtils/MCPPublic.cs b/Engine/Source/Programs/AutomationTool/AutomationUtils/MCPPublic.cs
--- a/Engine/Source/Programs/AutomationTool/AutomationUtils/MCPPublic.cs
+++ b/Engine/Source/Programs/AutomationTool/AutomationUtils/MCPPublic.cs
@@ -877,6 +877,10 @@ namespace EpicGames.MCP.Automation
 			/// The file path to the update manifest. New data will be added to the cloud directory that this manifest is in.
 			/// </summary>
 			public string DestinationManifest;
+			/// <summary>
+			/// Specifies in bytes, an upper limit for original diffs to try to enhance.
+			/// </summary>
+			public ulong DiffAbortThreshold;
 		}
 
 		static BuildPatchToolBase Handler = null;
