Move patchSeekOptMaxGapSize to ini file Part 1

diff --git a/Engine/Source/Programs/AutomationTool/Scripts/CopyBuildToStagingDirectory.Automation.cs b/Engine/Source/Programs/AutomationTool/Scripts/CopyBuildToStagingDirectory.Automation.cs
--- a/Engine/Source/Programs/AutomationTool/Scripts/CopyBuildToStagingDirectory.Automation.cs
+++ b/Engine/Source/Programs/AutomationTool/Scripts/CopyBuildToStagingDirectory.Automation.cs
@@ -1945,9 +1945,16 @@ public partial class Project : CommandUtils
                                 }
                             }
                         }
-                    }
+					}
+
+					string PatchSeekOptMaxGapSizeOption = String.Empty;
+					string PatchSeekOptMaxGapSize = String.Empty;
+					if (PlatformGameConfig.GetString("/Script/UnrealEd.ProjectPackagingSettings", "PatchSeekOptMaxGapSize", out PatchSeekOptMaxGapSize))
+					{
+						PatchSeekOptMaxGapSizeOption = String.Format(" -patchSeekOptMaxGapSize={0}", PatchSeekOptMaxGapSize);
+					}
 
-					Commands.Add(GetUnrealPakArguments(PakParams.UnrealPakResponseFile, OutputLocation, PrimaryOrderFile, SC.StageTargetPlatform.GetPlatformPakCommandLine(Params, SC) + BulkOption + CompressionFormats + " " + Params.AdditionalPakOptions, PakParams.bCompressed, CryptoSettings, CryptoKeysCacheFilename, PatchSourceContentPath, PakParams.EncryptionKeyGuid, SecondaryOrderFile));
+					Commands.Add(GetUnrealPakArguments(PakParams.UnrealPakResponseFile, OutputLocation, PrimaryOrderFile, SC.StageTargetPlatform.GetPlatformPakCommandLine(Params, SC) + PatchSeekOptMaxGapSizeOption + BulkOption + CompressionFormats + " " + Params.AdditionalPakOptions, PakParams.bCompressed, CryptoSettings, CryptoKeysCacheFilename, PatchSourceContentPath, PakParams.EncryptionKeyGuid, SecondaryOrderFile));
 					LogNames.Add(OutputLocation.GetFileNameWithoutExtension());
 				}
 			}
