Made the pak file binary compatible with previous version to reduce patch sizes.
Reenabled oodle compression in 8.30 for android, ios and switch.
Cleaned up unrealpak commandline when run from UAT.

diff --git a/Engine/Plugins/Runtime/PacketHandlers/CompressionComponents/Oodle/Oodle.uplugin b/Engine/Plugins/Runtime/PacketHandlers/CompressionComponents/Oodle/Oodle.uplugin
--- a/Engine/Plugins/Runtime/PacketHandlers/CompressionComponents/Oodle/Oodle.uplugin
+++ b/Engine/Plugins/Runtime/PacketHandlers/CompressionComponents/Oodle/Oodle.uplugin
@@ -36,7 +36,7 @@
 		},
 		{
 			"Name" : "OodleCompressionFormat",
-			"Type" : "Runtime",
+			"Type" : "RuntimeAndProgram",
 			"LoadingPhase" : "EarliestPossible",
 			"WhitelistPlatforms" :
 			[
diff --git a/Engine/Source/Developer/PakFileUtilities/Private/PakFileUtilities.cpp b/Engine/Source/Developer/PakFileUtilities/Private/PakFileUtilities.cpp
--- a/Engine/Source/Developer/PakFileUtilities/Private/PakFileUtilities.cpp
+++ b/Engine/Source/Developer/PakFileUtilities/Private/PakFileUtilities.cpp
@@ -180,6 +180,7 @@ struct FPakCommandLineParameters
 		, EncryptIndex(false)
 		, UseCustomCompressor(false)
 		, bSign(false)
+		, bPatchCompatibilityMode421(false)
 	{
 	}
 
@@ -197,6 +198,7 @@ struct FPakCommandLineParameters
 	bool UseCustomCompressor;
 	FGuid EncryptionKeyGuid;
 	bool bSign;
+	bool bPatchCompatibilityMode421;
 };
 
 struct FPakEntryPair
@@ -628,6 +630,11 @@ void ProcessCommandLine(const TCHAR* CmdLine, const TArray<FString>& NonOptionAr
 	FString ResponseFile;
 	FString ClusterSizeString;
 
+	if (FParse::Param(CmdLine, TEXT("patchcompatibilitymode421")))
+	{
+		CmdLineParameters.bPatchCompatibilityMode421 = true;
+	}
+
 	if (FParse::Value(CmdLine, TEXT("-blocksize="), ClusterSizeString) && 
 		FParse::Value(CmdLine, TEXT("-blocksize="), CmdLineParameters.FileSystemBlockSize))
 	{
@@ -1426,6 +1433,19 @@ bool CreatePakFile(const TCHAR* Filename, TArray<FPakInputPair>& FilesToAdd, con
 	Info.bEncryptedIndex = (InKeyChain.MasterEncryptionKey && CmdLineParameters.EncryptIndex);
 	Info.EncryptionKeyGuid = InKeyChain.MasterEncryptionKey ? InKeyChain.MasterEncryptionKey->Guid : FGuid();
 
+
+	if (CmdLineParameters.bPatchCompatibilityMode421)
+	{
+		// for old versions, put in some known names that we may have used
+		Info.GetCompressionMethodIndex(NAME_None);
+		Info.GetCompressionMethodIndex(NAME_Zlib);
+		Info.GetCompressionMethodIndex(NAME_Gzip);
+		Info.GetCompressionMethodIndex(TEXT("Bogus"));
+		Info.GetCompressionMethodIndex(TEXT("Oodle"));
+		
+	}
+
+
 	if (InKeyChain.MasterEncryptionKey)
 	{
 		UE_LOG(LogPakFile, Display, TEXT("Encrypting using key '%s' [%s]"), *InKeyChain.MasterEncryptionKey->Name, *InKeyChain.MasterEncryptionKey->Guid.ToString());
diff --git a/Engine/Source/Programs/AutomationTool/Scripts/CopyBuildToStagingDirectory.Automation.cs b/Engine/Source/Programs/AutomationTool/Scripts/CopyBuildToStagingDirectory.Automation.cs
--- a/Engine/Source/Programs/AutomationTool/Scripts/CopyBuildToStagingDirectory.Automation.cs
+++ b/Engine/Source/Programs/AutomationTool/Scripts/CopyBuildToStagingDirectory.Automation.cs
@@ -109,9 +109,10 @@ public partial class Project : CommandUtils
 		return Result;
 	}
 
-	static public string GetUnrealPakArguments(Dictionary<string, string> UnrealPakResponseFile, FileReference OutputLocation, FileReference PakOrderFileLocation, string PlatformOptions, bool Compressed, EncryptionAndSigning.CryptoSettings CryptoSettings, FileReference CryptoKeysCacheFilename, String PatchSourceContentPath, string EncryptionKeyGuid, FileReference SecondaryPakOrderFileLocation=null)
+	static public string GetUnrealPakArguments(FileReference ProjectPath, Dictionary<string, string> UnrealPakResponseFile, FileReference OutputLocation, FileReference PakOrderFileLocation, string PlatformOptions, bool Compressed, EncryptionAndSigning.CryptoSettings CryptoSettings, FileReference CryptoKeysCacheFilename, String PatchSourceContentPath, string EncryptionKeyGuid, FileReference SecondaryPakOrderFileLocation=null)
 	{
-		StringBuilder CmdLine = new StringBuilder(MakePathSafeToUseWithCommandLine(OutputLocation.FullName));
+		StringBuilder CmdLine = new StringBuilder(MakePathSafeToUseWithCommandLine(ProjectPath.FullName));
+		CmdLine.AppendFormat(" {0}", MakePathSafeToUseWithCommandLine(OutputLocation.FullName));
 
 		// Force encryption of ALL files if we're using specific encryption key. This should be made an option per encryption key in the settings, but for our initial
 		// implementation we will just assume that we require maximum security for this data.
@@ -177,7 +178,7 @@ public partial class Project : CommandUtils
 			return;
 		}
 
-		string Arguments = GetUnrealPakArguments(UnrealPakResponseFile, OutputLocation, PakOrderFileLocation, PlatformOptions, Compressed, CryptoSettings, CryptoKeysCacheFilename, PatchSourceContentPath, EncryptionKeyGuid, SecondaryPakOrderFileLocation);
+		string Arguments = GetUnrealPakArguments(Params.RawProjectPath, UnrealPakResponseFile, OutputLocation, PakOrderFileLocation, PlatformOptions, Compressed, CryptoSettings, CryptoKeysCacheFilename, PatchSourceContentPath, EncryptionKeyGuid, SecondaryPakOrderFileLocation);
 		RunAndLog(CmdEnv, GetUnrealPakLocation().FullName, Arguments, Options: ERunOptions.Default | ERunOptions.UTF8Output);
 	}
 
@@ -1954,7 +1955,7 @@ public partial class Project : CommandUtils
 						PatchSeekOptMaxGapSizeOption = String.Format(" -patchSeekOptMaxGapSize={0}", PatchSeekOptMaxGapSize);
 					}
 
-					Commands.Add(GetUnrealPakArguments(PakParams.UnrealPakResponseFile, OutputLocation, PrimaryOrderFile, SC.StageTargetPlatform.GetPlatformPakCommandLine(Params, SC) + PatchSeekOptMaxGapSizeOption + BulkOption + CompressionFormats + " " + Params.AdditionalPakOptions, PakParams.bCompressed, CryptoSettings, CryptoKeysCacheFilename, PatchSourceContentPath, PakParams.EncryptionKeyGuid, SecondaryOrderFile));
+					Commands.Add(GetUnrealPakArguments(Params.RawProjectPath, PakParams.UnrealPakResponseFile, OutputLocation, PrimaryOrderFile, SC.StageTargetPlatform.GetPlatformPakCommandLine(Params, SC) + PatchSeekOptMaxGapSizeOption + BulkOption + CompressionFormats + " " + Params.AdditionalPakOptions, PakParams.bCompressed, CryptoSettings, CryptoKeysCacheFilename, PatchSourceContentPath, PakParams.EncryptionKeyGuid, SecondaryOrderFile));
 					LogNames.Add(OutputLocation.GetFileNameWithoutExtension());
 				}
 			}
@@ -1965,7 +1966,7 @@ public partial class Project : CommandUtils
 		// Actually execute UnrealPak
 		if (Commands.Count > 0)
 		{
-			RunUnrealPakInParallel(Params.RawProjectPath, Commands, LogNames, AdditionalCompressionOptionsOnCommandLine);
+			RunUnrealPakInParallel( Commands, LogNames, AdditionalCompressionOptionsOnCommandLine);
 		}
 
 		// Do any additional processing on the command output
@@ -2101,7 +2102,7 @@ public partial class Project : CommandUtils
 		}
 	}
 
-	private static void RunUnrealPakInParallel(FileReference ProjectFile, List<string> Commands, List<string> LogNames, string AdditionalCompressionOptionsOnCommandLine)
+	private static void RunUnrealPakInParallel(List<string> Commands, List<string> LogNames, string AdditionalCompressionOptionsOnCommandLine)
 	{
 		LogInformation("Executing {0} UnrealPak command{1}...", Commands.Count, (Commands.Count > 1)? "s" : "");
 
@@ -2113,7 +2114,7 @@ public partial class Project : CommandUtils
 			for(int Idx = 0; Idx < Commands.Count; Idx++)
 			{
 				string LogFile = LogUtils.GetUniqueLogName(CombinePaths(CmdEnv.LogFolder, String.Format("UnrealPak-{0}", LogNames[Idx])));
-				string Arguments = Commands[Idx] = String.Format("{0} {1} -multiprocess -abslog={2} {3}", Commands[Idx], MakePathSafeToUseWithCommandLine(ProjectFile.FullName), MakePathSafeToUseWithCommandLine(LogFile), AdditionalCompressionOptionsOnCommandLine);
+				string Arguments = Commands[Idx] = String.Format("{0} -multiprocess -abslog={1} {2}", Commands[Idx], MakePathSafeToUseWithCommandLine(LogFile), AdditionalCompressionOptionsOnCommandLine);
 
 				int LocalIdx = Idx;
 				Queue.Enqueue(() => Results[LocalIdx] = Run(UnrealPakExe, Arguments, Options: ERunOptions.AppMustExist | ERunOptions.NoLoggingOfRunCommand));
diff --git a/Engine/Source/Runtime/PakFile/Public/IPlatformFilePak.h b/Engine/Source/Runtime/PakFile/Public/IPlatformFilePak.h
--- a/Engine/Source/Runtime/PakFile/Public/IPlatformFilePak.h
+++ b/Engine/Source/Runtime/PakFile/Public/IPlatformFilePak.h
@@ -70,7 +70,7 @@ struct FPakInfo
 		/** Length of a compression format name */
 		CompressionMethodNameLen = 32,
 		/** Number of allowed different methods */
-		MaxNumCompressionMethods=4,
+		MaxNumCompressionMethods=5, // when we remove patchcompatibilitymode421 we can reduce this to 4
 	};
 
 	/** Version numbers. */
@@ -311,7 +311,7 @@ struct FPakEntry
 	/** Size of a compressed block in the file. */
 	uint32 CompressionBlockSize;
 	/** Index into the compression methods in this pakfile. */
-	uint8 CompressionMethodIndex;
+	uint32 CompressionMethodIndex;
 	/** Pak entry flags. */
 	uint8 Flags;
 	/** Flag is set to true when FileHeader has been checked against PakHeader. It is not serialized. */
