BuildGraph: Add a "Project" attribute to the Compile task.

diff --git a/Engine/Source/Programs/AutomationTool/BuildGraph/Tasks/CompileTask.cs b/Engine/Source/Programs/AutomationTool/BuildGraph/Tasks/CompileTask.cs
--- a/Engine/Source/Programs/AutomationTool/BuildGraph/Tasks/CompileTask.cs
+++ b/Engine/Source/Programs/AutomationTool/BuildGraph/Tasks/CompileTask.cs
@@ -35,6 +35,12 @@ namespace AutomationTool
 		[TaskParameter]
 		public UnrealTargetPlatform Platform;
 
+		/// <summary>
+		/// The project to compile with
+		/// </summary>
+		[TaskParameter(Optional = true)]
+		public string Project;
+
 		/// <summary>
 		/// Additional arguments for UnrealBuildTool
 		/// </summary>
@@ -129,7 +135,7 @@ namespace AutomationTool
 			bAllowXGE &= Parameters.AllowXGE;
 			bAllowParallelExecutor &= Parameters.AllowParallelExecutor;
 
-			UE4Build.BuildTarget Target = new UE4Build.BuildTarget { TargetName = Parameters.Target, Platform = Parameters.Platform, Config = Parameters.Configuration, UBTArgs = "-nobuilduht " + (Parameters.Arguments ?? ""), Clean = Parameters.Clean };
+			UE4Build.BuildTarget Target = new UE4Build.BuildTarget { TargetName = Parameters.Target, Platform = Parameters.Platform, Config = Parameters.Configuration, UprojectPath = String.IsNullOrEmpty(Parameters.Project)? null : new FileReference(Parameters.Project), UBTArgs = "-nobuilduht " + (Parameters.Arguments ?? ""), Clean = Parameters.Clean };
 			if(!String.IsNullOrEmpty(Parameters.Tag))
 			{
 				TargetToTagName.Add(Target, Parameters.Tag);
