Various gauntlet functionality for getting clients to connect to propped prod servers.
- Changed gauntlet accountpool global object into an AccountManager system that handles an AccountPool instance to allow easier overriding
- Add -nobuildid flag to prevent buildidoverride from being added onto commandline
- Change game tests over to use new account pool

diff --git a/Engine/Source/Programs/AutomationTool/Gauntlet/Framework/Account/Gauntlet.AccountPool.cs b/Engine/Source/Programs/AutomationTool/Gauntlet/Framework/Account/Gauntlet.AccountPool.cs
--- a/Engine/Source/Programs/AutomationTool/Gauntlet/Framework/Account/Gauntlet.AccountPool.cs
+++ b/Engine/Source/Programs/AutomationTool/Gauntlet/Framework/Account/Gauntlet.AccountPool.cs
@@ -3,63 +3,52 @@
 using System;
 using System.Collections.Generic;
 using System.Linq;
+using System.Reflection;
 
 namespace Gauntlet
 {
-	/// <summary>
-	/// A pool of accounts for tests to use. The pool should be filled by the project-specific script that derives from
-	/// RunUnrealTests
-	/// </summary>
-	public class AccountPool : IDisposable
+
+	public class AccountPool
 	{
-		/// <summary>
-		/// Object used for locking access to internal data
-		/// </summary>
-		private Object LockObject = new Object();
 
 		/// <summary>
-		/// List of all registered accounts
-		/// </summary>
-		protected List<Account> AllAccounts = new List<Account>();
-
-		/// <summary>
-		/// Accounts that have been reserved
-		/// </summary>
-		protected Dictionary<Account, bool> ReservedAccounts = new Dictionary<Account, bool>();
-
-		/// <summary>
-		/// A random ordering of AllAccounts
-		/// </summary>
-		protected List<Account> RandomlyOrderedList = new List<Account>();
-
-		/// <summary>
 		/// Singleton
 		/// </summary>
-		private static AccountPool _Instance;
+		private static IManagedAccountPool _Instance;
 
-		/// <summary>
-		/// 
-		/// </summary>
-		protected AccountPool()
+		public static void Initialize()
 		{
-			if (_Instance == null)
-			{
-				_Instance = this;
-			}
+			Initialize <ManagedAccountPool> ();
 		}
 
 		~AccountPool()
 		{
-			Dispose(false);
+			if (_Instance != null)
+			{
+				_Instance.Dispose();
+			}
 		}
 
-		public static AccountPool Instance
+		public static void Initialize<T>()
 		{
+			object PossiblePool = Activator.CreateInstance(typeof(T));
+			if (PossiblePool is IManagedAccountPool)
+			{
+				if (_Instance != null)
+				{
+					_Instance.Dispose();
+				}
+				_Instance = (IManagedAccountPool)PossiblePool;
+			}
+		}
+
+		public static IManagedAccountPool Instance
+		{
 			get
 			{
 				if (_Instance == null)
 				{
-					new AccountPool();
+					Initialize();
 				}
 				return _Instance;
 			}
@@ -74,17 +63,64 @@
 			}
 		}
 
+	}
+
+	public interface IManagedAccountPool: IDisposable
+	{
+		void RegisterAccount(Account InAccount);
+
+		void ReleaseAccount(Account InAccount);
+
+		Account ReserveAccount();
+	}
+
+	/// <summary>
+	/// A pool of accounts for tests to use. The pool should be filled by the project-specific script that derives from
+	/// RunUnrealTests
+	/// </summary>
+	public class ManagedAccountPool : IManagedAccountPool
+	{
+
+		/// <summary>
+		/// Object used for locking access to internal data
+		/// </summary>
+		protected Object LockObject = new Object();
+
+		/// <summary>
+		/// List of all registered accounts
+		/// </summary>
+		protected List<Account> AllAccounts = new List<Account>();
+
+		/// <summary>
+		/// Accounts that have been reserved
+		/// </summary>
+		protected Dictionary<Account, bool> ReservedAccounts = new Dictionary<Account, bool>();
+
+		/// <summary>
+		/// A random ordering of AllAccounts
+		/// </summary>
+		protected List<Account> RandomlyOrderedList = new List<Account>();
+
+
+		~ManagedAccountPool()
+		{
+			Dispose(false);
+		}
+
+
+
 		public void Dispose()
 		{
 			Dispose(true);
 			GC.SuppressFinalize(this);
 		}
 
+
 		/// <summary>
 		/// Perform actual dispose behavior
 		/// </summary>
 		/// <param name="disposing"></param>
-		protected virtual void Dispose(bool disposing)
+		protected void Dispose(bool disposing)
 		{
 			lock (LockObject)
 			{
@@ -123,7 +159,7 @@
 			}
 		}
 
-		public Account ReserveAccount()
+		virtual public Account ReserveAccount()
 		{
 			lock (LockObject)
 			{
diff --git a/Engine/Source/Programs/AutomationTool/Gauntlet/SelfTest/Unreal/Gauntlet.SelfTest.TestUnrealSession.cs b/Engine/Source/Programs/AutomationTool/Gauntlet/SelfTest/Unreal/Gauntlet.SelfTest.TestUnrealSession.cs
--- a/Engine/Source/Programs/AutomationTool/Gauntlet/SelfTest/Unreal/Gauntlet.SelfTest.TestUnrealSession.cs
+++ b/Engine/Source/Programs/AutomationTool/Gauntlet/SelfTest/Unreal/Gauntlet.SelfTest.TestUnrealSession.cs
@@ -22,6 +22,8 @@ namespace Gauntlet.SelfTest
 	{
 		public override void TickTest()
 		{
+			AccountPool.Initialize();
+
 			AccountPool.Instance.RegisterAccount(new EpicAccount("Foo", "Bar"));
 
 			// Add three devices to the pool
diff --git a/Engine/Source/Programs/AutomationTool/Gauntlet/Unreal/Game/EpicGame.TestConfig.cs b/Engine/Source/Programs/AutomationTool/Gauntlet/Unreal/Game/EpicGame.TestConfig.cs
--- a/Engine/Source/Programs/AutomationTool/Gauntlet/Unreal/Game/EpicGame.TestConfig.cs
+++ b/Engine/Source/Programs/AutomationTool/Gauntlet/Unreal/Game/EpicGame.TestConfig.cs
@@ -175,13 +175,20 @@ namespace EpicGame
 					// if this is a client, and there is a server role, find our PC's IP address and tell it to connect to us
 					if (AppConfig.ProcessType.IsClient() &&
 							(RequiredRoles.ContainsKey(UnrealTargetRole.Server) || RequiredRoles.ContainsKey(UnrealTargetRole.EditorServer)))
-					{						
+					{
 						McpString += string.Format(" -ExecCmds=\"open {0}:{1}\"", ServerIP, ServerPort);
-					}					
+					}
 				}
 				else
 				{
-					McpString += string.Format(" -epicapp={0} -buildidoverride={1}", EpicApp, BuildIDOverride);
+					if (Globals.Params.ParseParam("nobuildid"))
+					{
+						McpString += string.Format(" -epicapp={0} ", EpicApp);
+					}
+					else
+					{
+						McpString += string.Format(" -epicapp={0} -buildidoverride={1}", EpicApp, BuildIDOverride);
+					}
 				}
 
 				if (FastCook)
diff --git a/Engine/Source/Programs/AutomationTool/Gauntlet/Unreal/RunUnreal.cs b/Engine/Source/Programs/AutomationTool/Gauntlet/Unreal/RunUnreal.cs
--- a/Engine/Source/Programs/AutomationTool/Gauntlet/Unreal/RunUnreal.cs
+++ b/Engine/Source/Programs/AutomationTool/Gauntlet/Unreal/RunUnreal.cs
@@ -377,6 +377,9 @@ namespace Gauntlet
 		/// <returns></returns>
 		protected virtual void SetupAccounts()
 		{
+			// Set up account manager before we set up accounts.
+			AccountPool.Initialize();
+
 			string Username = Globals.Params.ParseValue("username", null);
 			string Password = Globals.Params.ParseValue("password", null);
 
