Adding support for shared iterative cooks from manifests.
Fall back path to loose staged builds.
Fall back path to existing synced builds.

diff --git a/Engine/Source/Programs/AutomationTool/Scripts/CookCommand.Automation.cs b/Engine/Source/Programs/AutomationTool/Scripts/CookCommand.Automation.cs
--- a/Engine/Source/Programs/AutomationTool/Scripts/CookCommand.Automation.cs
+++ b/Engine/Source/Programs/AutomationTool/Scripts/CookCommand.Automation.cs
@@ -247,8 +247,6 @@ public partial class Project : CommandUtils
 				}
 
                 CookCommandlet(Params.RawProjectPath, Params.UE4Exe, Maps, Dirs, InternationalizationPreset, CulturesToCook, CombineCommandletParams(PlatformsToCook.ToArray()), CommandletParams);
-
-				SharedCookedBuild.WaitForCopy();
             }
 			catch (Exception Ex)
 			{
diff --git a/Engine/Source/Programs/AutomationTool/Scripts/CopySharedCookedBuild.Automation.cs b/Engine/Source/Programs/AutomationTool/Scripts/CopySharedCookedBuild.Automation.cs
--- a/Engine/Source/Programs/AutomationTool/Scripts/CopySharedCookedBuild.Automation.cs
+++ b/Engine/Source/Programs/AutomationTool/Scripts/CopySharedCookedBuild.Automation.cs
@@ -68,7 +68,9 @@ class CopySharedCookedBuild : BuildCommand
 
 		bool bOnlyCopyAssetRegistry = ParseParam("onlycopyassetregistry");
 
-		string SharedBuildCL = ParseParamValue("buildcl", "usesyncedbuild");
+		string SharedBuildCL = ParseParamValue("buildcl", "any");
+
+		SharedCookedBuild.SharedCookType BuildType = (SharedCookedBuild.SharedCookType)Enum.Parse(typeof(SharedCookedBuild.SharedCookType), SharedBuildCL, true);
 
 		List<UnrealTargetPlatform> TargetPlatforms = new List<UnrealTargetPlatform>();
 		var PlatformNames = new List<string>(CmdLinePlatform.Split('+'));
@@ -90,8 +92,7 @@ class CopySharedCookedBuild : BuildCommand
 
 		foreach (var PlatformType in TargetPlatforms)
 		{
-			SharedCookedBuild.CopySharedCookedBuildForTarget(ProjectFile.FullName, PlatformType, PlatformType.ToString(), SharedBuildCL, bOnlyCopyAssetRegistry);
-			SharedCookedBuild.WaitForCopy();
+			SharedCookedBuild.CopySharedCookedBuild(ProjectFile.FullName, PlatformType, BuildType, true);
 		}
 
 		
diff --git a/Engine/Source/Programs/AutomationTool/Scripts/SharedCookedBuild.cs b/Engine/Source/Programs/AutomationTool/Scripts/SharedCookedBuild.cs
--- a/Engine/Source/Programs/AutomationTool/Scripts/SharedCookedBuild.cs
+++ b/Engine/Source/Programs/AutomationTool/Scripts/SharedCookedBuild.cs
@@ -13,313 +13,334 @@ using System.Text.RegularExpressions;
 
 public class SharedCookedBuild
 {
-	private static Task CopySharedCookedBuildTask = null;
+	private const string SyncedBuildFileName = "SyncedBuild.txt";
 
+	/// <summary>
+	/// Types of shared cook base builds
+	/// </summary>
 	public enum SharedCookType
 	{
+		/// <summary>
+		/// Only allow shared cook build of version identical to local sync
+		/// </summary>
 		Exact,
-		RecentContentOnly,
-		Recent
-	}
 
-	public static bool FindBestSharedCookedBuild(ref string FinalCookedBuildPath, string ProjectFullPath, UnrealTargetPlatform TargetPlatform, string CookPlatform, SharedCookType BuildType)
+		/// <summary>
+		/// Allow any previous version that is only a content change from local sync
+		/// </summary>
+		Content,
+		
+		/// <summary>
+		/// Closest previous version, regardless of code/content changes
+		/// </summary>
+		Any,
+	}
+		
+	public static void CopySharedCookedBuild(ProjectParams Params)
 	{
-		string BuildRoot = CommandUtils.P4Enabled ? CommandUtils.P4Env.Branch.Replace("/", "+") : "";
-		int CurrentCLInt = CommandUtils.P4Enabled ? CommandUtils.P4Env.Changelist : 0;
-		int CurrentCodeCLInt = CommandUtils.P4Enabled ? CommandUtils.P4Env.CodeChangelist : 0;
+		foreach (TargetPlatformDescriptor ClientPlatform in Params.ClientTargetPlatforms)
+		{
+			TargetPlatformDescriptor DataPlatformDesc = Params.GetCookedDataPlatformForClientTarget(ClientPlatform);
+			string PlatformToCook = Platform.Platforms[DataPlatformDesc].GetCookPlatform(false, Params.Client);
+			UnrealTargetPlatform TargetPlatform = (UnrealTargetPlatform)Enum.Parse(typeof(UnrealTargetPlatform), PlatformToCook, true);
+			DirectoryReference InstallPath = DirectoryReference.Combine(Params.RawProjectPath.Directory, "Saved", "SharedIterativeBuild", PlatformToCook);
+			SharedCookType BuildType = (SharedCookType)Enum.Parse(typeof(SharedCookType), Params.IterateSharedCookedBuild, true);
 
+			CopySharedCookedBuild(Params.RawProjectPath.FullName, TargetPlatform, BuildType, true);
+		}
+	}
 
-		BuildVersion Version;
-		if (BuildVersion.TryRead(BuildVersion.GetDefaultFileName(), out Version))
+	public static void CopySharedCookedBuild(string ProjectFullPath, UnrealTargetPlatform TargetPlatform, SharedCookType BuildType, bool bAllowExistingBuild = true)
+	{
+		DirectoryReference InstallPath = DirectoryReference.Combine(new FileReference(ProjectFullPath).Directory, "Saved", "SharedIterativeBuild", TargetPlatform.ToString());
+		List<ISharedCookedBuild> SharedCookedBuilds = FindBestBuilds(ProjectFullPath, TargetPlatform, BuildType, true);
+		foreach (ISharedCookedBuild Build in SharedCookedBuilds)
 		{
-			CurrentCLInt = Version.Changelist;
-			CurrentCodeCLInt = Version.EffectiveCompatibleChangelist;
-			BuildRoot = Version.BranchName;
+			if (Build.CopyBuild(InstallPath))
+			{
+				return;
+			}
 		}
-		System.GC.Collect();
-		string CurrentCL = CurrentCLInt.ToString();
 
-		CommandUtils.LogInformation("Search for best shared cooked build...");
-		CommandUtils.LogInformation("CurrentCL: {0}, Type: {1}", CurrentCodeCLInt, BuildType);
+		throw new AutomationException("Failed to install shared cooked build");
+	}
 
-		FileReference ProjectFileRef = new FileReference(ProjectFullPath);
-		// get network location 
-		ConfigHierarchy Hierarchy = ConfigCache.ReadHierarchy(ConfigHierarchyType.Engine, DirectoryReference.FromFile(ProjectFileRef), TargetPlatform);
-		List<string> CookedBuildPaths;
-		if (Hierarchy.GetArray("SharedCookedBuildSettings", "SharedCookedBuildPath", out CookedBuildPaths) == false)
+	public static BuildVersion LocalSync()
+	{
+		BuildVersion P4Version = new BuildVersion();
+		if (CommandUtils.P4Enabled)
 		{
-			CommandUtils.LogInformation("Unable to copy shared cooked build: SharedCookedBuildPath not set in Engine.ini SharedCookedBuildSettings");
-			return false;
+			P4Version.BranchName = CommandUtils.P4Env.Branch.Replace("/", "+");
+			P4Version.Changelist = CommandUtils.P4Env.Changelist;
+			P4Version.CompatibleChangelist = CommandUtils.P4Env.CodeChangelist;
 		}
 
-		const string MetaDataFilename = "\\Metadata\\DevelopmentAssetRegistry.bin";
-
-
-		if (BuildType == SharedCookType.Exact)
+		BuildVersion UGSVersion;
+		if (BuildVersion.TryRead(BuildVersion.GetDefaultFileName(), out UGSVersion))
 		{
-			foreach (string CookedBuildPath in CookedBuildPaths)
-			{
-				if (CurrentCL == "" && FinalCookedBuildPath.Contains("[CL]"))
-				{
-					CommandUtils.LogInformation("Unable to copy shared cooked build: Unable to determine CL number from P4 or UGS, and is required by SharedCookedBuildPath");
-					return false;
-				}
-
-				if (CurrentCL == "" && FinalCookedBuildPath.Contains("[BRANCHNAME]"))
-				{
-					CommandUtils.LogInformation("Unable to copy shared cooked build: Unable to determine BRANCHNAME number from P4 or UGS, and is required by SharedCookedBuildPath");
-					return false;
-				}
-
-
-				FinalCookedBuildPath = FinalCookedBuildPath.Replace("[CL]", CurrentCL.ToString());
-				FinalCookedBuildPath = FinalCookedBuildPath.Replace("[BRANCHNAME]", BuildRoot);
-				FinalCookedBuildPath = FinalCookedBuildPath.Replace("[PLATFORM]", CookPlatform);
-
-				// make sure that the directory and metadata file exist.  otherwise this build might not be finished yet and we should skip it
-				if (Directory.Exists(FinalCookedBuildPath))
-				{
-					if ( File.Exists( FinalCookedBuildPath + MetaDataFilename) )
-					{
-						return true;
-					}
-				}
-			}
+			return UGSVersion;
 		}
-		else if (BuildType == SharedCookType.Recent || BuildType == SharedCookType.RecentContentOnly)
+		
+		if (!CommandUtils.P4Enabled)
 		{
-			// build our CookedBUildPath into a regex which we can execute on the directories and extract info from
-
-			string BestBuild = null;
-			int BestCLNumber = 0;
-
-			// find all the recent builds which are valid
-			foreach (string CookedBuildPath in CookedBuildPaths)
-			{
-				int IndexOfFirstParam = CookedBuildPath.IndexOf("[");
-				int CustomFolderStart = CookedBuildPath.LastIndexOf("\\", IndexOfFirstParam);
-
-				string CookedBuildDirectory = CookedBuildPath.Substring(0, CustomFolderStart);
-
-				string BuildNameWildcard = CookedBuildPath.Substring(CustomFolderStart);
-
-
-				BuildNameWildcard += MetaDataFilename;
-
-				FileFilter BuildSearch = new FileFilter();
-
-				// we know the platform and the branch name;
-				string BuildRule = BuildNameWildcard;
-				BuildRule = BuildRule.Replace("[BRANCHNAME]", BuildRoot);
-				BuildRule = BuildRule.Replace("[PLATFORM]", CookPlatform);
-
-				string IncludeRule = BuildRule.Replace("[CL]", "*");
-				string ForgetRule = BuildRule.Replace("[CL]", "*-PF-*"); // get rid of any preflights from the list... they don't count because who knows what they did...
-
-				BuildSearch.AddRule(IncludeRule);
-				BuildSearch.AddRule(ForgetRule, FileFilterType.Exclude);
-
-				List<FileReference> ValidBuilds = BuildSearch.ApplyToDirectory(new DirectoryReference(CookedBuildDirectory), false);
-
-				// figure out what the CL is
-				string BuildNameRegex = String.Format(".*{0}", CookedBuildPath.Substring(CustomFolderStart));
-				BuildNameRegex = BuildNameRegex.Replace("\\", "\\\\");
-				BuildNameRegex = BuildNameRegex.Replace("[BRANCHNAME]", BuildRoot);
-				BuildNameRegex = BuildNameRegex.Replace("+", "\\+");
-				BuildNameRegex = BuildNameRegex.Replace("[PLATFORM]", CookPlatform);
-				BuildNameRegex = BuildNameRegex.Replace("[CL]", "(?<CL>.*)");
-				
-				Regex ExtractCL = new Regex(BuildNameRegex);
-
-				foreach ( FileReference ValidBuild in ValidBuilds )
-				{
-					string BuildPath = ValidBuild.FullName.Replace(MetaDataFilename, "");
-
-					Match CLMatch = ExtractCL.Match(BuildPath);
-					if ( CLMatch != null )
-					{
-						string CLNumber = CLMatch.Result("${CL}");
-						int CLNumberInt = int.Parse(CLNumber);
-						if ( CLNumberInt <= CurrentCLInt )
-						{
-							if ( CLNumberInt > BestCLNumber )
-							{
-								BestCLNumber = CLNumberInt;
-								BestBuild = BuildPath;
-							}
-						}
-					}
-				}
-			}
-
-			if ( string.IsNullOrEmpty(BestBuild) )
-			{
-				CommandUtils.LogError("Cannot find a recent shared cooked build");
-				return false;
-			}
-
-			if (BuildType == SharedCookType.RecentContentOnly && BestCLNumber < CurrentCodeCLInt)
-			{
-				CommandUtils.LogError("Cannot find a recent shared cooked build without code changes");
-				CommandUtils.LogError("Your CL: {0}, Last code CL: {1}, Best fit: {2}", CurrentCLInt, CurrentCodeCLInt, BestCLNumber);
-				return false;
-			}
-
-			FinalCookedBuildPath = BestBuild;
-			CommandUtils.LogInformation("Selected build: {0}", BestCLNumber);
-			return true;
+			throw new AutomationException("Cannot determine local sync");
 		}
 
-
-		return false;
+		return P4Version;
 	}
 
-	private static bool CopySharedCookedBuildForTargetInternal(string CookedBuildPath, string CookPlatform, string LocalPath, bool bOnlyCopyAssetRegistry)
+	public static List<ISharedCookedBuild> FindBestBuilds(string ProjectFullPath, UnrealTargetPlatform TargetPlatform, SharedCookType BuildType, bool bAllowExistingBuild = true)
 	{
-
-		// check to see if we have already synced this build ;)
-		var SyncedBuildFile = CommandUtils.CombinePaths(LocalPath, "SyncedBuild.txt");
-		string BuildCL = "Invalid";
-		if (File.Exists(SyncedBuildFile))
+		// Attempt manifest searching first
+		FileReference ProjectFileRef = new FileReference(ProjectFullPath);
+		if (!FileReference.Exists(ProjectFileRef))
 		{
-			BuildCL = File.ReadAllText(SyncedBuildFile);
+			throw new AutomationException("Cannot locate project file: {0}", ProjectFileRef);
 		}
 
-		CommandUtils.LogInformation("Attempting download of latest shared build from {0}", CookedBuildPath);
+		ConfigHierarchy Herarchy = ConfigCache.ReadHierarchy(ConfigHierarchyType.Engine, DirectoryReference.FromFile(ProjectFileRef), TargetPlatform);
+		List<string> CookedBuildManifestPaths = null;
+		List<string> CookedBuildStagedPaths = null;
+		Herarchy.GetArray("SharedCookedBuildSettings", "SharedCookedManifestPath", out CookedBuildManifestPaths);
+		Herarchy.GetArray("SharedCookedBuildSettings", "SharedCookedBuildPath", out CookedBuildStagedPaths);
+		if (CookedBuildManifestPaths == null && CookedBuildStagedPaths == null)
+		{
+			throw new AutomationException("Unable to locate shared cooked builds. SharedCookedManifestPath and SharedCookedBuildPath not set in Engine.ini [SharedCookedBuildSettings]");
+		}
 
-		string SavedBuildCL = string.Format( "{0} {1}", CookedBuildPath, bOnlyCopyAssetRegistry ? "RegistryOnly" : "" );
+		BuildVersion Version = LocalSync();
+		List<ISharedCookedBuild> CandidateBuilds = new List<ISharedCookedBuild>();
 
-		if (BuildCL == SavedBuildCL)
+		if (bAllowExistingBuild)
 		{
-			CommandUtils.LogInformation("Already downloaded latest shared build at CL {0}", SavedBuildCL);
-			return false;
+			// If existing sync is present, stick to it. Read version out of sync file
+			FileReference SyncedBuildFile = new FileReference(CommandUtils.CombinePaths(Path.GetDirectoryName(ProjectFullPath), "Saved", "SharedIterativeBuild", TargetPlatform.ToString(), SyncedBuildFileName));
+			if (FileReference.Exists(SyncedBuildFile))
+			{
+				string[] SyncedBuildInfo = FileReference.ReadAllLines(SyncedBuildFile);
+				int SyncedCL = int.Parse(SyncedBuildInfo[0]);
+				if (IsValidCL(SyncedCL, BuildType, Version))
+				{
+					CandidateBuilds.Add(new ExistingSharedCookedBuild { CL = SyncedCL });
+				}
+			}
 		}
 
-		if (Directory.Exists(CookedBuildPath) == false)
+		if (CookedBuildManifestPaths.Count > 0)
 		{
-			CommandUtils.LogInformation("Unable to copy shared cooked build: Unable to find shared build at location {0} check SharedCookedBuildPath in Engine.ini SharedCookedBuildSettings is correct", CookedBuildPath);
-			return false;
+			foreach (string ManifestPath in CookedBuildManifestPaths)
+			{
+				ISharedCookedBuild Candidate = FindBestManifestBuild(ManifestPath, TargetPlatform, BuildType, Version);
+				if (Candidate != null)
+				{
+					CandidateBuilds.Add(Candidate);
+				}
+			}
 		}
-		
-		// delete all the stuff
-		CommandUtils.LogInformation("Deleting previous shared build because it was out of date");
-		CommandUtils.DeleteDirectory(LocalPath);
-		Directory.CreateDirectory(LocalPath);
-
-
 
-		string CookedBuildMetadataDirectory = Path.Combine(CookedBuildPath, "Metadata");
-		CookedBuildMetadataDirectory = Path.GetFullPath(CookedBuildMetadataDirectory);
-		string LocalBuildMetadataDirectory = Path.Combine(LocalPath, "Metadata");
-		LocalBuildMetadataDirectory = Path.GetFullPath(LocalBuildMetadataDirectory);
-		if (Directory.Exists(CookedBuildMetadataDirectory))
+		if (CookedBuildStagedPaths.Count > 0)
 		{
-			foreach (string FileName in Directory.EnumerateFiles(CookedBuildMetadataDirectory, "*.*", SearchOption.AllDirectories))
+			foreach (string StagedPath in CookedBuildStagedPaths)
 			{
-				string SourceFileName = Path.GetFullPath(FileName);
-				string DestFileName = SourceFileName.Replace(CookedBuildMetadataDirectory, LocalBuildMetadataDirectory);
-				Directory.CreateDirectory(Path.GetDirectoryName(DestFileName));
-				File.Copy(SourceFileName, DestFileName);
+				ISharedCookedBuild Candidate = FindBestLooseBuild(StagedPath, TargetPlatform, BuildType, Version);
+				if (Candidate != null)
+				{
+					CandidateBuilds.Add(Candidate);
+				}
 			}
 		}
 
-		if ( CopySharedCookedBuildTask != null )
+		if (CandidateBuilds.Count == 0)
 		{
-			WaitForCopy();
+			CommandUtils.LogInformation("Could not locate valid shared cooked build");
 		}
 
-		if (bOnlyCopyAssetRegistry == false)
+		return CandidateBuilds.OrderBy(x => x.CL).ToList();
+	}
+
+	public static ISharedCookedBuild FindBestManifestBuild(string Path, UnrealTargetPlatform TargetPlatform, SharedCookType BuildType, BuildVersion Version)
+	{
+		Tuple<string, string> SplitPath = SplitOnFixedPrefix(Path);
+		Regex Pattern = RegexFromWildcards(SplitPath.Item2, Version, TargetPlatform);
+
+		ManifestSharedCookedBuild Build = new ManifestSharedCookedBuild { CL = 0, Manifest = null };
+		DirectoryReference SearchDir = new DirectoryReference(SplitPath.Item1);
+		if (DirectoryReference.Exists(SearchDir))
 		{
-			CopySharedCookedBuildTask = Task.Run(() =>
+			foreach (FileReference File in DirectoryReference.EnumerateFiles(SearchDir))
+			{
+				Match Match = Pattern.Match(File.FullName);
+				if (Match.Success)
 				{
-					// find all the files in the staged directory
-					string CookedBuildStagedDirectory = Path.GetFullPath(Path.Combine(CookedBuildPath, "Staged"));
-					string LocalBuildStagedDirectory = Path.GetFullPath(Path.Combine(LocalPath, "Staged"));
-					if (Directory.Exists(CookedBuildStagedDirectory))
+					int MatchCL = int.Parse(Match.Result("${CL}"));
+					if (IsValidCL(MatchCL, BuildType, Version) && MatchCL >= Build.CL)
 					{
-						foreach (string FileName in Directory.EnumerateFiles(CookedBuildStagedDirectory, "*.*", SearchOption.AllDirectories))
-						{
-							string SourceFileName = Path.GetFullPath(FileName);
-							string DestFileName = SourceFileName.Replace(CookedBuildStagedDirectory, LocalBuildStagedDirectory);
-							Directory.CreateDirectory(Path.GetDirectoryName(DestFileName));
-							File.Copy(SourceFileName, DestFileName);
-						}
+						Build = new ManifestSharedCookedBuild { CL = MatchCL, Manifest = File };
 					}
-					File.WriteAllText(SyncedBuildFile, SavedBuildCL);
 				}
-				);
+			}
 		}
-		else
+
+		if (Build.CL != 0)
 		{
-			File.WriteAllText(SyncedBuildFile, SavedBuildCL);
+			return Build;
 		}
 
-		
-		return true;
+		return null;
 	}
-
-	public static void CopySharedCookedBuildForTarget(string ProjectFullPath, UnrealTargetPlatform TargetPlatform, string CookPlatform, string BuildCL, bool bOnlyCopyAssetRegistry = false)
+	public static ISharedCookedBuild FindBestLooseBuild(string Path, UnrealTargetPlatform TargetPlatform, SharedCookType BuildType, BuildVersion Version)
 	{
-		var LocalPath = CommandUtils.CombinePaths(Path.GetDirectoryName(ProjectFullPath), "Saved", "SharedIterativeBuild", CookPlatform);
+		Tuple<string, string> SplitPath = SplitOnFixedPrefix(Path);
+		Regex Pattern = RegexFromWildcards(SplitPath.Item2, Version, TargetPlatform);
+		LooseSharedCookedBuild Build = new LooseSharedCookedBuild { CL = 0, Path = null };
 
-		SharedCookType SharedBuildType = SharedCookType.Exact;
-		if (BuildCL.Equals("userecentbuild"))
+		// Search for all available builds
+		const string MetaDataFilename = "\\Metadata\\DevelopmentAssetRegistry.bin";
+		string BuildRule = Path + MetaDataFilename;
+		BuildRule = BuildRule.Replace("[BRANCHNAME]", Version.BranchName);
+		BuildRule = BuildRule.Replace("[PLATFORM]", TargetPlatform.ToString());
+		string IncludeRule = BuildRule.Replace("[CL]", "*");
+		string ExcludeRule = BuildRule.Replace("[CL]", "*-PF-*"); // Exclude preflights
+		FileFilter BuildSearch = new FileFilter();
+		BuildSearch.AddRule(IncludeRule);
+		BuildSearch.AddRule(ExcludeRule, FileFilterType.Exclude);
+		
+		foreach (FileReference CandidateBuild in BuildSearch.ApplyToDirectory(new DirectoryReference(SplitPath.Item1), false))
 		{
-			SharedBuildType = SharedCookType.Recent;
+			string BaseBuildPath = CandidateBuild.FullName.Replace(MetaDataFilename, "");
+			Match Match = Pattern.Match(BaseBuildPath);
+			if (Match.Success)
+			{
+				int MatchCL = int.Parse(Match.Result("${CL}"));
+				if (IsValidCL(MatchCL, BuildType, Version) && MatchCL >= Build.CL)
+				{
+					Build = new LooseSharedCookedBuild { CL = MatchCL, Path = new DirectoryReference(BaseBuildPath) };
+				}
+			}
 		}
-		else if (BuildCL.Equals("userecentcontentonly"))
+
+		if (Build.CL != 0)
 		{
-			SharedBuildType = SharedCookType.RecentContentOnly;
+			return Build;
 		}
 
-		string CookedBuildPath = null;
-		if ( FindBestSharedCookedBuild(ref CookedBuildPath, ProjectFullPath, TargetPlatform, CookPlatform, SharedBuildType) )
+		return null;
+	}
+
+
+	private static bool IsValidCL(int CL, SharedCookType BuildType, BuildVersion Version)
+	{
+		if (BuildType == SharedCookType.Exact && CL == Version.Changelist)
 		{
-			CopySharedCookedBuildForTargetInternal(CookedBuildPath, CookPlatform, LocalPath, bOnlyCopyAssetRegistry);
+			return true;
 		}
-		else if( FindBestSharedCookedBuild(ref CookedBuildPath, ProjectFullPath, TargetPlatform, TargetPlatform.ToString(), SharedBuildType) )
+		else if (BuildType == SharedCookType.Content && CL >= Version.EffectiveCompatibleChangelist && CL <= Version.Changelist)
 		{
-			CopySharedCookedBuildForTargetInternal(CookedBuildPath, TargetPlatform.ToString(), LocalPath, bOnlyCopyAssetRegistry);
+			return true;
 		}
-		else
+		else if (BuildType == SharedCookType.Any && CL <= Version.Changelist)
 		{
-			throw new AutomationException("Can't find cooked build for {0} {1} {2}", SharedBuildType,  CookPlatform, TargetPlatform.ToString());
+			return true;
 		}
 
-		return;
+		return false;
+	}
+
+	private static Regex RegexFromWildcards(string Path, BuildVersion Version, UnrealTargetPlatform TargetPlatform)
+	{
+		string Pattern = Path.Replace(@"\", @"\\");
+		Pattern = Pattern.Replace("[BRANCHNAME]", Version.BranchName.Replace(@"+", @"\+"));
+		Pattern = Pattern.Replace("[PLATFORM]", TargetPlatform.ToString());
+		Pattern = Pattern.Replace("[CL]", @"(?<CL>\d+)");
+		return new Regex(Pattern);
 	}
 
-	public static void CopySharedCookedBuild(ProjectParams Params)
+	private static Tuple<string, string> SplitOnFixedPrefix(string Path)
+	{
+		int IndexOfFirstParam = Path.IndexOf("[");
+		int PrefixStart = Path.LastIndexOf(@"\", IndexOfFirstParam);
+		return new Tuple<string, string>(Path.Substring(0, PrefixStart), Path.Substring(PrefixStart));
+	}
+
+	public interface ISharedCookedBuild
+	{
+		int CL { get; }
+		bool CopyBuild(DirectoryReference InstallPath);
+	}
+
+	private class ManifestSharedCookedBuild : ISharedCookedBuild
 	{
+		public int CL { get; set; }
+		public FileReference Manifest { get; set; }
 
-		if (!Params.NoClient)
+		public bool CopyBuild(DirectoryReference InstallPath)
 		{
-			foreach (var ClientPlatform in Params.ClientTargetPlatforms)
+			CommandUtils.LogInformation("Installing shared cooked build from manifest: {0} to {1}", Manifest.FullName, InstallPath.FullName);
+
+			FileReference PreviousManifest = FileReference.Combine(InstallPath, ".build", "Current.manifest");
+
+			FileReference BPTI = FileReference.Combine(CommandUtils.RootDirectory, "Engine", "Binaries", "Win64", "NotForLicensees", "BuildPatchToolInstaller.exe");
+			if (!FileReference.Exists(BPTI))
 			{
-				// Use the data platform, sometimes we will copy another platform's data
-				var DataPlatformDesc = Params.GetCookedDataPlatformForClientTarget(ClientPlatform);
-				string PlatformToCook = Platform.Platforms[DataPlatformDesc].GetCookPlatform(false, Params.Client);
-				CopySharedCookedBuildForTarget(Params.RawProjectPath.FullName, ClientPlatform.Type, PlatformToCook, Params.IterateSharedCookedBuild);
+				CommandUtils.LogInformation("Could not locate BuildPatchToolInstaller.exe");
+				return false;
 			}
+
+			bool PreviousManifestExists = FileReference.Exists(PreviousManifest);
+			if (!PreviousManifestExists && DirectoryReference.Exists(InstallPath))
+			{
+				DirectoryReference.Delete(InstallPath, true);
+			}
+
+			IProcessResult Result = CommandUtils.Run(BPTI.FullName, string.Format("-Manifest={0} -OutputDir={1}", Manifest.FullName, InstallPath.FullName), null, CommandUtils.ERunOptions.AllowSpew);
+			if (Result.ExitCode != 0)
+			{
+				CommandUtils.LogWarning("Failed to install manifest {0} to {1}", Manifest.FullName, InstallPath.FullName);
+				return false;
+			}
+
+			FileReference SyncedBuildFile = new FileReference(CommandUtils.CombinePaths(InstallPath.FullName, SyncedBuildFileName));
+			FileReference.WriteAllLines(SyncedBuildFile, new string[] { CL.ToString(), Manifest.FullName });
+
+			return true;
 		}
-		if (Params.DedicatedServer)
+	}
+
+	private class LooseSharedCookedBuild : ISharedCookedBuild
+	{
+		public int CL { get; set; }
+		public DirectoryReference Path { get; set; }
+		public bool CopyBuild(DirectoryReference InstallPath)
 		{
-			foreach (var ServerPlatform in Params.ServerTargetPlatforms)
+			CommandUtils.LogInformation("Copying shared cooked build from stage directory: {0} to {1}", Path.FullName, InstallPath.FullName);
+
+			// Delete existing
+			if (DirectoryReference.Exists(InstallPath))
 			{
-				// Use the data platform, sometimes we will copy another platform's data
-				var DataPlatformDesc = Params.GetCookedDataPlatformForServerTarget(ServerPlatform);
-				string PlatformToCook = Platform.Platforms[DataPlatformDesc].GetCookPlatform(true, false);
-				CopySharedCookedBuildForTarget(Params.RawProjectPath.FullName, ServerPlatform.Type, PlatformToCook, Params.IterateSharedCookedBuild);
+				DirectoryReference.Delete(InstallPath, true);
 			}
+			DirectoryReference.CreateDirectory(InstallPath);
+
+			// Copy new
+			if (!CommandUtils.CopyDirectory_NoExceptions(Path.FullName, InstallPath.FullName))
+			{
+				CommandUtils.LogWarning("Failed to copy {0} -> {1}", Path.FullName, InstallPath.FullName);
+				return false;
+			}
+			FileReference SyncedBuildFile = new FileReference(CommandUtils.CombinePaths(InstallPath.FullName, SyncedBuildFileName));
+			FileReference.WriteAllLines(SyncedBuildFile, new string[] { CL.ToString(), Path.FullName });
+
+			return true;
 		}
 	}
 
-	public static void WaitForCopy()
+	private class ExistingSharedCookedBuild : ISharedCookedBuild
 	{
-		if (CopySharedCookedBuildTask != null)
+		public int CL { get; set; }
+		public bool CopyBuild(DirectoryReference InstallPath)
 		{
-			CopySharedCookedBuildTask.Wait();
+			CommandUtils.LogInformation("Using previously synced shared cooked build");
+			return true;
 		}
 	}
-
 }
\ No newline at end of file
