Fix exception in clean devices

diff --git a/Engine/Source/Programs/AutomationTool/Gauntlet/Unreal/Tasks/Gauntlet.Task.CleanDevices.cs b/Engine/Source/Programs/AutomationTool/Gauntlet/Unreal/Tasks/Gauntlet.Task.CleanDevices.cs
--- a/Engine/Source/Programs/AutomationTool/Gauntlet/Unreal/Tasks/Gauntlet.Task.CleanDevices.cs
+++ b/Engine/Source/Programs/AutomationTool/Gauntlet/Unreal/Tasks/Gauntlet.Task.CleanDevices.cs
@@ -18,7 +18,10 @@ namespace Gauntlet
 	{
 
 		[AutoParamWithNames("", "device", "devices")]
-		public string Devices;
+		public string Devices = "";
+
+		[AutoParam("")]
+		public string DeviceURL;
 
 		[AutoParam("")]
 		public string TempDir;
@@ -39,6 +42,7 @@ namespace Gauntlet
 			}
 
 			// add devices. We're quick so can ignore constraints
+			DevicePool.Instance.SetLocalOptions(TempDir, false, DeviceURL);
 			DevicePool.Instance.AddDevices(UnrealTargetPlatform.Win64, Devices, false);
 
 			UnrealTargetPlatform[] SupportedPlatforms = { UnrealTargetPlatform.PS4, UnrealTargetPlatform.Win64 };
