Fixed issue with parallel tests of the same name trying to use the same artifact folder.

diff --git a/Engine/Source/Programs/AutomationTool/Gauntlet/Unreal/Base/Gauntlet.UnrealTestContext.cs b/Engine/Source/Programs/AutomationTool/Gauntlet/Unreal/Base/Gauntlet.UnrealTestContext.cs
--- a/Engine/Source/Programs/AutomationTool/Gauntlet/Unreal/Base/Gauntlet.UnrealTestContext.cs
+++ b/Engine/Source/Programs/AutomationTool/Gauntlet/Unreal/Base/Gauntlet.UnrealTestContext.cs
@@ -158,6 +158,18 @@ namespace Gauntlet
 		[AutoParam("")]
 		public string LogDir;
 
+		/// <summary>
+		/// Custom name for the directory under LogDir where artifacts are stored.
+		/// </summary>
+		[AutoParam("")]
+		public string ArtifactName;
+
+		/// <summary>
+		/// Psotfix to apply to the artifact folder (e.g. don't replace the whole thing like ArtifactName, just append this string)
+		/// </summary>
+		[AutoParam("")]
+		public string ArtifactPostfix;
+
 		/// <summary>
 		/// Less logging
 		/// </summary>
diff --git a/Engine/Source/Programs/AutomationTool/Gauntlet/Unreal/Base/Gauntlet.UnrealTestNode.cs b/Engine/Source/Programs/AutomationTool/Gauntlet/Unreal/Base/Gauntlet.UnrealTestNode.cs
--- a/Engine/Source/Programs/AutomationTool/Gauntlet/Unreal/Base/Gauntlet.UnrealTestNode.cs
+++ b/Engine/Source/Programs/AutomationTool/Gauntlet/Unreal/Base/Gauntlet.UnrealTestNode.cs
@@ -107,7 +107,9 @@ namespace Gauntlet
 
 		// End  UnrealTestNode properties and members 
 
-		// UnrealTestNode member functions
+		// artifact paths that have been used in this run
+		static protected HashSet<string> ReservedArtifcactPaths = new HashSet<string>();
+
 		public UnrealTestNode(UnrealTestContext InContext)
 		{
 			Context = InContext;
@@ -417,31 +419,56 @@ namespace Gauntlet
 			CurrentPass = Pass;
 			NumPasses = InNumPasses;
 
-			string TestFolder = ToString();
+			// Either use the ArtifactName param or name of this test
+			string TestFolder = string.IsNullOrEmpty(Context.Options.ArtifactName) ? this.ToString() : Context.Options.ArtifactName;
+
+			if (string.IsNullOrEmpty(Context.Options.ArtifactPostfix) == false)
+			{
+				TestFolder += "_" + Context.Options.ArtifactPostfix;
+			}
+
 			TestFolder = TestFolder.Replace(" ", "_");
 			TestFolder = TestFolder.Replace(",", "");
 
 			ArtifactPath = Path.Combine(Context.Options.LogDir, TestFolder);
-
+		
 			// if doing multiple passes, put each in a subdir
 			if (NumPasses > 1)
 			{
 				ArtifactPath = Path.Combine(ArtifactPath, string.Format("Pass_{0}_of_{1}", CurrentPass, NumPasses));
 			}
 
-			// Basic pre-existing directory check.
-			if (CommandUtils.IsBuildMachine && Directory.Exists(ArtifactPath))
+			// When running with -parallel we could have several identical tests (same test, configurations) in flight so
+			// we need unique artifact paths. We also don't overwrite dest directories from the build machine for the same
+			// reason of multiple tests for a build. Really though these should use ArtifactPrefix to save to
+			// SmokeTest_HighQuality etc
+			int ArtifactNumericPostfix = 0;
+			bool ArtifactPathIsTaken = false;
+
+			do
 			{
-				string NewOutputPath = ArtifactPath;
-				int i = 0;
-				while (Directory.Exists(NewOutputPath))
+				string PotentialPath = ArtifactPath;
+
+				if (ArtifactNumericPostfix > 0)
 				{
-					i++;
-					NewOutputPath = string.Format("{0}_{1}", ArtifactPath, i);
+					PotentialPath = string.Format("{0}_{1}", ArtifactPath, ArtifactNumericPostfix);
 				}
-				Log.Info("Directory already exists at {0}", ArtifactPath);
-				ArtifactPath = NewOutputPath;
-			}
+
+				ArtifactPathIsTaken = ReservedArtifcactPaths.Contains(PotentialPath) || (CommandUtils.IsBuildMachine && Directory.Exists(PotentialPath));
+
+				if (ArtifactPathIsTaken)
+				{
+					Log.Info("Directory already exists at {0}", PotentialPath);
+					ArtifactNumericPostfix++;
+				}
+				else
+				{
+					ArtifactPath = PotentialPath;
+				}
+
+			} while (ArtifactPathIsTaken);
+
+			ReservedArtifcactPaths.Add(ArtifactPath);
 
 			// Launch the test
 			TestInstance = UnrealApp.LaunchSession();
