Prevent packaged plugins getting a hot reload suffix if the editor is running.

diff --git a/Engine/Source/Programs/AutomationTool/Scripts/BuildPluginCommand.Automation.cs b/Engine/Source/Programs/AutomationTool/Scripts/BuildPluginCommand.Automation.cs
--- a/Engine/Source/Programs/AutomationTool/Scripts/BuildPluginCommand.Automation.cs
+++ b/Engine/Source/Programs/AutomationTool/Scripts/BuildPluginCommand.Automation.cs
@@ -197,7 +197,7 @@ class BuildPlugin : BuildCommand
 			FileReference ManifestFileName = FileReference.Combine(HostProjectFile.Directory, "Saved", String.Format("Manifest-{0}-{1}-{2}.xml", TargetName, Platform, Configuration));
 			ManifestFileNames.Add(ManifestFileName);
 			
-			string Arguments = String.Format("-plugin={0} -iwyu -noubtmakefiles -manifest={1}", CommandUtils.MakePathSafeToUseWithCommandLine(HostProjectPluginFile.FullName), CommandUtils.MakePathSafeToUseWithCommandLine(ManifestFileName.FullName));
+			string Arguments = String.Format("-plugin={0} -iwyu -noubtmakefiles -manifest={1} -nohotreload", CommandUtils.MakePathSafeToUseWithCommandLine(HostProjectPluginFile.FullName), CommandUtils.MakePathSafeToUseWithCommandLine(ManifestFileName.FullName));
 			if(!String.IsNullOrEmpty(InAdditionalArgs))
 			{
 				Arguments += InAdditionalArgs;
