UE-66627 - Full Screen Native Resolution not supported on new iPad Pro 11-inch and iPad Pro 12.9 inch (3rd generation)

diff --git a/Engine/Source/Programs/AutomationTool/IOS/IOSPlatform.Automation.cs b/Engine/Source/Programs/AutomationTool/IOS/IOSPlatform.Automation.cs
--- a/Engine/Source/Programs/AutomationTool/IOS/IOSPlatform.Automation.cs
+++ b/Engine/Source/Programs/AutomationTool/IOS/IOSPlatform.Automation.cs
@@ -915,16 +915,24 @@ public class IOSPlatform : Platform
 				ImageFileNames.Add("Default-IPhone6.png");
 				ImageFileNames.Add("Default-IPhone6Plus-Portrait.png");
 				ImageFileNames.Add("[REDACTED_EMAIL]");
+				ImageFileNames.Add("[REDACTED_EMAIL]");
+				ImageFileNames.Add("[REDACTED_EMAIL]");
 				ImageFileNames.Add("[REDACTED_EMAIL]");
 				ImageFileNames.Add("Default-IPhoneXS-Portrait.png");
+				ImageFileNames.Add("Default-IPhoneXSMax-Portrait.png");
+				ImageFileNames.Add("Default-IPhoneXR-Portrait.png");
 			}
 			if (bSupportsLandscape)
 			{
 				ImageFileNames.Add("Default-IPhone6-Landscape.png");
 				ImageFileNames.Add("Default-IPhone6Plus-Landscape.png");
 				ImageFileNames.Add("[REDACTED_EMAIL]");
+				ImageFileNames.Add("[REDACTED_EMAIL]");
+				ImageFileNames.Add("[REDACTED_EMAIL]");
 				ImageFileNames.Add("[REDACTED_EMAIL]");
 				ImageFileNames.Add("Default-IPhoneXS-Landscape.png");
+				ImageFileNames.Add("Default-IPhoneXSMax-Landscape.png");
+				ImageFileNames.Add("Default-IPhoneXR-Landscape.png");
 			}
 			ImageFileNames.Add("[REDACTED_EMAIL]");
 			ImageFileNames.Add("[REDACTED_EMAIL]");
