Fixes to ListThirdPartySoftware script.

* Switched to use JsonExport mode in UBT.
* Added runtime dependencies into the output file.

#rnx

diff --git a/Engine/Source/Programs/AutomationTool/Scripts/ListThirdPartySoftware.Automation.cs b/Engine/Source/Programs/AutomationTool/Scripts/ListThirdPartySoftware.Automation.cs
--- a/Engine/Source/Programs/AutomationTool/Scripts/ListThirdPartySoftware.Automation.cs
+++ b/Engine/Source/Programs/AutomationTool/Scripts/ListThirdPartySoftware.Automation.cs
@@ -32,7 +32,7 @@ class ListThirdPartySoftware : BuildCommand
 
 			IProcessResult Result;
 
-			Result = Run(UE4Build.GetUBTExecutable(), String.Format("{0} {1} -jsonexport=\"{2}\" -skipbuild", Target.Replace('|', ' '),  ProjectPath, OutputFile.FullName), Options: ERunOptions.Default);
+			Result = Run(UE4Build.GetUBTExecutable(), String.Format("{0} {1} -Mode=JsonExport -OutputFile=\"{2}\"", Target.Replace('|', ' '),  ProjectPath, OutputFile.FullName), Options: ERunOptions.Default);
 
 			if (Result.ExitCode != 0)
 			{
@@ -59,14 +59,6 @@ class ListThirdPartySoftware : BuildCommand
 				DirectoriesToScan.Add(DirectoryReference.Combine(ProjectFile.Directory, "Content"));
 			}
 
-			// Get the variables to be expanded in any runtime dependencies variables
-			Dictionary<string, string> Variables = new Dictionary<string, string>();
-			Variables.Add("EngineDir", CommandUtils.EngineDirectory.FullName);
-			if(ProjectFile != null)
-			{
-				Variables.Add("ProjectDir", ProjectFile.Directory.FullName);
-			}
-
 			// Add all the paths for each module, and its runtime dependencies
 			JsonObject Modules = Object.GetObjectField("Modules");
 			foreach(string ModuleName in Modules.KeyNames)
@@ -76,9 +68,11 @@ class ListThirdPartySoftware : BuildCommand
 
 				foreach(JsonObject RuntimeDependency in Module.GetObjectArrayField("RuntimeDependencies"))
 				{
-					string RuntimeDependencyPath = RuntimeDependency.GetStringField("Path");
-					RuntimeDependencyPath = Utils.ExpandVariables(RuntimeDependencyPath, Variables);
-					DirectoriesToScan.Add(new FileReference(RuntimeDependencyPath).Directory);
+					string RuntimeDependencyPath;
+					if(RuntimeDependency.TryGetStringField("SourcePath", out RuntimeDependencyPath) || RuntimeDependency.TryGetStringField("Path", out RuntimeDependencyPath))
+					{
+						DirectoriesToScan.Add(new FileReference(RuntimeDependencyPath).Directory);
+					}
 				}
 			}
 
diff --git a/Engine/Source/Programs/UnrealBuildTool/Configuration/UEBuildModule.cs b/Engine/Source/Programs/UnrealBuildTool/Configuration/UEBuildModule.cs
--- a/Engine/Source/Programs/UnrealBuildTool/Configuration/UEBuildModule.cs
+++ b/Engine/Source/Programs/UnrealBuildTool/Configuration/UEBuildModule.cs
@@ -863,6 +863,20 @@ namespace UnrealBuildTool
 				Writer.WriteValue(ModuleName);
 			}
 			Writer.WriteArrayEnd();
+
+			Writer.WriteArrayStart("RuntimeDependencies");
+			foreach(ModuleRules.RuntimeDependency RuntimeDependency in Rules.RuntimeDependencies.Inner)
+			{
+				Writer.WriteObjectStart();
+				Writer.WriteValue("Path", ExpandPathVariables(RuntimeDependency.Path, null, null));
+				if(RuntimeDependency.SourcePath != null)
+				{
+					Writer.WriteValue("SourcePath", ExpandPathVariables(RuntimeDependency.SourcePath, null, null));
+				}
+				Writer.WriteValue("Type", RuntimeDependency.Type.ToString());
+				Writer.WriteObjectEnd();
+			}
+			Writer.WriteArrayEnd();
 		}
 
 		/// <summary>
