Removing obsolete check for GameTarget while detecting settings on UAT startup during cooks

diff --git a/Engine/Source/Programs/AutomationTool/AutomationUtils/ProjectParams.cs b/Engine/Source/Programs/AutomationTool/AutomationUtils/ProjectParams.cs
--- a/Engine/Source/Programs/AutomationTool/AutomationUtils/ProjectParams.cs
+++ b/Engine/Source/Programs/AutomationTool/AutomationUtils/ProjectParams.cs
@@ -2117,11 +2117,6 @@ namespace AutomationTool
 				}
 			}
 
-			if (String.IsNullOrEmpty(GameTarget) && Run && !NoClient && (Cook || CookOnTheFly) && CommandUtils.IsNullOrEmpty(ClientCookedTargetsList))
-			{
-				throw new AutomationException("Game target not found. Game target is required with -cook or -cookonthefly");
-			}
-
 			if (EditorTargetsList == null)
 			{
 				if (!GlobalCommandLine.NoCompileEditor && !IsProgramTarget && !String.IsNullOrEmpty(EditorTarget))
