Merging fixes to ListThirdPartySoftware script.

diff --git a/Engine/Source/Programs/AutomationTool/Scripts/ListThirdPartySoftware.Automation.cs b/Engine/Source/Programs/AutomationTool/Scripts/ListThirdPartySoftware.Automation.cs
--- a/Engine/Source/Programs/AutomationTool/Scripts/ListThirdPartySoftware.Automation.cs
+++ b/Engine/Source/Programs/AutomationTool/Scripts/ListThirdPartySoftware.Automation.cs
@@ -71,7 +71,8 @@ class ListThirdPartySoftware : BuildCommand
 					string RuntimeDependencyPath;
 					if(RuntimeDependency.TryGetStringField("SourcePath", out RuntimeDependencyPath) || RuntimeDependency.TryGetStringField("Path", out RuntimeDependencyPath))
 					{
-						DirectoriesToScan.Add(new FileReference(RuntimeDependencyPath).Directory);
+						List<FileReference> Files = FileFilter.ResolveWildcard(DirectoryReference.Combine(CommandUtils.EngineDirectory, "Source"), RuntimeDependencyPath);
+						DirectoriesToScan.UnionWith(Files.Select(x => x.Directory));
 					}
 				}
 			}
@@ -108,21 +109,27 @@ class ListThirdPartySoftware : BuildCommand
 		foreach(FileReference TpsFile in TpsFiles)
 		{
 			string Message = TpsFile.FullName;
-
-			string[] Lines = FileReference.ReadAllLines(TpsFile);
-			foreach(string Line in Lines)
+			try
 			{
-				const string RedirectPrefix = "Redirect:";
-
-				int Idx = Line.IndexOf(RedirectPrefix, StringComparison.InvariantCultureIgnoreCase);
-				if(Idx >= 0)
+				string[] Lines = FileReference.ReadAllLines(TpsFile);
+				foreach(string Line in Lines)
 				{
-					FileReference RedirectTpsFile = FileReference.Combine(TpsFile.Directory, Line.Substring(Idx + RedirectPrefix.Length).Trim());
-					Message = String.Format("{0} (redirect from {1})", RedirectTpsFile.FullName, TpsFile.FullName);
-					break;
+					const string RedirectPrefix = "Redirect:";
+
+					int Idx = Line.IndexOf(RedirectPrefix, StringComparison.InvariantCultureIgnoreCase);
+					if(Idx >= 0)
+					{
+						FileReference RedirectTpsFile = FileReference.Combine(TpsFile.Directory, Line.Substring(Idx + RedirectPrefix.Length).Trim());
+						Message = String.Format("{0} (redirect from {1})", RedirectTpsFile.FullName, TpsFile.FullName);
+						break;
+					}
 				}
 			}
-
+			catch (Exception Ex)
+			{
+				ExceptionUtils.AddContext(Ex, "while processing {0}", TpsFile);
+				throw;
+			}
 			OutputMessages.Add(Message);
 		}
 		OutputMessages.Sort();
