Move the FinalizeInstalledBuild command into a file by the same name, to make it easier to find.

#rnx

diff --git a/Engine/Source/Programs/AutomationTool/Scripts/AutomationScripts.Automation.csproj b/Engine/Source/Programs/AutomationTool/Scripts/AutomationScripts.Automation.csproj
--- a/Engine/Source/Programs/AutomationTool/Scripts/AutomationScripts.Automation.csproj
+++ b/Engine/Source/Programs/AutomationTool/Scripts/AutomationScripts.Automation.csproj
@@ -62,6 +62,7 @@
     </Compile>
     <Compile Include="AnalyzeThirdPartyLibs.Automation.cs" />
     <Compile Include="ArchiveCommand.Automation.cs" />
+    <Compile Include="FinalizeInstalledBuild.cs" />
     <Compile Include="ParseMsvcTimingInfo.cs" />
     <Compile Include="StashTarget.cs" />
     <Compile Include="BlameKeyword.Automation.cs" />
@@ -93,7 +94,6 @@
     <Compile Include="RebasePublicIncludePaths.cs" />
     <Compile Include="FixupRedirects.Automation.cs" />
     <Compile Include="GenerateDSYM.Automation.cs" />
-    <Compile Include="InstalledBuild.cs" />
     <Compile Include="ListThirdPartySoftware.Automation.cs" />
     <Compile Include="RebuildHLODCommand.Automation.cs" />
     <Compile Include="RebuildLightMapsCommand.Automation.cs" />
diff --git a/Engine/Source/Programs/AutomationTool/Scripts/FinalizeInstalledBuild.cs b/Engine/Source/Programs/AutomationTool/Scripts/FinalizeInstalledBuild.cs
new file mode 100644
--- /dev/null
+++ b/Engine/Source/Programs/AutomationTool/Scripts/FinalizeInstalledBuild.cs
@@ -0,0 +1,154 @@
+// Copyright 1998-2019 Epic Games, Inc. All Rights Reserved.
+
+using System;
+using System.Collections.Generic;
+using System.IO;
+using System.Linq;
+using Tools.DotNETCommon;
+using UnrealBuildTool;
+
+namespace AutomationTool
+{
+	/// <summary>
+	/// Commandlet to finalize the creation of an installed build - creating an InstalledBuild.txt file and writing
+	/// out installed platform entries for all platforms/configurations where a UE4Game .target file can be found
+	/// </summary>
+	[Help("Command to perform additional steps to prepare an installed build.")]
+	[Help("OutputDir=<RootDirectory>", "Root Directory of the installed build data (required)")]
+	[Help("ContentOnlyPlatforms=<PlatformNameList>", "List of platforms that should be marked as only supporting content projects (optional)")]
+	[Help("<Platform>Architectures=<ArchitectureList>)", "List of architectures that are used for a given platform (optional)")]
+	[Help("<Platform>GPUArchitectures=<GPUArchitectureList>)", "List of GPU architectures that are used for a given platform (optional)")]
+	[Help("AnalyticsTypeOverride=<TypeName>", "Name to give this build for analytics purposes (optional)")]
+	class FinalizeInstalledBuild : BuildCommand
+	{
+		/// <summary>
+		/// Entry point for the commandlet
+		/// </summary>
+		public override void ExecuteBuild()
+		{
+			string OutputDir = ParseParamValue("OutputDir");
+			string ContentOnlyPlatformsString = ParseParamValue("ContentOnlyPlatforms");
+			IEnumerable<UnrealTargetPlatform> ContentOnlyPlatforms = Enumerable.Empty<UnrealTargetPlatform>();
+			if (!String.IsNullOrWhiteSpace(ContentOnlyPlatformsString))
+			{
+				ContentOnlyPlatforms = ContentOnlyPlatformsString.Split(';').Where(x => !String.IsNullOrWhiteSpace(x)).Select(x => (UnrealTargetPlatform)Enum.Parse(typeof(UnrealTargetPlatform), x));
+			}
+			string AnalyticsTypeOverride = ParseParamValue("AnalyticsTypeOverride");
+
+			// Write InstalledBuild.txt to indicate Engine is installed
+			string InstalledBuildFile = CommandUtils.CombinePaths(OutputDir, "Engine/Build/InstalledBuild.txt");
+			CommandUtils.CreateDirectory(CommandUtils.GetDirectoryName(InstalledBuildFile));
+			CommandUtils.WriteAllText(InstalledBuildFile, "");
+
+			// Write InstalledBuild.txt to indicate Engine is installed
+			string Project = ParseParamValue("Project");
+			if(Project != null)
+			{
+				string InstalledProjectBuildFile = CommandUtils.CombinePaths(OutputDir, "Engine/Build/InstalledProjectBuild.txt");
+				CommandUtils.CreateDirectory(CommandUtils.GetDirectoryName(InstalledProjectBuildFile));
+				CommandUtils.WriteAllText(InstalledProjectBuildFile, new FileReference(Project).MakeRelativeTo(new DirectoryReference(OutputDir)));
+			}
+
+			string OutputEnginePath = Path.Combine(OutputDir, "Engine");
+			string OutputBaseEnginePath = Path.Combine(OutputEnginePath, "Config", "BaseEngine.ini");
+			FileAttributes OutputAttributes = FileAttributes.ReadOnly;
+			List<String> IniLines = new List<String>();
+
+			// Should always exist but if not, we don't need extra line
+			if (File.Exists(OutputBaseEnginePath))
+			{
+				OutputAttributes = File.GetAttributes(OutputBaseEnginePath);
+				IniLines.Add("");
+			}
+			else
+			{
+				CommandUtils.CreateDirectory(CommandUtils.GetDirectoryName(OutputBaseEnginePath));
+				CommandUtils.WriteAllText(OutputBaseEnginePath, "");
+				OutputAttributes = File.GetAttributes(OutputBaseEnginePath) | OutputAttributes;
+			}
+
+			// Create list of platform configurations installed in a Rocket build
+			List<InstalledPlatformInfo.InstalledPlatformConfiguration> InstalledConfigs = new List<InstalledPlatformInfo.InstalledPlatformConfiguration>();
+
+			// Add the editor platform, otherwise we'll never be able to run UAT
+			InstalledConfigs.Add(new InstalledPlatformInfo.InstalledPlatformConfiguration(UnrealTargetConfiguration.Development, HostPlatform.Current.HostEditorPlatform, TargetRules.TargetType.Editor, "", "", EProjectType.Unknown, false));
+
+			foreach (UnrealTargetPlatform CodeTargetPlatform in Enum.GetValues(typeof(UnrealTargetPlatform)))
+			{
+				if (PlatformExports.IsPlatformAvailable(CodeTargetPlatform))
+				{
+					string Architecture = PlatformExports.GetDefaultArchitecture(CodeTargetPlatform, null);
+
+					// Try to parse additional Architectures from the command line
+					string Architectures = ParseParamValue(CodeTargetPlatform.ToString() + "Architectures");
+					string GPUArchitectures = ParseParamValue(CodeTargetPlatform.ToString() + "GPUArchitectures");
+
+					// Build a list of pre-compiled architecture combinations for this platform if any
+					List<string> AllArchNames;
+
+					if (!String.IsNullOrWhiteSpace(Architectures) && !String.IsNullOrWhiteSpace(GPUArchitectures))
+					{
+						AllArchNames = (from Arch in Architectures.Split('+')
+										from GPUArch in GPUArchitectures.Split('+')
+										select "-" + Arch + "-" + GPUArch).ToList();
+					}
+					else if (!String.IsNullOrWhiteSpace(Architectures))
+					{
+						AllArchNames = Architectures.Split('+').ToList();
+					}
+					else
+					{
+						AllArchNames = new List<string>();
+					}
+
+					// Check whether this platform should only be used for content based projects
+					EProjectType ProjectType = ContentOnlyPlatforms.Contains(CodeTargetPlatform) ? EProjectType.Content : EProjectType.Any;
+
+					// Allow Content only platforms to be shown as options in all projects
+					bool bCanBeDisplayed = ProjectType == EProjectType.Content;
+					foreach (UnrealTargetConfiguration CodeTargetConfiguration in Enum.GetValues(typeof(UnrealTargetConfiguration)))
+					{
+						// Need to check for development receipt as we use that for the Engine code in DebugGame
+						UnrealTargetConfiguration EngineConfiguration = (CodeTargetConfiguration == UnrealTargetConfiguration.DebugGame) ? UnrealTargetConfiguration.Development : CodeTargetConfiguration;
+						FileReference ReceiptFileName = TargetReceipt.GetDefaultPath(new DirectoryReference(OutputEnginePath), "UE4Game", CodeTargetPlatform, EngineConfiguration, Architecture);
+
+						if (FileReference.Exists(ReceiptFileName))
+						{
+							// Strip the output folder so that this can be used on any machine
+							string RelativeReceiptFileName = ReceiptFileName.MakeRelativeTo(new DirectoryReference(OutputDir));
+
+							// If we have pre-compiled architectures for this platform then add an entry for each of these -
+							// there isn't a receipt for each architecture like some other platforms
+							if (AllArchNames.Count > 0)
+							{
+								foreach (string Arch in AllArchNames)
+								{
+									InstalledConfigs.Add(new InstalledPlatformInfo.InstalledPlatformConfiguration(CodeTargetConfiguration, CodeTargetPlatform, TargetType.Game, Arch, RelativeReceiptFileName, ProjectType, bCanBeDisplayed));
+								}
+							}
+							else
+							{
+								InstalledConfigs.Add(new InstalledPlatformInfo.InstalledPlatformConfiguration(CodeTargetConfiguration, CodeTargetPlatform, TargetType.Game, Architecture, RelativeReceiptFileName, ProjectType, bCanBeDisplayed));
+							}
+						}
+					}
+				}
+			}
+
+			UnrealBuildTool.InstalledPlatformInfo.WriteConfigFileEntries(InstalledConfigs, ref IniLines);
+
+			if (!String.IsNullOrEmpty(AnalyticsTypeOverride))
+			{
+				// Write Custom Analytics type setting
+				IniLines.Add("");
+				IniLines.Add("[Analytics]");
+				IniLines.Add(String.Format("UE4TypeOverride=\"{0}\"", AnalyticsTypeOverride));
+			}
+
+			// Make sure we can write to the the config file
+			File.SetAttributes(OutputBaseEnginePath, OutputAttributes & ~FileAttributes.ReadOnly);
+			File.AppendAllLines(OutputBaseEnginePath, IniLines);
+			File.SetAttributes(OutputBaseEnginePath, OutputAttributes);
+		}
+	}
+}
\ No newline at end of file
diff --git a/Engine/Source/Programs/AutomationTool/Scripts/InstalledBuild.cs b/Engine/Source/Programs/AutomationTool/Scripts/InstalledBuild.cs
deleted file mode 100644
