Sanitize device service web exception error message
#rnx

diff --git a/Engine/Source/Programs/AutomationTool/AutomationUtils/DeviceReservation.cs b/Engine/Source/Programs/AutomationTool/AutomationUtils/DeviceReservation.cs
--- a/Engine/Source/Programs/AutomationTool/AutomationUtils/DeviceReservation.cs
+++ b/Engine/Source/Programs/AutomationTool/AutomationUtils/DeviceReservation.cs
@@ -205,6 +205,22 @@ namespace AutomationTool.DeviceReservation
 			public string ReservationDetails;
 		}
 
+		private static string SanitizeErrorMessage(string Message)
+		{
+			string[] TriggersSrc = { "Warning:", "Error:", "Exception:" };
+			string[] TriggersDst = { "Warn1ng:", "Err0r:", "Except10n:" };
+
+			for (int Index = 0; Index < TriggersSrc.Length; ++Index)
+			{
+				if (Message.IndexOf(TriggersSrc[Index], StringComparison.OrdinalIgnoreCase) != -1)
+				{
+					Message = Regex.Replace(Message, TriggersSrc[Index], TriggersDst[Index], RegexOptions.IgnoreCase);
+				}
+			}
+
+			return Message;
+		}
+
 		public static Reservation Create(Uri BaseUri, string[] DeviceTypes, TimeSpan Duration, int RetryMax = 5)
 		{
 			bool bFirst = true;
@@ -236,8 +252,7 @@ namespace AutomationTool.DeviceReservation
 				}
 				catch (WebException WebEx)
 				{
-
-					Console.WriteLine(String.Format("WebException on reservation request: {0} : {1}", WebEx.Message, WebEx.Status));
+					Console.WriteLine(String.Format("WebException on reservation request: {0} : {1}", SanitizeErrorMessage(WebEx.Message), WebEx.Status));
 
 					if (RetryCount == RetryMax)
 					{
@@ -264,21 +279,7 @@ namespace AutomationTool.DeviceReservation
 				catch (Exception Ex)
 				{
 					UnknownException = Ex;
-
-					string Line = UnknownException.Message;
-
-					string[] TriggersSrc = { "Warning:", "Error:", "Exception:" };
-					string[] TriggersDst = { "Warn1ng:", "Err0r:", "Except10n:" };
-
-					for (int Index = 0; Index < TriggersSrc.Length; ++Index)
-					{
-						if (Line.IndexOf(TriggersSrc[Index], StringComparison.OrdinalIgnoreCase) != -1)
-						{
-							Line = Regex.Replace(Line, TriggersSrc[Index], TriggersDst[Index], RegexOptions.IgnoreCase);
-						}
-					}
-
-					Console.WriteLine("Device reservation unsuccessful: {0}", Line);
+					Console.WriteLine("Device reservation unsuccessful: {0}", SanitizeErrorMessage(UnknownException.Message));
 				}
 			}
 		}
