Suspend EC error parsing while killing adb server
#rnx

diff --git a/Engine/Source/Programs/AutomationTool/Gauntlet/Platform/Android/Gauntlet.TargetDeviceAndroid.cs b/Engine/Source/Programs/AutomationTool/Gauntlet/Platform/Android/Gauntlet.TargetDeviceAndroid.cs
--- a/Engine/Source/Programs/AutomationTool/Gauntlet/Platform/Android/Gauntlet.TargetDeviceAndroid.cs
+++ b/Engine/Source/Programs/AutomationTool/Gauntlet/Platform/Android/Gauntlet.TargetDeviceAndroid.cs
@@ -1225,7 +1225,7 @@ namespace Gauntlet
 					bUsingCustomKeys = true;
 
 					Log.Info("Running adb kill-server to refresh credentials");
-					TargetDeviceAndroid.RunAdbGlobalCommand("kill-server");
+					TargetDeviceAndroid.RunAdbGlobalCommand("kill-server", bPauseErrorParsing: true);
 
 					Thread.Sleep(5000);
 				}
@@ -1247,7 +1247,7 @@ namespace Gauntlet
 					Reset();
 
 					Log.Info("Running adb kill-server to refresh credentials");
-					TargetDeviceAndroid.RunAdbGlobalCommand("kill-server");
+					TargetDeviceAndroid.RunAdbGlobalCommand("kill-server", bPauseErrorParsing: true);
 					Thread.Sleep(2500);
 				}
 
