Adding logging to help catch android device error
#rnx

diff --git a/Engine/Source/Programs/AutomationTool/Gauntlet/Platform/Android/Gauntlet.TargetDeviceAndroid.cs b/Engine/Source/Programs/AutomationTool/Gauntlet/Platform/Android/Gauntlet.TargetDeviceAndroid.cs
--- a/Engine/Source/Programs/AutomationTool/Gauntlet/Platform/Android/Gauntlet.TargetDeviceAndroid.cs
+++ b/Engine/Source/Programs/AutomationTool/Gauntlet/Platform/Android/Gauntlet.TargetDeviceAndroid.cs
@@ -1059,6 +1059,9 @@ namespace Gauntlet
 			{
 				RunOptions |= CommandUtils.ERunOptions.NoLoggingOfRunCommand;
 			}
+
+			bShouldLogCommand = true;
+			Log.Verbose("Forcing Android bShouldLogCommand to true, remove once device issue resolved: https://jira.it.epicgames.net/browse/UEATM-508");
 	
 			if (bShouldLogCommand)
 			{
