Switching SharedCookedBuild to be stateful
Support multiple simultaneous platform targets
Fix to use CookPlatform correctly

diff --git a/Engine/Source/Programs/AutomationTool/Scripts/CookCommand.Automation.cs b/Engine/Source/Programs/AutomationTool/Scripts/CookCommand.Automation.cs
--- a/Engine/Source/Programs/AutomationTool/Scripts/CookCommand.Automation.cs
+++ b/Engine/Source/Programs/AutomationTool/Scripts/CookCommand.Automation.cs
@@ -161,8 +161,8 @@ public partial class Project : CommandUtils
                 }
 				if ( Params.HasIterateSharedCookedBuild)
 				{
-					SharedCookedBuild.CopySharedCookedBuild(Params);
-					CommandletParams += " -iteratesharedcookedbuild";					
+					new SharedCookedBuild(Params).CopySharedCookedBuilds();
+					CommandletParams += " -iteratesharedcookedbuild";
 				}
 
 				if (Params.CookMapsOnly)
diff --git a/Engine/Source/Programs/AutomationTool/Scripts/CopySharedCookedBuild.Automation.cs b/Engine/Source/Programs/AutomationTool/Scripts/CopySharedCookedBuild.Automation.cs
--- a/Engine/Source/Programs/AutomationTool/Scripts/CopySharedCookedBuild.Automation.cs
+++ b/Engine/Source/Programs/AutomationTool/Scripts/CopySharedCookedBuild.Automation.cs
@@ -87,13 +87,8 @@ class CopySharedCookedBuild : BuildCommand
 			}
 		}
 
-		
-
+		new SharedCookedBuild(ProjectFile, TargetPlatforms.Select(x => x.ToString()), BuildType).CopySharedCookedBuilds();
 
-		foreach (var PlatformType in TargetPlatforms)
-		{
-			SharedCookedBuild.CopySharedCookedBuild(ProjectFile.FullName, PlatformType, BuildType, true);
-		}
 
 		
 		/*
diff --git a/Engine/Source/Programs/AutomationTool/Scripts/SharedCookedBuild.cs b/Engine/Source/Programs/AutomationTool/Scripts/SharedCookedBuild.cs
--- a/Engine/Source/Programs/AutomationTool/Scripts/SharedCookedBuild.cs
+++ b/Engine/Source/Programs/AutomationTool/Scripts/SharedCookedBuild.cs
@@ -29,43 +29,88 @@ public class SharedCookedBuild
 		/// Allow any previous version that is only a content change from local sync
 		/// </summary>
 		Content,
-		
+
 		/// <summary>
 		/// Closest previous version, regardless of code/content changes
 		/// </summary>
 		Any,
 	}
-		
-	public static void CopySharedCookedBuild(ProjectParams Params)
+
+	public enum SharedCookSource
+	{
+		Manifest,
+		LooseFiles,
+	}
+
+	private FileReference ProjectFile { get; set; }
+	private DirectoryReference InstallPath { get; set; }
+	private HashSet<string> TargetPlatforms { get; set; }
+	private SharedCookType BuildType { get; set; }
+	private List<ISharedCookedBuild> CandidateBuilds { get; set; }
+	private BuildVersion LocalSync { get; set; }
+
+	public SharedCookedBuild(ProjectParams Params)
 	{
+		List<string> Platforms = new List<string>();
 		foreach (TargetPlatformDescriptor ClientPlatform in Params.ClientTargetPlatforms)
 		{
 			TargetPlatformDescriptor DataPlatformDesc = Params.GetCookedDataPlatformForClientTarget(ClientPlatform);
-			string PlatformToCook = Platform.Platforms[DataPlatformDesc].GetCookPlatform(false, Params.Client);
-			UnrealTargetPlatform TargetPlatform = (UnrealTargetPlatform)Enum.Parse(typeof(UnrealTargetPlatform), PlatformToCook, true);
-			DirectoryReference InstallPath = DirectoryReference.Combine(Params.RawProjectPath.Directory, "Saved", "SharedIterativeBuild", PlatformToCook);
-			SharedCookType BuildType = (SharedCookType)Enum.Parse(typeof(SharedCookType), Params.IterateSharedCookedBuild, true);
+			Platforms.Add(Platform.Platforms[DataPlatformDesc].GetCookPlatform(false, Params.Client));
+		}
+		SharedCookType BuildType = (SharedCookType)Enum.Parse(typeof(SharedCookType), Params.IterateSharedCookedBuild, true);
+		SharedCookedBuildConstructor(Params.RawProjectPath, Platforms, BuildType);
+	}
 
-			CopySharedCookedBuild(Params.RawProjectPath.FullName, TargetPlatform, BuildType, true);
+	public SharedCookedBuild(FileReference ProjectFile, IEnumerable<string> TargetPlatforms, SharedCookType BuildType)
+	{
+		SharedCookedBuildConstructor(ProjectFile, TargetPlatforms, BuildType);
+	}
+
+	private void SharedCookedBuildConstructor(FileReference ProjectFile, IEnumerable<string> Platforms, SharedCookType BuildType)
+	{
+		TargetPlatforms = new HashSet<string>();
+		CandidateBuilds = new List<ISharedCookedBuild>();
+		this.ProjectFile = ProjectFile;
+		this.BuildType = BuildType;
+		InstallPath = DirectoryReference.Combine(ProjectFile.Directory, "Saved", "SharedIterativeBuild");
+		LocalSync = GetLocalSync();
+		foreach (string Platform in Platforms)
+		{
+			TargetPlatforms.Add(Platform);
 		}
 	}
 
-	public static void CopySharedCookedBuild(string ProjectFullPath, UnrealTargetPlatform TargetPlatform, SharedCookType BuildType, bool bAllowExistingBuild = true)
+	public void CopySharedCookedBuilds()
 	{
-		DirectoryReference InstallPath = DirectoryReference.Combine(new FileReference(ProjectFullPath).Directory, "Saved", "SharedIterativeBuild", TargetPlatform.ToString());
-		List<ISharedCookedBuild> SharedCookedBuilds = FindBestBuilds(ProjectFullPath, TargetPlatform, BuildType, true);
-		foreach (ISharedCookedBuild Build in SharedCookedBuilds)
+		CandidateBuilds = FindBestBuilds();
+		if (CandidateBuilds.Count == 0)
 		{
-			if (Build.CopyBuild(InstallPath))
+			throw new AutomationException("No valid shared cooked builds available");
+		}
+
+		foreach (string TargetPlatform in TargetPlatforms)
+		{
+			// Prefer existing sync if listed
+			IEnumerable<ISharedCookedBuild> LocalSync = CandidateBuilds.Where(x => x.GetType() == typeof(ExistingSharedCookedBuild) && x.Platform.Equals(TargetPlatform, StringComparison.InvariantCultureIgnoreCase));
+			if (LocalSync.Count() > 0)
 			{
-				return;
+				LocalSync.First().CopyBuild(InstallPath);
+				continue;
 			}
-		}
 
-		throw new AutomationException("Failed to install shared cooked build");
+			IEnumerable<ISharedCookedBuild> PlatformBuilds = CandidateBuilds.Where(x => x.Platform.Equals(TargetPlatform, StringComparison.InvariantCultureIgnoreCase));
+			if (PlatformBuilds.Count() > 0)
+			{
+				ISharedCookedBuild Build = PlatformBuilds.First();
+				if (!Build.CopyBuild(InstallPath))
+				{
+					CommandUtils.LogError("Failed to copy shared build for {0} {1}", Build.Platform, Build.CL);
+				}
+			}
+		}
 	}
 
-	public static BuildVersion LocalSync()
+	public static BuildVersion GetLocalSync()
 	{
 		BuildVersion P4Version = new BuildVersion();
 		if (CommandUtils.P4Enabled)
@@ -89,81 +134,92 @@ public class SharedCookedBuild
 		return P4Version;
 	}
 
-	public static List<ISharedCookedBuild> FindBestBuilds(string ProjectFullPath, UnrealTargetPlatform TargetPlatform, SharedCookType BuildType, bool bAllowExistingBuild = true)
+	public List<ISharedCookedBuild> FindBestBuilds()
 	{
 		// Attempt manifest searching first
-		FileReference ProjectFileRef = new FileReference(ProjectFullPath);
-		if (!FileReference.Exists(ProjectFileRef))
+		ConfigHierarchy Hierarchy = ConfigCache.ReadHierarchy(ConfigHierarchyType.Engine, DirectoryReference.FromFile(ProjectFile), UnrealTargetPlatform.Win64);
+
+		IEnumerable<string> RawSharedCookedSources = null;
+		Hierarchy.TryGetValues("SharedCookedBuildSettings", "SharedCookedSources", out RawSharedCookedSources);
+		if (RawSharedCookedSources == null)
 		{
-			throw new AutomationException("Cannot locate project file: {0}", ProjectFileRef);
+			throw new AutomationException("Unable to locate shared cooked builds. SharedCookedSources not set in Engine.ini [SharedCookedBuildSettings]");
 		}
 
-		ConfigHierarchy Herarchy = ConfigCache.ReadHierarchy(ConfigHierarchyType.Engine, DirectoryReference.FromFile(ProjectFileRef), TargetPlatform);
-		List<string> CookedBuildManifestPaths = null;
-		List<string> CookedBuildStagedPaths = null;
-		Herarchy.GetArray("SharedCookedBuildSettings", "SharedCookedManifestPath", out CookedBuildManifestPaths);
-		Herarchy.GetArray("SharedCookedBuildSettings", "SharedCookedBuildPath", out CookedBuildStagedPaths);
-		if (CookedBuildManifestPaths == null && CookedBuildStagedPaths == null)
+		List<Dictionary<string, string>> ParsedSharedCookSources = new List<Dictionary<string, string>>();
+		foreach (string RawConfig in RawSharedCookedSources)
 		{
-			throw new AutomationException("Unable to locate shared cooked builds. SharedCookedManifestPath and SharedCookedBuildPath not set in Engine.ini [SharedCookedBuildSettings]");
+			Dictionary<string, string> ParsedSource = null;
+			if (ConfigHierarchy.TryParse(RawConfig, out ParsedSource))
+			{
+				ParsedSharedCookSources.Add(ParsedSource);
+			}
 		}
 
-		BuildVersion Version = LocalSync();
 		List<ISharedCookedBuild> CandidateBuilds = new List<ISharedCookedBuild>();
 
-		if (bAllowExistingBuild)
+		// If existing sync is present, stick to it. Read version out of sync file
+		foreach (string Platform in TargetPlatforms)
 		{
-			// If existing sync is present, stick to it. Read version out of sync file
-			FileReference SyncedBuildFile = new FileReference(CommandUtils.CombinePaths(Path.GetDirectoryName(ProjectFullPath), "Saved", "SharedIterativeBuild", TargetPlatform.ToString(), SyncedBuildFileName));
+			FileReference SyncedBuildFile = new FileReference(CommandUtils.CombinePaths(InstallPath.FullName, Platform, SyncedBuildFileName));
 			if (FileReference.Exists(SyncedBuildFile))
 			{
 				string[] SyncedBuildInfo = FileReference.ReadAllLines(SyncedBuildFile);
 				int SyncedCL = int.Parse(SyncedBuildInfo[0]);
-				if (IsValidCL(SyncedCL, BuildType, Version))
+				if (IsValidCL(SyncedCL, BuildType, LocalSync))
 				{
-					CandidateBuilds.Add(new ExistingSharedCookedBuild { CL = SyncedCL });
+					CandidateBuilds.Add(new ExistingSharedCookedBuild() { CL = SyncedCL, Platform = Platform });
 				}
 			}
 		}
 
-		if (CookedBuildManifestPaths.Count > 0)
+		foreach (Dictionary<string, string> Source in ParsedSharedCookSources)
 		{
-			foreach (string ManifestPath in CookedBuildManifestPaths)
+			SharedCookSource SourceType = (SharedCookSource)Enum.Parse(typeof(SharedCookSource), Source["Type"], true);
+			foreach (string Platform in TargetPlatforms)
 			{
-				ISharedCookedBuild Candidate = FindBestManifestBuild(ManifestPath, TargetPlatform, BuildType, Version);
-				if (Candidate != null)
+				if (SourceType == SharedCookSource.Manifest)
 				{
-					CandidateBuilds.Add(Candidate);
+					CandidateBuilds.AddRange(FindValidManifestBuilds(Source["Path"], Platform));
+				}
+				else if (SourceType == SharedCookSource.LooseFiles)
+				{
+					CandidateBuilds.AddRange(FindValidLooseBuilds(Source["Path"], Platform));
 				}
 			}
 		}
 
-		if (CookedBuildStagedPaths.Count > 0)
+		// Strip all failed searches
+		CandidateBuilds.RemoveAll(x => x == null);
+
+		// Make sure we have a matching CL for all target platforms, regardless of source
+		List<int> OrderedDistinctCLs = CandidateBuilds.Select(x => x.CL).Distinct().OrderByDescending(i => i).ToList();
+		int BestCL = -1;
+		foreach (int CL in OrderedDistinctCLs)
 		{
-			foreach (string StagedPath in CookedBuildStagedPaths)
+			// Ensure we have a platform for each
+			HashSet<string> CLPlatforms = new HashSet<string>(CandidateBuilds.Where(x => x.CL == CL).Select(x => x.Platform).ToList());
+			if (CLPlatforms.SetEquals(TargetPlatforms))
 			{
-				ISharedCookedBuild Candidate = FindBestLooseBuild(StagedPath, TargetPlatform, BuildType, Version);
-				if (Candidate != null)
-				{
-					CandidateBuilds.Add(Candidate);
-				}
+				BestCL = CL;
+				break;
 			}
 		}
 
-		if (CandidateBuilds.Count == 0)
+		if (BestCL < 0)
 		{
-			CommandUtils.LogInformation("Could not locate valid shared cooked build");
+			CommandUtils.LogError("Could not locate valid shared cooked build for all target platforms");
 		}
 
-		return CandidateBuilds.OrderBy(x => x.CL).ToList();
+		return CandidateBuilds.Where(x => x.CL == BestCL).ToList();
 	}
 
-	public static ISharedCookedBuild FindBestManifestBuild(string Path, UnrealTargetPlatform TargetPlatform, SharedCookType BuildType, BuildVersion Version)
+	public List<ISharedCookedBuild> FindValidManifestBuilds(string Path, string TargetPlatform)
 	{
+		List<ISharedCookedBuild> ValidBuilds = new List<ISharedCookedBuild>();
 		Tuple<string, string> SplitPath = SplitOnFixedPrefix(Path);
-		Regex Pattern = RegexFromWildcards(SplitPath.Item2, Version, TargetPlatform);
+		Regex Pattern = RegexFromWildcards(SplitPath.Item2, LocalSync, TargetPlatform);
 
-		ManifestSharedCookedBuild Build = new ManifestSharedCookedBuild { CL = 0, Manifest = null };
 		DirectoryReference SearchDir = new DirectoryReference(SplitPath.Item1);
 		if (DirectoryReference.Exists(SearchDir))
 		{
@@ -173,32 +229,28 @@ public class SharedCookedBuild
 				if (Match.Success)
 				{
 					int MatchCL = int.Parse(Match.Result("${CL}"));
-					if (IsValidCL(MatchCL, BuildType, Version) && MatchCL >= Build.CL)
+					if (IsValidCL(MatchCL, BuildType, LocalSync))
 					{
-						Build = new ManifestSharedCookedBuild { CL = MatchCL, Manifest = File };
+						ValidBuilds.Add(new ManifestSharedCookedBuild { CL = MatchCL, Manifest = File, Platform = TargetPlatform });
 					}
 				}
 			}
 		}
 
-		if (Build.CL != 0)
-		{
-			return Build;
-		}
-
-		return null;
+		return ValidBuilds;
 	}
-	public static ISharedCookedBuild FindBestLooseBuild(string Path, UnrealTargetPlatform TargetPlatform, SharedCookType BuildType, BuildVersion Version)
+
+	public List<ISharedCookedBuild> FindValidLooseBuilds(string Path, string TargetPlatform)
 	{
+		List<ISharedCookedBuild> ValidBuilds = new List<ISharedCookedBuild>();
 		Tuple<string, string> SplitPath = SplitOnFixedPrefix(Path);
-		Regex Pattern = RegexFromWildcards(SplitPath.Item2, Version, TargetPlatform);
-		LooseSharedCookedBuild Build = new LooseSharedCookedBuild { CL = 0, Path = null };
+		Regex Pattern = RegexFromWildcards(SplitPath.Item2, LocalSync, TargetPlatform);
 
 		// Search for all available builds
 		const string MetaDataFilename = "\\Metadata\\DevelopmentAssetRegistry.bin";
 		string BuildRule = Path + MetaDataFilename;
-		BuildRule = BuildRule.Replace("[BRANCHNAME]", Version.BranchName);
-		BuildRule = BuildRule.Replace("[PLATFORM]", TargetPlatform.ToString());
+		BuildRule = BuildRule.Replace("[BRANCHNAME]", LocalSync.BranchName);
+		BuildRule = BuildRule.Replace("[PLATFORM]", TargetPlatform);
 		string IncludeRule = BuildRule.Replace("[CL]", "*");
 		string ExcludeRule = BuildRule.Replace("[CL]", "*-PF-*"); // Exclude preflights
 		FileFilter BuildSearch = new FileFilter();
@@ -212,19 +264,14 @@ public class SharedCookedBuild
 			if (Match.Success)
 			{
 				int MatchCL = int.Parse(Match.Result("${CL}"));
-				if (IsValidCL(MatchCL, BuildType, Version) && MatchCL >= Build.CL)
+				if (IsValidCL(MatchCL, BuildType, LocalSync))
 				{
-					Build = new LooseSharedCookedBuild { CL = MatchCL, Path = new DirectoryReference(BaseBuildPath) };
+					ValidBuilds.Add(new LooseSharedCookedBuild { CL = MatchCL, Path = new DirectoryReference(BaseBuildPath), Platform = TargetPlatform });
 				}
 			}
 		}
 
-		if (Build.CL != 0)
-		{
-			return Build;
-		}
-
-		return null;
+		return ValidBuilds;
 	}
 
 
@@ -246,11 +293,11 @@ public class SharedCookedBuild
 		return false;
 	}
 
-	private static Regex RegexFromWildcards(string Path, BuildVersion Version, UnrealTargetPlatform TargetPlatform)
+	private static Regex RegexFromWildcards(string Path, BuildVersion Version, string TargetPlatform)
 	{
 		string Pattern = Path.Replace(@"\", @"\\");
 		Pattern = Pattern.Replace("[BRANCHNAME]", Version.BranchName.Replace(@"+", @"\+"));
-		Pattern = Pattern.Replace("[PLATFORM]", TargetPlatform.ToString());
+		Pattern = Pattern.Replace("[PLATFORM]", TargetPlatform);
 		Pattern = Pattern.Replace("[CL]", @"(?<CL>\d+)");
 		return new Regex(Pattern);
 	}
@@ -264,20 +311,24 @@ public class SharedCookedBuild
 
 	public interface ISharedCookedBuild
 	{
-		int CL { get; }
+		int CL { get; set; }
+		string Platform { get; set; }
 		bool CopyBuild(DirectoryReference InstallPath);
 	}
 
 	private class ManifestSharedCookedBuild : ISharedCookedBuild
 	{
 		public int CL { get; set; }
+		public string Platform { get; set; }
 		public FileReference Manifest { get; set; }
 
 		public bool CopyBuild(DirectoryReference InstallPath)
 		{
 			CommandUtils.LogInformation("Installing shared cooked build from manifest: {0} to {1}", Manifest.FullName, InstallPath.FullName);
 
-			FileReference PreviousManifest = FileReference.Combine(InstallPath, ".build", "Current.manifest");
+			DirectoryReference PlatformInstallPath = DirectoryReference.Combine(InstallPath, Platform.ToString());
+
+			FileReference PreviousManifest = FileReference.Combine(PlatformInstallPath, ".build", "Current.manifest");
 
 			FileReference BPTI = FileReference.Combine(CommandUtils.RootDirectory, "Engine", "Binaries", "Win64", "NotForLicensees", "BuildPatchToolInstaller.exe");
 			if (!FileReference.Exists(BPTI))
@@ -287,19 +338,19 @@ public class SharedCookedBuild
 			}
 
 			bool PreviousManifestExists = FileReference.Exists(PreviousManifest);
-			if (!PreviousManifestExists && DirectoryReference.Exists(InstallPath))
+			if (!PreviousManifestExists && DirectoryReference.Exists(PlatformInstallPath))
 			{
-				DirectoryReference.Delete(InstallPath, true);
+				DirectoryReference.Delete(PlatformInstallPath, true);
 			}
 
-			IProcessResult Result = CommandUtils.Run(BPTI.FullName, string.Format("-Manifest={0} -OutputDir={1}", Manifest.FullName, InstallPath.FullName), null, CommandUtils.ERunOptions.AllowSpew);
+			IProcessResult Result = CommandUtils.Run(BPTI.FullName, string.Format("-Manifest={0} -OutputDir={1}", Manifest.FullName, PlatformInstallPath.FullName), null, CommandUtils.ERunOptions.Default);
 			if (Result.ExitCode != 0)
 			{
-				CommandUtils.LogWarning("Failed to install manifest {0} to {1}", Manifest.FullName, InstallPath.FullName);
+				CommandUtils.LogWarning("Failed to install manifest {0} to {1}", Manifest.FullName, PlatformInstallPath.FullName);
 				return false;
 			}
 
-			FileReference SyncedBuildFile = new FileReference(CommandUtils.CombinePaths(InstallPath.FullName, SyncedBuildFileName));
+			FileReference SyncedBuildFile = new FileReference(CommandUtils.CombinePaths(PlatformInstallPath.FullName, SyncedBuildFileName));
 			FileReference.WriteAllLines(SyncedBuildFile, new string[] { CL.ToString(), Manifest.FullName });
 
 			return true;
@@ -309,6 +360,7 @@ public class SharedCookedBuild
 	private class LooseSharedCookedBuild : ISharedCookedBuild
 	{
 		public int CL { get; set; }
+		public string Platform { get; set; }
 		public DirectoryReference Path { get; set; }
 		public bool CopyBuild(DirectoryReference InstallPath)
 		{
@@ -337,6 +389,8 @@ public class SharedCookedBuild
 	private class ExistingSharedCookedBuild : ISharedCookedBuild
 	{
 		public int CL { get; set; }
+		public string Platform { get; set; }
+
 		public bool CopyBuild(DirectoryReference InstallPath)
 		{
 			CommandUtils.LogInformation("Using previously synced shared cooked build");
