Gauntlet: Fix issue with adb attempting to get activity log after disposal (and adb server kill)
#rnx

diff --git a/Engine/Source/Programs/AutomationTool/Gauntlet/Platform/Android/Gauntlet.TargetDeviceAndroid.cs b/Engine/Source/Programs/AutomationTool/Gauntlet/Platform/Android/Gauntlet.TargetDeviceAndroid.cs
--- a/Engine/Source/Programs/AutomationTool/Gauntlet/Platform/Android/Gauntlet.TargetDeviceAndroid.cs
+++ b/Engine/Source/Programs/AutomationTool/Gauntlet/Platform/Android/Gauntlet.TargetDeviceAndroid.cs
@@ -146,6 +146,11 @@ namespace Gauntlet
 					return ActivityLogCached;
 				}
 
+				if (Install.AndroidDevice != null && Install.AndroidDevice.Disposed)
+				{
+					return String.IsNullOrEmpty(ActivityLogCached) ? String.Empty : ActivityLogCached;
+				}
+
 				ActivityLogTime = DateTime.UtcNow;
 
 				string GetLogCommand = string.Format("shell cat {0}/Logs/{1}.log", Install.AndroidDevice.DeviceArtifactPath, Install.Name);
@@ -511,6 +516,16 @@ namespace Gauntlet
 			// TODO: uncomment the following line if the finalizer is overridden above.
 			// GC.SuppressFinalize(this);
 		}
+
+		public bool Disposed
+		{
+			get
+			{
+				return disposedValue;
+			}
+			
+		}
+
 		#endregion
 
 		/// <summary>
@@ -1059,9 +1074,6 @@ namespace Gauntlet
 			{
 				RunOptions |= CommandUtils.ERunOptions.NoLoggingOfRunCommand;
 			}
-
-			bShouldLogCommand = true;
-			Log.Verbose("Forcing Android bShouldLogCommand to true, remove once device issue resolved: https://jira.it.epicgames.net/browse/UEATM-508");
 	
 			if (bShouldLogCommand)
 			{
