Added Win32 support to Gauntlet

Tests with EngineTest project

diff --git a/Engine/Source/Programs/AutomationTool/Gauntlet/Framework/Devices/Gauntlet.DevicePool.cs b/Engine/Source/Programs/AutomationTool/Gauntlet/Framework/Devices/Gauntlet.DevicePool.cs
--- a/Engine/Source/Programs/AutomationTool/Gauntlet/Framework/Devices/Gauntlet.DevicePool.cs
+++ b/Engine/Source/Programs/AutomationTool/Gauntlet/Framework/Devices/Gauntlet.DevicePool.cs
@@ -810,7 +810,7 @@ namespace Gauntlet
 
 				if (IsDesktop)
 				{
-					string ClientTempDir = Path.Combine(LocalTempDir, "PCTemp");
+					string ClientTempDir = Path.Combine(LocalTempDir, "DeviceCache", Def.Platform.ToString());
 					int DeviceCount = AvailableDevices.Union(ReservedDevices).Where(D => D.Platform == Def.Platform).Count();
 
 					NewDevice = Factory.CreateDevice(Def.Name, ClientTempDir);
diff --git a/Engine/Source/Programs/AutomationTool/Gauntlet/Platform/Windows/Gauntlet.TargetDeviceWindows.cs b/Engine/Source/Programs/AutomationTool/Gauntlet/Platform/Windows/Gauntlet.TargetDeviceWindows.cs
--- a/Engine/Source/Programs/AutomationTool/Gauntlet/Platform/Windows/Gauntlet.TargetDeviceWindows.cs
+++ b/Engine/Source/Programs/AutomationTool/Gauntlet/Platform/Windows/Gauntlet.TargetDeviceWindows.cs
@@ -173,7 +173,7 @@ namespace Gauntlet
 		}
 	}
 
-	public class WindowsDeviceFactory : IDeviceFactory
+	public class Win64DeviceFactory : IDeviceFactory
 	{
 		public bool CanSupportPlatform(UnrealTargetPlatform Platform)
 		{
@@ -186,6 +186,19 @@ namespace Gauntlet
 		}
 	}
 
+	public class Wind32DeviceFactory : IDeviceFactory
+	{
+		public bool CanSupportPlatform(UnrealTargetPlatform Platform)
+		{
+			return Platform == UnrealTargetPlatform.Win32;
+		}
+
+		public ITargetDevice CreateDevice(string InRef, string InParam)
+		{
+			return new TargetDeviceWindows(InRef, InParam, false);
+		}
+	}
+
 	/// <summary>
 	/// Win32/64 implementation of a device to run applications
 	/// </summary>
@@ -195,12 +208,14 @@ namespace Gauntlet
 
 		protected string UserDir { get; set; }
 
+		protected bool IsWin64 { get; set; }
+
 		/// <summary>
 		/// Our mappings of Intended directories to where they actually represent on this platform.
 		/// </summary>
 		protected Dictionary<EIntendedBaseCopyDirectory, string> LocalDirectoryMappings { get; set; }
 
-		public TargetDeviceWindows(string InName, string InTempDir)
+		public TargetDeviceWindows(string InName, string InTempDir, bool InIsWin64=true)
 		{
 			Name = InName;
 			TempDir = InTempDir;
@@ -208,6 +223,8 @@ namespace Gauntlet
 
 			UserDir = Path.Combine(TempDir, string.Format("{0}_UserDir", Name));
             LocalDirectoryMappings = new Dictionary<EIntendedBaseCopyDirectory, string>();
+
+			IsWin64 = InIsWin64;
 		}
 
 		#region IDisposable Support
@@ -495,7 +512,7 @@ namespace Gauntlet
 			return string.Compare(Path.GetPathRoot(InPath), Path.GetPathRoot(this.TempDir), StringComparison.OrdinalIgnoreCase) == 0;
 		}
 
-		public UnrealTargetPlatform Platform { get { return UnrealTargetPlatform.Win64; } }
+		public UnrealTargetPlatform Platform { get { return IsWin64 ? UnrealTargetPlatform.Win64 : UnrealTargetPlatform.Win32; } }
 
 		public string TempDir { get; private set; }
 		public bool IsAvailable { get { return true; } }
diff --git a/Engine/Source/Programs/AutomationTool/Gauntlet/Platform/Windows/Gauntlet.WindowsBuildSources.cs b/Engine/Source/Programs/AutomationTool/Gauntlet/Platform/Windows/Gauntlet.WindowsBuildSources.cs
--- a/Engine/Source/Programs/AutomationTool/Gauntlet/Platform/Windows/Gauntlet.WindowsBuildSources.cs
+++ b/Engine/Source/Programs/AutomationTool/Gauntlet/Platform/Windows/Gauntlet.WindowsBuildSources.cs
@@ -13,4 +13,13 @@ namespace Gauntlet
 
 		public override string PlatformFolderPrefix { get { return "Windows"; } }
 	}
+
+	public class Win32BuildSource : StagedBuildSource<StagedBuild>
+	{
+		public override string BuildName { get { return "Win32StagedBuild"; } }
+
+		public override UnrealTargetPlatform Platform { get { return UnrealTargetPlatform.Win32; } }
+
+		public override string PlatformFolderPrefix { get { return "Windows"; } }
+	}
 }
