Add option to split OBB into main and patch if exceeds OBB size limit

#ue4
#android

diff --git a/Engine/Source/Programs/AutomationTool/Android/AndroidPlatform.Automation.cs b/Engine/Source/Programs/AutomationTool/Android/AndroidPlatform.Automation.cs
--- a/Engine/Source/Programs/AutomationTool/Android/AndroidPlatform.Automation.cs
+++ b/Engine/Source/Programs/AutomationTool/Android/AndroidPlatform.Automation.cs
@@ -151,6 +151,41 @@ public class AndroidPlatform : Platform
 		return ObbName;
 	}
 
+	private static string GetFinalPatchName(string ApkName, bool bUseAppType = true)
+	{
+		// calculate the name for the .obb file
+		string PackageName = GetPackageInfo(ApkName, false);
+		if (PackageName == null)
+		{
+			throw new AutomationException(ExitCode.Error_FailureGettingPackageInfo, "Failed to get package name from " + ApkName);
+		}
+
+		string PackageVersion = GetPackageInfo(ApkName, true);
+		if (PackageVersion == null || PackageVersion.Length == 0)
+		{
+			throw new AutomationException(ExitCode.Error_FailureGettingPackageInfo, "Failed to get package version from " + ApkName);
+		}
+
+		if (PackageVersion.Length > 0)
+		{
+			int IntVersion = int.Parse(PackageVersion);
+			PackageVersion = IntVersion.ToString("0");
+		}
+
+		string AppType = bUseAppType ? GetMetaAppType() : "";
+		if (AppType.Length > 0)
+		{
+			AppType += ".";
+		}
+
+		string PatchName = string.Format("patch.{0}.{1}.{2}obb", PackageVersion, PackageName, AppType);
+
+		// plop the .obb right next to the executable
+		PatchName = Path.Combine(Path.GetDirectoryName(ApkName), PatchName);
+
+		return PatchName;
+	}
+
 
 	public override string GetPlatformPakCommandLine(ProjectParams Params, DeploymentContext SC)
 	{
@@ -172,7 +207,14 @@ public class AndroidPlatform : Platform
         return TargetAndroidLocation + PackageName + "/" + Path.GetFileName(ObbName);
 	}
 
-    public static string GetStorageQueryCommand()
+	private static string GetDevicePatchName(string ApkName)
+	{
+		string PatchName = GetFinalPatchName(ApkName, false);
+		string PackageName = GetPackageInfo(ApkName, false);
+		return TargetAndroidLocation + PackageName + "/" + Path.GetFileName(PatchName);
+	}
+
+	public static string GetStorageQueryCommand()
     {
 		if (Utils.IsRunningOnMono)
 		{
@@ -268,6 +310,55 @@ public class AndroidPlatform : Platform
 		return bAllowLargeOBBFiles ? MaxOBBSizeAllowed : NormalOBBSizeAllowed;
 	}
 
+	private bool AllowPatchOBBFile(DeploymentContext SC)
+	{
+		ConfigHierarchy Ini = ConfigCache.ReadHierarchy(ConfigHierarchyType.Engine, DirectoryReference.FromFile(SC.RawProjectPath), SC.StageTargetPlatform.PlatformType);
+		bool bAllowPatchOBBFile = false;
+		Ini.GetBool("/Script/AndroidRuntimeSettings.AndroidRuntimeSettings", "bAllowPatchOBBFile", out bAllowPatchOBBFile);
+		return bAllowPatchOBBFile;
+	}
+
+	private bool CreateOBBFile(DeploymentContext SC, string StageDirectoryPath, string OutputFilename, List<FileReference> FilesForObb)
+	{
+		LogInformation("Creating {0} from {1}", OutputFilename, SC.StageDirectory);
+		using (ZipFile ObbFile = new ZipFile(OutputFilename))
+		{
+			ObbFile.CompressionMethod = CompressionMethod.None;
+			ObbFile.CompressionLevel = Ionic.Zlib.CompressionLevel.None;
+			ObbFile.UseZip64WhenSaving = Ionic.Zip.Zip64Option.Never;
+
+			int ObbFileCount = 0;
+			ObbFile.AddProgress +=
+				delegate (object sender, AddProgressEventArgs e)
+				{
+					if (e.EventType == ZipProgressEventType.Adding_AfterAddEntry)
+					{
+						ObbFileCount += 1;
+						LogInformation("[{0}/{1}] Adding {2} to OBB",
+							ObbFileCount, e.EntriesTotal,
+							e.CurrentEntry.FileName);
+					}
+				};
+
+			foreach (FileReference FileRef in FilesForObb)
+			{
+				string DestinationPath = Path.GetDirectoryName(FileRef.FullName).Replace(StageDirectoryPath, SC.ShortProjectName);
+				ObbFile.AddFile(FileRef.FullName, DestinationPath);
+			}
+
+			// ObbFile.AddDirectory(SC.StageDirectory+"/"+SC.ShortProjectName, SC.ShortProjectName);
+			try
+			{
+				ObbFile.Save();
+			}
+			catch (Exception)
+			{
+				return false;
+			}
+		}
+		return true;
+	}
+
 	public override void Package(ProjectParams Params, DeploymentContext SC, int WorkingCL)
 	{
 		if (SC.StageTargetConfigurations.Count != 1)
@@ -294,6 +385,7 @@ public class AndroidPlatform : Platform
 		// includes any PAK files, movie files, etc.
 
 		string LocalObbName = SC.StageDirectory.FullName+".obb";
+		string LocalPatchName = SC.StageDirectory.FullName + ".patch.obb";
 
 		FileFilter ObbFileFilter = new FileFilter(FileFilterType.Include);
 		ConfigHierarchy EngineIni = ConfigCache.ReadHierarchy(ConfigHierarchyType.Engine, DirectoryReference.FromFile(Params.RawProjectPath), UnrealTargetPlatform.Android);
@@ -305,11 +397,11 @@ public class AndroidPlatform : Platform
 		}
 
 		string StageDirectoryPath = Path.Combine(SC.StageDirectory.FullName, SC.ShortProjectName);
-		List<FileReference> FilesToObb = ObbFileFilter.ApplyToDirectory(new DirectoryReference(StageDirectoryPath), true);
+		List<FileReference> FilesForObb = ObbFileFilter.ApplyToDirectory(new DirectoryReference(StageDirectoryPath), true);
 
 		bool OBBNeedsUpdate = false;
 		System.DateTime OBBTimeStamp = File.GetLastWriteTimeUtc(LocalObbName);
-		foreach (FileReference FileToObb in FilesToObb)
+		foreach (FileReference FileToObb in FilesForObb)
 		{
 			System.DateTime FileTimeStamp = File.GetLastWriteTimeUtc(FileToObb.FullName);
 			if(FileTimeStamp > OBBTimeStamp)
@@ -319,66 +411,89 @@ public class AndroidPlatform : Platform
 			}
 		}
 
+		Int64 OBBSizeAllowed = GetMaxOBBSizeAllowed(SC);
+		string LimitString = (OBBSizeAllowed < MaxOBBSizeAllowed) ? "2 GiB" : "4 GiB";
+
 		if (!OBBNeedsUpdate)
 		{
 			LogInformation("OBB is up to date: " + LocalObbName);
 		}
 		else
 		{
+			// Always delete the target OBB file if it exists
+			if (File.Exists(LocalObbName))
+			{
+				File.Delete(LocalObbName);
+			}
 
-		// Always delete the target OBB file if it exists
-		if (File.Exists(LocalObbName))
-		{
-			File.Delete(LocalObbName);
-		}
+			// Always delete the target patch OBB file if it exists
+			if (File.Exists(LocalPatchName))
+			{
+				File.Delete(LocalPatchName);
+			}
 
-		// Now create the OBB as a ZIP archive.
-		LogInformation("Creating {0} from {1}", LocalObbName, SC.StageDirectory);
-		using (ZipFile ObbFile = new ZipFile(LocalObbName))
-		{
-			ObbFile.CompressionMethod = CompressionMethod.None;
-			ObbFile.CompressionLevel = Ionic.Zlib.CompressionLevel.None;
-			ObbFile.UseZip64WhenSaving = Ionic.Zip.Zip64Option.Never;
+			List<FileReference> FilesToObb = FilesForObb;
+			List<FileReference> FilesToPatch = new List<FileReference>();
 
-			int ObbFileCount = 0;
-			ObbFile.AddProgress +=
-				delegate(object sender, AddProgressEventArgs e)
+			if (AllowPatchOBBFile(SC))
+			{
+				FilesToObb = new List<FileReference>();
+
+				// Collect the filesize and place into Obb or Patch list
+				Int64 MainObbSize = 22;		// EOCD without comment
+				Int64 PatchObbSize = 22;	// EOCD without comment
+				foreach (FileReference FileRef in FilesForObb)
 				{
-					if (e.EventType == ZipProgressEventType.Adding_AfterAddEntry)
+					FileInfo LocalFileInfo = new FileInfo(FileRef.FullName);
+					Int64 LocalFileLength = LocalFileInfo.Length;
+
+					string DestinationPath = Path.GetDirectoryName(FileRef.FullName).Replace(StageDirectoryPath, SC.ShortProjectName);
+					int FilenameLength = DestinationPath.Length;
+					Int64 Overhead = (30 + 36 + FilenameLength) + (46 + 36 + FilenameLength); // local file descriptor + central directory cost per file
+					Int64 FileRequirements = LocalFileLength + Overhead;
+
+					if (MainObbSize + FileRequirements < OBBSizeAllowed)
 					{
-						ObbFileCount += 1;
-						LogInformation("[{0}/{1}] Adding {2} to OBB",
-							ObbFileCount, e.EntriesTotal,
-							e.CurrentEntry.FileName);
+						FilesToObb.Add(FileRef);
+						MainObbSize += FileRequirements;
 					}
-				};
-
-			foreach (FileReference FileToObb in FilesToObb)
-			{
-				string DestinationPath = Path.GetDirectoryName(FileToObb.FullName).Replace(StageDirectoryPath, SC.ShortProjectName);
-				ObbFile.AddFile(FileToObb.FullName, DestinationPath);
+					else if (PatchObbSize + FileRequirements < OBBSizeAllowed)
+					{
+						FilesToPatch.Add(FileRef);
+						PatchObbSize += FileRequirements;
+					}
+					else
+					{
+						// no room in either file
+						LogInformation("Failed to build OBB: " + LocalObbName);
+						throw new AutomationException(ExitCode.Error_AndroidOBBError, "Stage Failed. Could not build OBB {0}. The file may be too big to fit in an OBB ({1} limit)", LocalObbName, LimitString);
+					}
+				}
 			}
-			
-			// ObbFile.AddDirectory(SC.StageDirectory+"/"+SC.ShortProjectName, SC.ShortProjectName);
-			try
+
+			// Now create the main OBB as a ZIP archive.
+			if (!CreateOBBFile(SC, StageDirectoryPath, LocalObbName, FilesToObb))
 			{
-				ObbFile.Save();
+				LogInformation("Failed to build OBB: " + LocalObbName);
+				throw new AutomationException(ExitCode.Error_AndroidOBBError, "Stage Failed. Could not build OBB {0}. The file may be too big to fit in an OBB ({1} limit)", LocalObbName, LimitString);
 			}
-			catch (Exception)
+
+			// Now create the patch OBB as a ZIP archive if required.
+			if (FilesToPatch.Count() > 0)
 			{
-				LogInformation("Failed to build OBB: " + LocalObbName);
-				throw new AutomationException(ExitCode.Error_AndroidOBBError, "Stage Failed. Could not build OBB {0}. The file may be too big to fit in an OBB (2 GiB limit)", LocalObbName);
+				if (!CreateOBBFile(SC, StageDirectoryPath, LocalPatchName, FilesToPatch))
+				{
+					LogInformation("Failed to build OBB: " + LocalPatchName);
+					throw new AutomationException(ExitCode.Error_AndroidOBBError, "Stage Failed. Could not build OBB {0}. The file may be too big to fit in an OBB ({1} limit)", LocalPatchName, LimitString);
+				}
 			}
 		}
-		}
 
 		// make sure the OBB is <= 2GiB (or 4GiB if large OBB enabled)
 		FileInfo OBBFileInfo = new FileInfo(LocalObbName);
-		Int64 OBBSizeAllowed = GetMaxOBBSizeAllowed(SC);
 		Int64 ObbFileLength = OBBFileInfo.Length;
 		if (ObbFileLength > OBBSizeAllowed)
 		{
-			string LimitString = (OBBSizeAllowed < MaxOBBSizeAllowed) ? "2 GiB" : "4 GiB";
 			LogInformation("OBB exceeds " + LimitString + " limit: " + ObbFileLength + " bytes");
 			throw new AutomationException(ExitCode.Error_AndroidOBBError, "Stage Failed. OBB {0} exceeds {1} limit)", LocalObbName, LimitString);
 		}
@@ -442,12 +557,21 @@ public class AndroidPlatform : Platform
 			    // Create APK specific OBB in case we have a detached OBB.
 			    string DeviceObbName = "";
 			    string ObbName = "";
+				string DevicePatchName = "";
+				string PatchName = "";
 				if (!bPackageDataInsideApk)
 			    {
 				    DeviceObbName = GetDeviceObbName(ApkName);
 				    ObbName = GetFinalObbName(ApkName);
-				    CopyFile(LocalObbName, ObbName);
-			    }
+					CopyFile(LocalObbName, ObbName);
+
+					if (File.Exists(LocalPatchName))
+					{
+						DevicePatchName = GetDevicePatchName(ApkName);
+						PatchName = GetFinalPatchName(ApkName);
+						CopyFile(LocalPatchName, PatchName);
+					}
+				}
 
 				//figure out which platforms we need to create install files for
 				bool bNeedsPCInstall = false;
@@ -463,11 +587,11 @@ public class AndroidPlatform : Platform
 					string PackageName = GetPackageInfo(ApkName, false);
 					string BatchName = GetFinalBatchName(ApkName, SC, bMakeSeparateApks ? Architecture : "", bMakeSeparateApks ? GPUArchitecture : "", false, EBatchType.Install, Target);
 					// make a batch file that can be used to install the .apk and .obb files
-					string[] BatchLines = GenerateInstallBatchFile(bPackageDataInsideApk, PackageName, ApkName, Params, ObbName, DeviceObbName, false, bIsPC, Params.Distribution, TargetSDKVersion > 22);
+					string[] BatchLines = GenerateInstallBatchFile(bPackageDataInsideApk, PackageName, ApkName, Params, ObbName, DeviceObbName, false, PatchName, DevicePatchName, false, bIsPC, Params.Distribution, TargetSDKVersion > 22);
 					File.WriteAllLines(BatchName, BatchLines);
 					// make a batch file that can be used to uninstall the .apk and .obb files
 					string UninstallBatchName = GetFinalBatchName(ApkName, SC, bMakeSeparateApks ? Architecture : "", bMakeSeparateApks ? GPUArchitecture : "", false, EBatchType.Uninstall, Target);
-					BatchLines = GenerateUninstallBatchFile(bPackageDataInsideApk, PackageName, ApkName, Params, ObbName, DeviceObbName, false, bIsPC);
+					BatchLines = GenerateUninstallBatchFile(bPackageDataInsideApk, PackageName, ApkName, Params, ObbName, DeviceObbName, false, PatchName, DevicePatchName, false, bIsPC);
 					File.WriteAllLines(UninstallBatchName, BatchLines);
 
 					string SymbolizeBatchName = GetFinalBatchName(ApkName, SC, Architecture, GPUArchitecture, false, EBatchType.Symbolize, Target);
@@ -490,8 +614,7 @@ public class AndroidPlatform : Platform
 						//    CommandUtils.FixUnixFilePermissions(NoInstallBatchName);
 						//}
 					}
-				}
-				);
+				});
 
 				if (bNeedsPCInstall)
 				{
@@ -520,7 +643,7 @@ public class AndroidPlatform : Platform
 		PrintRunTime();
 	}
 
-    private string[] GenerateInstallBatchFile(bool bPackageDataInsideApk, string PackageName, string ApkName, ProjectParams Params, string ObbName, string DeviceObbName, bool bNoObbInstall, bool bIsPC, bool bIsDistribution, bool bRequireRuntimeStoragePermission)
+    private string[] GenerateInstallBatchFile(bool bPackageDataInsideApk, string PackageName, string ApkName, ProjectParams Params, string ObbName, string DeviceObbName, bool bNoObbInstall, string PatchName, string DevicePatchName, bool bNoPatchInstall, bool bIsPC, bool bIsDistribution, bool bRequireRuntimeStoragePermission)
     {
         string[] BatchLines = null;
         string ReadPermissionGrantCommand = "shell pm grant " + PackageName + " android.permission.READ_EXTERNAL_STORAGE";
@@ -533,13 +656,16 @@ public class AndroidPlatform : Platform
 		// We can't always push directly to Android/obb so uploads to Download then moves it
 		bool bDontMoveOBB = bPackageDataInsideApk || !bIsDistribution;
 
+		bool bHavePatch = (PatchName != "");
+
 		if (!bIsPC)
         {
 			// If it is a distribution build, push to $STORAGE/Android/obb folder instead of $STORAGE/obb folder.
 			// Note that $STORAGE/Android/obb will be the folder that contains the obb if you download the app from playstore.
 			string OBBInstallCommand = bNoObbInstall ? "shell 'rm -r $EXTERNAL_STORAGE/" + DeviceObbName + "'" : "push " + Path.GetFileName(ObbName) + " $STORAGE/" + (bIsDistribution ? "Download/" : "") + DeviceObbName;
+			string PatchInstallCommand = bNoPatchInstall ? "shell 'rm -r $EXTERNAL_STORAGE/" + DevicePatchName + "'" : "push " + Path.GetFileName(PatchName) + " $STORAGE/" + (bIsDistribution ? "Download/" : "") + DevicePatchName;
 
-            LogInformation("Writing shell script for install with {0}", bPackageDataInsideApk ? "data in APK" : "separate obb");
+			LogInformation("Writing shell script for install with {0}", bPackageDataInsideApk ? "data in APK" : "separate obb");
             BatchLines = new string[] {
 						"#!/bin/sh",
 						"cd \"`dirname \"$0\"`\"",
@@ -569,6 +695,7 @@ public class AndroidPlatform : Platform
 						bPackageDataInsideApk ? "" : "\tSTORAGE=$(echo \"`$ADB $DEVICE shell 'echo $EXTERNAL_STORAGE'`\" | cat -v | tr -d '^M')",
 						bPackageDataInsideApk ? "" : "\t$ADB $DEVICE " + OBBInstallCommand,
 						bPackageDataInsideApk ? "if [ 1 ]; then" : "\tif [ $? -eq 0 ]; then",
+						!bHavePatch ? "" : (bPackageDataInsideApk ? "" : "\t$ADB $DEVICE " + PatchInstallCommand),
 						bDontMoveOBB ? "" : "\t\t$ADB $DEVICE shell mv $STORAGE/Download/obb/" + PackageName + " $STORAGE/Android/obb/" + PackageName,
 						"\t\techo",
 						"\t\techo Installation successful",
@@ -588,6 +715,7 @@ public class AndroidPlatform : Platform
         else
         {
 			string OBBInstallCommand = bNoObbInstall ? "shell rm -r %STORAGE%/" + DeviceObbName : "push " + Path.GetFileName(ObbName) + " %STORAGE%/" + (bIsDistribution ? "Download/" : "") + DeviceObbName;
+			string PatchInstallCommand = bNoPatchInstall ? "shell rm -r %STORAGE%/" + DevicePatchName : "push " + Path.GetFileName(PatchName) + " %STORAGE%/" + (bIsDistribution ? "Download/" : "") + DevicePatchName;
 
 			LogInformation("Writing bat for install with {0}", bPackageDataInsideApk ? "data in APK" : "separate OBB");
             BatchLines = new string[] {
@@ -614,6 +742,8 @@ public class AndroidPlatform : Platform
 						bPackageDataInsideApk ? "" : "@echo Installing new data. Failures here indicate storage problems (missing SD card or bad permissions) and are fatal.",
 						bPackageDataInsideApk ? "" : "%ADB% %DEVICE% " + OBBInstallCommand,
 						bPackageDataInsideApk ? "" : "if \"%ERRORLEVEL%\" NEQ \"0\" goto Error",
+						!bHavePatch ? "" : (bPackageDataInsideApk ? "" : "%ADB% %DEVICE% " + PatchInstallCommand),
+						!bHavePatch ? "" : (bPackageDataInsideApk ? "" : "if \"%ERRORLEVEL%\" NEQ \"0\" goto Error"),
 						bDontMoveOBB ? "" : "%ADB% %DEVICE% shell mv %STORAGE%/Download/obb/" + PackageName + " %STORAGE%/Android/obb/" + PackageName,
 						bDontMoveOBB ? "" : "if \"%ERRORLEVEL%\" NEQ \"0\" goto Error",
 						"@echo.",
@@ -637,7 +767,7 @@ public class AndroidPlatform : Platform
         return BatchLines;
     }
 
-	private string[] GenerateUninstallBatchFile(bool bPackageDataInsideApk, string PackageName, string ApkName, ProjectParams Params, string ObbName, string DeviceObbName, bool bNoObbInstall, bool bIsPC)
+	private string[] GenerateUninstallBatchFile(bool bPackageDataInsideApk, string PackageName, string ApkName, ProjectParams Params, string ObbName, string DeviceObbName, bool bNoObbInstall, string PatchName, string DevicePatchName, bool bNoPatchInstall, bool bIsPC)
 	{
 		string[] BatchLines = null;
 
@@ -759,6 +889,7 @@ public class AndroidPlatform : Platform
 			{
 				string ApkName = GetFinalApkName(Params, SC.StageExecutables[0], true, bMakeSeparateApks ? Architecture : "", bMakeSeparateApks ? GPUArchitecture : "");
 				string ObbName = GetFinalObbName(ApkName);
+				string PatchName = GetFinalPatchName(ApkName);
 				bool bBuildWithHiddenSymbolVisibility = BuildWithHiddenSymbolVisibility(SC);
 				bool bSaveSymbols = GetSaveSymbols(SC);
 				//string NoOBBBatchName = GetFinalBatchName(ApkName, Params, bMakeSeparateApks ? Architecture : "", bMakeSeparateApks ? GPUArchitecture : "", true, false);
@@ -791,6 +922,11 @@ public class AndroidPlatform : Platform
 				{
 					bAddedOBB = true;
 					SC.ArchiveFiles(Path.GetDirectoryName(ObbName), Path.GetFileName(ObbName));
+
+					if (FileExists(PatchName))
+					{
+						SC.ArchiveFiles(Path.GetDirectoryName(PatchName), Path.GetFileName(PatchName));
+					}
 				}
 
 				// copy optional unprotected APK if exists
@@ -1084,7 +1220,8 @@ public class AndroidPlatform : Platform
             IProcessResult Result = RunAdbCommand(Params, DeviceName, DeviceStorageQueryCommand, null, ERunOptions.AppMustExist);
             String StorageLocation = Result.Output.Trim(); // "/mnt/sdcard";
             string DeviceObbName = StorageLocation + "/" + GetDeviceObbName(ApkName);
-            string RemoteDir = StorageLocation + "/UE4Game/" + Params.ShortProjectName;
+			string DevicePatchName = StorageLocation + "/" + GetDevicePatchName(ApkName);
+			string RemoteDir = StorageLocation + "/UE4Game/" + Params.ShortProjectName;
 
             // determine if APK out of date
             string APKLastUpdateTime = new FileInfo(ApkName).LastWriteTime.ToString();
@@ -1459,8 +1596,9 @@ public class AndroidPlatform : Platform
 
                 // delete the .obb file, since it will cause nothing we just deployed to be used
                 RunAdbCommand(Params, DeviceName, "shell rm " + DeviceObbName);
-            }
-            else if (SC.Archive)
+				RunAdbCommand(Params, DeviceName, "shell rm " + DevicePatchName);
+			}
+			else if (SC.Archive)
             {
                 // deploy the obb if there is one
                 string ObbPath = Path.Combine(SC.StageDirectory.FullName, GetFinalObbName(ApkName));
@@ -1472,8 +1610,19 @@ public class AndroidPlatform : Platform
                     string Commandline = string.Format("{0} \"{1}\" \"{2}\"", BaseCommandline, ObbPath, DeviceObbName);
                     RunAdbCommand(Params, DeviceName, Commandline);
                 }
-            }
-            else
+
+				// deploy the patch if there is one
+				string PatchPath = Path.Combine(SC.StageDirectory.FullName, GetFinalPatchName(ApkName));
+				if (File.Exists(PatchPath))
+				{
+					// cache some strings
+					string BaseCommandline = "push";
+
+					string Commandline = string.Format("{0} \"{1}\" \"{2}\"", BaseCommandline, PatchPath, DevicePatchName);
+					RunAdbCommand(Params, DeviceName, Commandline);
+				}
+			}
+			else
             {
                 // cache some strings
                 string BaseCommandline = "push";
diff --git a/Engine/Source/Programs/UnrealBuildTool/Platform/Android/UEDeployAndroid.cs b/Engine/Source/Programs/UnrealBuildTool/Platform/Android/UEDeployAndroid.cs
--- a/Engine/Source/Programs/UnrealBuildTool/Platform/Android/UEDeployAndroid.cs
+++ b/Engine/Source/Programs/UnrealBuildTool/Platform/Android/UEDeployAndroid.cs
@@ -951,7 +951,7 @@ namespace UnrealBuildTool
 				bool bOBBExists = File.Exists(ObbSource);
 				AnyOBBExists |= bOBBExists;
 
-				obbData.Append("new XAPKFile(\ntrue, // true signifies a main file\n");
+				obbData.Append("new XAPKFile(\n" + (ObbSource.Contains(".patch.") ? "false, // false signifies a patch file\n" : "true, // true signifies a main file\n"));
 				obbData.AppendFormat("\"{0}\", // the version of the APK that the file was uploaded against\n", GetOBBVersionNumber(StoreVersion));
 				obbData.AppendFormat("{0}L // the length of the file in bytes\n", bOBBExists ? new FileInfo(ObbSource).Length : 0);
 				obbData.AppendFormat("){0}\n", first ? "," : "");
@@ -3352,9 +3352,15 @@ namespace UnrealBuildTool
 
 			// Generate the OBB and Shim files here
 			string ObbFileLocation = ProjectDirectory + "/Saved/StagedBuilds/Android" + CookFlavor + ".obb";
+			string PatchFileLocation = ProjectDirectory + "/Saved/StagedBuilds/Android" + CookFlavor + ".patch.obb";
+			List<string> RequiredOBBFiles = new List<String> { ObbFileLocation };
+			if (File.Exists(PatchFileLocation))
+			{
+				RequiredOBBFiles.Add(PatchFileLocation);
+			}
 
 			// Generate the OBBData.java file if out of date (can skip rewriting it if packaging inside Apk in some cases)
-			WriteJavaOBBDataFile(UE4OBBDataFileName, PackageName, new List<string> { ObbFileLocation }, CookFlavor, bPackageDataInsideApk);
+			WriteJavaOBBDataFile(UE4OBBDataFileName, PackageName, RequiredOBBFiles, CookFlavor, bPackageDataInsideApk);
 
 			// Make sure any existing proguard file in project is NOT used (back it up)
 			string ProjectBuildProguardFile = Path.Combine(GameBuildFilesPath, "proguard-project.txt");
diff --git a/Engine/Source/Runtime/Android/AndroidRuntimeSettings/Classes/AndroidRuntimeSettings.h b/Engine/Source/Runtime/Android/AndroidRuntimeSettings/Classes/AndroidRuntimeSettings.h
--- a/Engine/Source/Runtime/Android/AndroidRuntimeSettings/Classes/AndroidRuntimeSettings.h
+++ b/Engine/Source/Runtime/Android/AndroidRuntimeSettings/Classes/AndroidRuntimeSettings.h
@@ -232,6 +232,10 @@ public:
 	UPROPERTY(GlobalConfig, EditAnywhere, Category = "APK Packaging", Meta = (DisplayName = "Allow large OBB files."))
 	bool bAllowLargeOBBFiles;
 
+	// If checked, a patch OBB is generated for files not fitting in the main OBB (requires using multiple PAK files so split up content by chunk id). 
+	UPROPERTY(GlobalConfig, EditAnywhere, Category = "APK Packaging", Meta = (DisplayName = "Allow patch OBB file."))
+	bool bAllowPatchOBBFile;
+
 	// If checked, UE4Game files will be placed in ExternalFilesDir which is removed on uninstall.
 	// You should also check this if you need to save you game progress without requesting runtime WRITE_EXTERNAL_STORAGE permission in android api 23+
 	UPROPERTY(GlobalConfig, EditAnywhere, Category = "APK Packaging", Meta = (DisplayName = "Use ExternalFilesDir for UE4Game files?"))
