Gauntlet: update cached stdout property when pulling artifacts on Android
#rnx

diff --git a/Engine/Source/Programs/AutomationTool/Gauntlet/Platform/Android/Gauntlet.TargetDeviceAndroid.cs b/Engine/Source/Programs/AutomationTool/Gauntlet/Platform/Android/Gauntlet.TargetDeviceAndroid.cs
--- a/Engine/Source/Programs/AutomationTool/Gauntlet/Platform/Android/Gauntlet.TargetDeviceAndroid.cs
+++ b/Engine/Source/Programs/AutomationTool/Gauntlet/Platform/Android/Gauntlet.TargetDeviceAndroid.cs
@@ -155,8 +155,16 @@ namespace Gauntlet
 
 				string GetLogCommand = string.Format("shell cat {0}/Logs/{1}.log", Install.AndroidDevice.DeviceArtifactPath, Install.Name);
 				IProcessResult LogQuery = Install.AndroidDevice.RunAdbDeviceCommand(GetLogCommand, true);
-				ActivityLogCached = LogQuery.Output;
 
+				if (LogQuery.ExitCode != 0)
+				{
+					Log.VeryVerbose("Unable to query activity stdout on device {0}", Install.AndroidDevice.Name);					
+				}
+				else
+				{
+					ActivityLogCached = LogQuery.Output;
+				}
+					
 				// the activity has exited, mark final log sentinel 
 				if (ActivityExited)
 				{
@@ -222,13 +230,24 @@ namespace Gauntlet
 			{
 				Log.Warning("Failed to retrieve artifacts. {0}", PullCmd.Output);
 			}
+			else
+			{
+				// update final cached stdout property
+				string LogFilename = string.Format("{0}/Logs/{1}.log", LocalSaved, Install.Name);
+				if (File.Exists(LogFilename))
+				{
+					ActivityLogCached = File.ReadAllText(LogFilename);
+					ActivityLogTime = DateTime.MinValue;
+				}
+
+			}
 
 			// pull the logcat over from device.
 			IProcessResult LogcatResult = Install.AndroidDevice.RunAdbDeviceCommand("logcat -d");
 
 			string LogcatFilename = "Logcat.log";
 			// Save logcat dump to local artifact path.
-			System.IO.File.WriteAllText(Path.Combine(LocalSaved, LogcatFilename), LogcatResult.Output);
+			File.WriteAllText(Path.Combine(LocalSaved, LogcatFilename), LogcatResult.Output);
 
 			Install.AndroidDevice.PostRunCleanup();
 		}
