Pass through -gauntlet when not requiring a test controller to let the client know we are still running via Gauntlet.

diff --git a/Engine/Source/Programs/AutomationTool/Gauntlet/Unreal/Base/Gauntlet.UnrealTestNode.cs b/Engine/Source/Programs/AutomationTool/Gauntlet/Unreal/Base/Gauntlet.UnrealTestNode.cs
--- a/Engine/Source/Programs/AutomationTool/Gauntlet/Unreal/Base/Gauntlet.UnrealTestNode.cs
+++ b/Engine/Source/Programs/AutomationTool/Gauntlet/Unreal/Base/Gauntlet.UnrealTestNode.cs
@@ -360,10 +360,9 @@ namespace Gauntlet
 						}
 
 						// add controllers
-						if (TestRole.Controllers.Count > 0)
-						{
-							SessionRole.CommandLine += string.Format(" -gauntlet=\"{0}\"", string.Join(",", TestRole.Controllers));
-						}
+						SessionRole.CommandLine += TestRole.Controllers.Count > 0 ?
+							string.Format(" -gauntlet=\"{0}\"", string.Join(",", TestRole.Controllers)) 
+							: " -gauntlet";
 
 						if (PassThroughArgs.Count() > 0)
 						{
