UAT: Ignore empty arguments on the command line (ie. arguments passed as "") rather than treating them as command names.

diff --git a/Engine/Source/Programs/AutomationTool/AutomationUtils/Automation.cs b/Engine/Source/Programs/AutomationTool/AutomationUtils/Automation.cs
--- a/Engine/Source/Programs/AutomationTool/AutomationUtils/Automation.cs
+++ b/Engine/Source/Programs/AutomationTool/AutomationUtils/Automation.cs
@@ -375,16 +375,20 @@ AutomationTool.exe [-verbose] [-compileonly] [-p4] Command0 [-Arg0 -Arg1 -Arg2 .
 			CommandInfo CurrentCommand = null;
 			for (int Index = 0; Index < Arguments.Length; ++Index)
 			{
-				var Param = Arguments[Index];
-				if (Param.StartsWith("-") || Param.Contains("="))
+				// Guard against empty arguments passed as "" on the command line
+				string Param = Arguments[Index];
+				if(Param.Length > 0) 
 				{
-					ParseParam(Arguments[Index], CurrentCommand, ref OutScriptsForProjectFileName, OutAdditionalScriptsFolders);
-				}
-				else
-				{
-					CurrentCommand = new CommandInfo();
-					CurrentCommand.CommandName = Arguments[Index];
-					OutCommandsToExecute.Add(CurrentCommand);
+					if (Param.StartsWith("-") || Param.Contains("="))
+					{
+						ParseParam(Arguments[Index], CurrentCommand, ref OutScriptsForProjectFileName, OutAdditionalScriptsFolders);
+					}
+					else
+					{
+						CurrentCommand = new CommandInfo();
+						CurrentCommand.CommandName = Arguments[Index];
+						OutCommandsToExecute.Add(CurrentCommand);
+					}
 				}
 			}
 
