Gauntlet - cleaned up some logging on error cases.

diff --git a/Engine/Source/Programs/AutomationTool/Gauntlet/Framework/Gauntlet.TestExecutor.cs b/Engine/Source/Programs/AutomationTool/Gauntlet/Framework/Gauntlet.TestExecutor.cs
--- a/Engine/Source/Programs/AutomationTool/Gauntlet/Framework/Gauntlet.TestExecutor.cs
+++ b/Engine/Source/Programs/AutomationTool/Gauntlet/Framework/Gauntlet.TestExecutor.cs
@@ -588,21 +588,6 @@ namespace Gauntlet
 
 			MB.UnorderedList(TestResults);
 
-			string Summary = string.Format("Completed test pass {0} of {1}.", CurrentPass, NumPasses);
-
-			if (FailedCount > 0)
-			{
-				Log.Error("{0}", Summary);
-			}
-			else if (WarningCount > 0)
-			{
-				Log.Warning("{0}", Summary);
-			}
-			else
-			{
-				Log.Info("{0}", Summary);
-			}
-
 			// write the markdown out with each line indented
 			MB.ToString().Split('\n').ToList().ForEach(L => Log.Info("  " + L));
 
@@ -638,12 +623,14 @@ namespace Gauntlet
 			{
 				TestInfo.CancellationReason = string.Format("Terminating Test {0} due to maximum duration of {1} seconds. ", TestInfo.TestNode, TestInfo.TestNode.MaxDuration);
 				TestInfo.FinalResult = TestResult.TimedOut;
+				Log.Info("{0}", TestInfo.CancellationReason);
 			}
 
 			if (IsCancelled)
 			{
 				TestInfo.CancellationReason = string.Format("Cancelling Test {0} on request", TestInfo.TestNode);
 				TestInfo.FinalResult = TestResult.Cancelled;
+				Log.Info("{0}", TestInfo.CancellationReason);
 			}
 
 			// if the test is not running. or we've determined a result for it..
diff --git a/Engine/Source/Programs/AutomationTool/Gauntlet/Unreal/Base/Gauntlet.UnrealTestNode.cs b/Engine/Source/Programs/AutomationTool/Gauntlet/Unreal/Base/Gauntlet.UnrealTestNode.cs
--- a/Engine/Source/Programs/AutomationTool/Gauntlet/Unreal/Base/Gauntlet.UnrealTestNode.cs
+++ b/Engine/Source/Programs/AutomationTool/Gauntlet/Unreal/Base/Gauntlet.UnrealTestNode.cs
@@ -427,6 +427,8 @@ namespace Gauntlet
 			}
 
 			TestFolder = TestFolder.Replace(" ", "_");
+			TestFolder = TestFolder.Replace(":", "_");
+			TestFolder = TestFolder.Replace("|", "_");
 			TestFolder = TestFolder.Replace(",", "");
 
 			ArtifactPath = Path.Combine(Context.Options.LogDir, TestFolder);
@@ -720,13 +722,25 @@ namespace Gauntlet
 			}
 			else
 			{
+				bool WasGauntletTest = InArtifacts.SessionRole.CommandLine.ToLower().Contains("-gauntlet");
 				// ok, process appears to have exited for no good reason so try to divine a result...
-				if (LogSummary.HasTestExitCode == false
-					&& InArtifacts.SessionRole.CommandLine.ToLower().Contains("-gauntlet"))
+				if (WasGauntletTest)
 				{
-					Log.Verbose("Role {0} had 0 exit code but used Gauntlet and no TestExitCode was found. Assuming failure", InArtifacts.SessionRole.RoleType);
-					ExitCode = -1;
-					ExitReason = "No test result from Gauntlet controller";
+					if (LogSummary.HasTestExitCode == false)
+					{
+						Log.Verbose("Role {0} had 0 exit code but used Gauntlet and no TestExitCode was found. Assuming failure", InArtifacts.SessionRole.RoleType);
+						ExitCode = -1;
+						ExitReason = "No test result from Gauntlet controller";
+					}
+				}
+				else
+				{
+					// if all else fails, fall back to the exit code from the process. Not great.
+					ExitCode = InArtifacts.AppInstance.ExitCode;
+					if (ExitCode == 0)
+					{
+						ExitReason = "app exited with code 0";
+					}
 				}
 			}
 
@@ -823,11 +837,11 @@ namespace Gauntlet
 
 			if (ExitCode != 0)
 			{
-				MB.H4(string.Format("Result: Abnormal Exit: {0}", ExitReason));
+				MB.H4(string.Format("Result: Abnormal Exit: Reason={0}, Code={1}", ExitReason, ExitCode));
 			}
 			else
 			{
-				MB.H4(string.Format("Result: {0}", ExitReason));
+				MB.H4(string.Format("Result: Reason={0}, Code=0", ExitReason));
 			}
 
 			int FatalErrors = LogSummary.FatalError != null ? 1 : 0;
@@ -988,7 +1002,7 @@ namespace Gauntlet
 
 		/// <summary>
 		/// THe base implementation considers  considers Classes can override this to implement more custom detection of success/failure than our
-		/// log parsing
+		/// log parsing. Not guaranteed to be called if a test is marked complete
 		/// </summary>
 		/// <returns></returns>in
 		protected virtual TestResult GetUnrealTestResult()
@@ -1092,7 +1106,7 @@ namespace Gauntlet
 						}
 					}
 
-					MB.Paragraph("See below for callstack and logs");
+					MB.Paragraph("See below for logs and any callstacks");
 				}
 			}
 			MB.Paragraph(string.Format("Context: {0}", Context.ToString()));
diff --git a/Engine/Source/Programs/AutomationTool/Gauntlet/Unreal/Utils/Gauntlet.UnrealLogParser.cs b/Engine/Source/Programs/AutomationTool/Gauntlet/Unreal/Utils/Gauntlet.UnrealLogParser.cs
--- a/Engine/Source/Programs/AutomationTool/Gauntlet/Unreal/Utils/Gauntlet.UnrealLogParser.cs
+++ b/Engine/Source/Programs/AutomationTool/Gauntlet/Unreal/Utils/Gauntlet.UnrealLogParser.cs
@@ -124,11 +124,22 @@ namespace Gauntlet
 		/// <param name="InContent"></param>
 		/// <param name="InPattern"></param>
 		/// <returns></returns>
-		protected string[] GetAllMatchingLines(string InContent, string InPattern)
+		protected IEnumerable<Match> GetAllMatches(string InContent, string InPattern)
 		{
 			Regex regex = new Regex(InPattern);
 
-			return regex.Matches(InContent).Cast<Match>().Select(M => M.Value).ToArray();
+			return regex.Matches(InContent).Cast<Match>();
+		}
+
+		/// <summary>
+		/// Returns all lines from the specified content match the specified regex
+		/// </summary>
+		/// <param name="InContent"></param>
+		/// <param name="InPattern"></param>
+		/// <returns></returns>
+		protected string[] GetAllMatchingLines(string InContent, string InPattern)
+		{
+			return GetAllMatches(InContent, InPattern).Select(M => M.Value).ToArray();
 		}
 
 		/// <summary>
@@ -141,6 +152,16 @@ namespace Gauntlet
 			return GetAllMatchingLines(Content, InPattern);
 		}
 
+		/// <summary>
+		/// Returns all Matches that match the specified regex
+		/// </summary>
+		/// <param name="InPattern"></param>
+		/// <returns></returns>
+		public IEnumerable<Match> GetAllMatches(string InPattern)
+		{
+			return GetAllMatches(Content, InPattern);
+		}
+
 		/// <summary>
 		/// Returns a structure containing platform information extracted from the log
 		/// </summary>
@@ -196,9 +217,18 @@ namespace Gauntlet
 		/// </summary>
 		/// <param name="Channel"></param>
 		/// <returns></returns>
-		public IEnumerable<string> GetLogChannel(string Channel)
+		public IEnumerable<string> GetLogChannel(string Channel, bool ExactMatch=true)
 		{
-			string Pattern = string.Format(@"(Log{0}:\s{{0,1}}.+)", Channel);
+			string Pattern;
+
+			if (ExactMatch)
+			{
+				Pattern = string.Format(@"(Log{0}:\s{{0,1}}.+)", Channel);
+			}
+			else
+			{
+				Pattern = string.Format(@"(Log{0}.*:\s{{0,1}}.+)", Channel);
+			}
 
 			return Regex.Matches(Content, Pattern).Cast<Match>().Select(M => M.Groups[1].ToString()).ToArray();
 		}
