Fixed a config format issue and improved logging to more easily find these issues in the future.

#rnx

diff --git a/Engine/Source/Programs/AutomationTool/Gauntlet/Framework/Utils/Gauntlet.IniConfigUtil.cs b/Engine/Source/Programs/AutomationTool/Gauntlet/Framework/Utils/Gauntlet.IniConfigUtil.cs
--- a/Engine/Source/Programs/AutomationTool/Gauntlet/Framework/Utils/Gauntlet.IniConfigUtil.cs
+++ b/Engine/Source/Programs/AutomationTool/Gauntlet/Framework/Utils/Gauntlet.IniConfigUtil.cs
@@ -125,7 +125,7 @@ namespace Gauntlet
 				int SeparatorIndex = RawPair.IndexOf(KeyValueSeparator);
 				if (RawPair.Length <= SeparatorIndex + 1)
 				{
-					Log.Warning(string.Format("Found an element in the config string that was not a KeyValuePair! Skipping this value: {0}", RawPair));
+					Log.Warning(string.Format("Found an element in a config string that was not a KeyValuePair! Skipping this value: {0} From this config string: {1}", RawPair, ConfigString));
 					continue;
 				}
 				string KeyString = RawPair.Substring(0, SeparatorIndex).Trim();
