update CSV tool path in automation tools.

diff --git a/Engine/Source/Programs/AutomationTool/Scripts/RecordPerformance.Automation.cs b/Engine/Source/Programs/AutomationTool/Scripts/RecordPerformance.Automation.cs
--- a/Engine/Source/Programs/AutomationTool/Scripts/RecordPerformance.Automation.cs
+++ b/Engine/Source/Programs/AutomationTool/Scripts/RecordPerformance.Automation.cs
@@ -147,7 +147,7 @@ class RecordPerformance : BuildCommand
 
 	private void RunCsvTool(string ToolName, string Arguments)
 	{
-		var ToolPath = FileReference.Combine(DirectoryReference.Combine(RootDir, "Engine", "Source", "Programs", "NotForLicensees", "CSVTools", "Binaries"), ToolName + ".exe");
+		var ToolPath = FileReference.Combine(DirectoryReference.Combine(RootDir, "Engine", "Source", "Programs", "CSVTools", "Binaries"), ToolName + ".exe");
 		Run(ToolPath.FullName, Arguments);
 	}
 
