Gauntlet -

ITestNode now defines GetWarnings and GetErrors calls for tests to return instances that occurred during a test.

Split creation of Summary report in UnrealTestNode into Header/Body functions to allow tests to more easily customize these.

Added GetLogChannels() function to UnrealLogParser that will return all lines from multiple channels

Added AutomatonLogParser that specifically parses the output of automation (non-Gauntlet) tests.

diff --git a/Engine/Source/Programs/AutomationTool/Gauntlet/Framework/Base/Gauntlet.BaseTestNode.cs b/Engine/Source/Programs/AutomationTool/Gauntlet/Framework/Base/Gauntlet.BaseTestNode.cs
--- a/Engine/Source/Programs/AutomationTool/Gauntlet/Framework/Base/Gauntlet.BaseTestNode.cs
+++ b/Engine/Source/Programs/AutomationTool/Gauntlet/Framework/Base/Gauntlet.BaseTestNode.cs
@@ -67,6 +67,24 @@ namespace Gauntlet
 			return string.Format("{0}: {1}", Name, GetTestResult());
 		}
 
+		/// <summary>
+		/// Return list of warnings. Empty by default
+		/// </summary>
+		/// <returns></returns>
+		public virtual IEnumerable<string> GetWarnings()
+		{
+			return new string[0];
+		}
+
+		/// <summary>
+		/// Return list of errors. Empty by default
+		/// </summary>
+		/// <returns></returns>
+		public virtual IEnumerable<string> GetErrors()
+		{
+			return new string[0];
+		}
+
 		/// <summary>
 		/// Mark the test as started
 		/// </summary>
diff --git a/Engine/Source/Programs/AutomationTool/Gauntlet/Framework/Base/Gauntlet.TestNode.cs b/Engine/Source/Programs/AutomationTool/Gauntlet/Framework/Base/Gauntlet.TestNode.cs
--- a/Engine/Source/Programs/AutomationTool/Gauntlet/Framework/Base/Gauntlet.TestNode.cs
+++ b/Engine/Source/Programs/AutomationTool/Gauntlet/Framework/Base/Gauntlet.TestNode.cs
@@ -132,10 +132,14 @@ namespace Gauntlet
 		string GetTestSummary();
 
 		/// <summary>
-		/// Returns true if the test has encountered warnings. Test is expected to list any warnings it considers appropriate in the summary
+		/// Return a list of all warnings that should be reported. May be different than Summary contents
 		/// </summary>
-		bool HasWarnings { get;  }
+		IEnumerable<string> GetWarnings();
 
+		/// <summary>
+		/// Return a list of all errors that should be reported. May be different than Summary contents
+		/// </summary>
+		IEnumerable<string> GetErrors();
 
 		/// <summary>
 		/// Called to request any that any necessary cleanup be performed. After CleanupTest is called no further calls will be
diff --git a/Engine/Source/Programs/AutomationTool/Gauntlet/Framework/Gauntlet.TestExecutor.cs b/Engine/Source/Programs/AutomationTool/Gauntlet/Framework/Gauntlet.TestExecutor.cs
--- a/Engine/Source/Programs/AutomationTool/Gauntlet/Framework/Gauntlet.TestExecutor.cs
+++ b/Engine/Source/Programs/AutomationTool/Gauntlet/Framework/Gauntlet.TestExecutor.cs
@@ -523,7 +523,7 @@ namespace Gauntlet
 			}
 			else
 			{
-				if (TestInfo.TestNode.HasWarnings)
+				if (TestInfo.TestNode.GetWarnings().Any())
 				{
 					Log.Warning("{0} result={1}", TestInfo, TestInfo.FinalResult);
 				}
@@ -533,6 +533,11 @@ namespace Gauntlet
 				}
 			}
 			Summary.Split('\n').ToList().ForEach(L => Log.Info("  " + L));
+
+			TestInfo.TestNode.GetErrors().ToList().ForEach(E => Log.Error("{0}", E));
+
+			TestInfo.TestNode.GetWarnings().ToList().ForEach(E => Log.Warning("{0}", E));
+
 		}
 
 		/// <summary>
@@ -550,7 +555,7 @@ namespace Gauntlet
 
 			int TestCount = AllInfo.Count();
 			int FailedCount = AllInfo.Where(T => T.FinalResult != TestResult.Passed).Count();
-			int WarningCount = AllInfo.Where(T => T.TestNode.HasWarnings).Count();
+			int WarningCount = AllInfo.Where(T => T.TestNode.GetWarnings().Any()).Count();
 
 			var SortedInfo = AllInfo;
 
@@ -564,7 +569,7 @@ namespace Gauntlet
 						return 10;
 					}
 
-					if (T.TestNode.HasWarnings)
+					if (T.TestNode.GetWarnings().Any())
 					{
 						return 5;
 					}
@@ -582,7 +587,7 @@ namespace Gauntlet
 			List<string> TestResults = new List<string>();
 			foreach (TestExecutionInfo Info in SortedInfo)
 			{
-				string WarningString = Info.TestNode.HasWarnings ? " With Warnings" : "";
+				string WarningString = Info.TestNode.GetWarnings().Any() ? " With Warnings" : "";
 				TestResults.Add(string.Format("\t{0} result={1}{2}", Info, Info.FinalResult, WarningString));
 			}
 
diff --git a/Engine/Source/Programs/AutomationTool/Gauntlet/Framework/Utils/Gauntlet.MarkdownBuilder.cs b/Engine/Source/Programs/AutomationTool/Gauntlet/Framework/Utils/Gauntlet.MarkdownBuilder.cs
--- a/Engine/Source/Programs/AutomationTool/Gauntlet/Framework/Utils/Gauntlet.MarkdownBuilder.cs
+++ b/Engine/Source/Programs/AutomationTool/Gauntlet/Framework/Utils/Gauntlet.MarkdownBuilder.cs
@@ -20,6 +20,11 @@ namespace Gauntlet
 			SB = new StringBuilder();
 		}
 
+		public MarkdownBuilder(string Contents)
+		{
+			SB = new StringBuilder(Contents);
+		}
+
 		/// <summary>
 		/// Returns our formatted text
 		/// </summary>
diff --git a/Engine/Source/Programs/AutomationTool/Gauntlet/Gauntlet.Automation.csproj b/Engine/Source/Programs/AutomationTool/Gauntlet/Gauntlet.Automation.csproj
--- a/Engine/Source/Programs/AutomationTool/Gauntlet/Gauntlet.Automation.csproj
+++ b/Engine/Source/Programs/AutomationTool/Gauntlet/Gauntlet.Automation.csproj
@@ -133,6 +133,7 @@
     <Compile Include="Unreal\Tasks\Gauntlet.Task.CleanDevices.cs" />
     <Compile Include="Unreal\Game\Tests\Gauntlet.UnrealTest.ErrorTest.cs" />
     <Compile Include="Unreal\Game\UE4Game.TestConfig.cs" />
+    <Compile Include="Unreal\Utils\Gauntlet.AutomationLogParser.cs" />
     <Compile Include="Unreal\Utils\Gauntlet.UnrealHelpers.cs" />
     <Compile Include="Unreal\Utils\Gauntlet.UnrealLogParser.cs" />
     <Compile Include="Unreal\BuildSource\Gauntlet.StagedBuildSource.cs" />
diff --git a/Engine/Source/Programs/AutomationTool/Gauntlet/SelfTest/Framework/Gauntlet.SelfTest.OrderOfOpsTest.cs b/Engine/Source/Programs/AutomationTool/Gauntlet/SelfTest/Framework/Gauntlet.SelfTest.OrderOfOpsTest.cs
--- a/Engine/Source/Programs/AutomationTool/Gauntlet/SelfTest/Framework/Gauntlet.SelfTest.OrderOfOpsTest.cs
+++ b/Engine/Source/Programs/AutomationTool/Gauntlet/SelfTest/Framework/Gauntlet.SelfTest.OrderOfOpsTest.cs
@@ -84,6 +84,7 @@ namespace Gauntlet.SelfTest
 			return "";
 		}
 
+
 		public TestStatus GetTestStatus()
 		{
 			SetNewStage(TestStages.Status);
@@ -96,6 +97,24 @@ namespace Gauntlet.SelfTest
 			return new ITestNode[0];
 		}*/
 
+		/// <summary>
+		/// Return list of warnings. Empty by default
+		/// </summary>
+		/// <returns></returns>
+		public virtual IEnumerable<string> GetWarnings()
+		{
+			return new string[0];
+		}
+
+		/// <summary>
+		/// Return list of errors. Empty by default
+		/// </summary>
+		/// <returns></returns>
+		public virtual IEnumerable<string> GetErrors()
+		{
+			return new string[0];
+		}
+
 		public bool IsReadyToStart()
 		{
 			SetNewStage(TestStages.IsTestReady);
diff --git a/Engine/Source/Programs/AutomationTool/Gauntlet/SelfTest/Gauntlet.SelfTest.BaseNode.cs b/Engine/Source/Programs/AutomationTool/Gauntlet/SelfTest/Gauntlet.SelfTest.BaseNode.cs
--- a/Engine/Source/Programs/AutomationTool/Gauntlet/SelfTest/Gauntlet.SelfTest.BaseNode.cs
+++ b/Engine/Source/Programs/AutomationTool/Gauntlet/SelfTest/Gauntlet.SelfTest.BaseNode.cs
@@ -37,16 +37,36 @@ namespace Gauntlet.SelfTest
 			return InnerResult;
 		}
 
+
+		public TestStatus GetTestStatus()
+		{
+			return InnerStatus;
+		}
+
 		public string GetTestSummary()
 		{
 			return string.Format("{0} {1}", Name, GetTestResult());
 		}
 
-		public TestStatus GetTestStatus()
+		/// <summary>
+		/// Return list of warnings. Empty by default
+		/// </summary>
+		/// <returns></returns>
+		public virtual IEnumerable<string> GetWarnings()
 		{
-			return InnerStatus;
+			return new string[0];
+		}
+
+		/// <summary>
+		/// Return list of errors. Empty by default
+		/// </summary>
+		/// <returns></returns>
+		public virtual IEnumerable<string> GetErrors()
+		{
+			return new string[0];
 		}
 
+
 		public virtual ITestNode[] GetSubTests()
 		{
 			return new ITestNode[0];
diff --git a/Engine/Source/Programs/AutomationTool/Gauntlet/Unreal/Base/Gauntlet.UnrealTestNode.cs b/Engine/Source/Programs/AutomationTool/Gauntlet/Unreal/Base/Gauntlet.UnrealTestNode.cs
--- a/Engine/Source/Programs/AutomationTool/Gauntlet/Unreal/Base/Gauntlet.UnrealTestNode.cs
+++ b/Engine/Source/Programs/AutomationTool/Gauntlet/Unreal/Base/Gauntlet.UnrealTestNode.cs
@@ -33,17 +33,32 @@ namespace Gauntlet
 		public override TestPriority Priority { get { return GetPriority(); } }
 
 		/// <summary>
-		/// Returns true if the test has warnings. At this time we only consider Ensures() to be warnings by default
+		/// Returns Warnings found during tests. By default only ensures are considered
 		/// </summary>
-		public override bool HasWarnings { get {
-				if (SessionArtifacts == null)
-				{
-					return false;
-				}
+		public override IEnumerable<string> GetWarnings()
+		{
+			if (SessionArtifacts == null)
+			{
+				return new string[0];
+			}
 
-				bool HaveEnsures = SessionArtifacts.Any(A => A.LogSummary.Ensures.Count() > 0);
-				return HaveEnsures;
+			return SessionArtifacts.SelectMany(A =>
+			{
+				return A.LogSummary.Ensures.Select(E => E.Message);
+			}); 
+		}
+
+		/// <summary>
+		/// Returns Errors found during tests. By default only fatal errors are considered
+		/// </summary>
+		public override IEnumerable<string> GetErrors()
+		{
+			if (SessionArtifacts == null)
+			{
+				return new string[0];
 			}
+
+			return SessionArtifacts.Where(A => A.LogSummary.FatalError != null).Select(A => A.LogSummary.FatalError.Message);
 		}
 
 		// Begin UnrealTestNode properties and members
@@ -1030,40 +1045,22 @@ namespace Gauntlet
 		}
 
 		/// <summary>
-		/// Returns a summary of this test
+		/// Return the header for the test summary. The header is the first block of text and will be
+		/// followed by the summary of each individual role in the test
 		/// </summary>
 		/// <returns></returns>
-		public override string GetTestSummary()
+		protected virtual string GetTestSummaryHeader()
 		{
-
 			int AbnormalExits = 0;
 			int FatalErrors = 0;
 			int Ensures = 0;
 			int Errors = 0;
 			int Warnings = 0;
 
-			// Handle case where there aren't any session artifacts, for example with device starvation
-			if (SessionArtifacts == null)
-			{
-				return "NoSummary";
-			}
-
-			StringBuilder SB = new StringBuilder();
-
-			// Get any artifacts with failures
-			var FailureArtifacts = GetArtifactsWithFailures();
-
-			// Any with warnings (ensures)
-			var WarningArtifacts = SessionArtifacts.Where(A => A.LogSummary.Ensures.Count() > 0);
-
-			// combine artifacts into order as Failures, Warnings, Other
-			var AllArtifacts = FailureArtifacts.Union(WarningArtifacts);
-			AllArtifacts = AllArtifacts.Union(SessionArtifacts);
-
 			// create a quicck summary of total failures, ensures, errors, etc
-			foreach ( var Artifact in AllArtifacts)
+			foreach (var Artifact in SessionArtifacts)
 			{
-				string Summary = "NoSummary";
+				string Summary;
 				int ExitCode = GetRoleSummary(Artifact, out Summary);
 
 				if (ExitCode != 0 && Artifact.AppInstance.WasKilled == false)
@@ -1071,12 +1068,6 @@ namespace Gauntlet
 					AbnormalExits++;
 				}
 
-				if (SB.Length > 0)
-				{
-					SB.AppendLine();
-				}
-				SB.Append(Summary);
-
 				FatalErrors += Artifact.LogSummary.FatalError != null ? 1 : 0;
 				Ensures += Artifact.LogSummary.Ensures.Count();
 				Errors += Artifact.LogSummary.Errors.Count();
@@ -1087,6 +1078,8 @@ namespace Gauntlet
 
 			string WarningStatement = HasWarnings ? " With Warnings" : "";
 
+			var FailureArtifacts = GetArtifactsWithFailures();
+
 			// Create a summary
 			MB.H2(string.Format("{0} {1}{2}", Name, GetTestResult(), WarningStatement));
 
@@ -1113,14 +1106,56 @@ namespace Gauntlet
 			MB.Paragraph(string.Format("FatalErrors: {0}, Ensures: {1}, Errors: {2}, Warnings: {3}", FatalErrors, Ensures, Errors, Warnings));
 			MB.Paragraph(string.Format("ResultHash: {0}", GetTestResultHash()));
 			//MB.Paragraph(string.Format("Artifacts: {0}", CachedArtifactPath));
-			MB.Append("--------");
-			MB.Append(SB.ToString());
 
-			//SB.Clear();
-			//SB.AppendLine("begin: stack for UAT");
-			//SB.Append(MB.ToString());
-			//SB.Append("end: stack for UAT");
 			return MB.ToString();
 		}
+
+		/// <summary>
+		/// Returns a summary of this test
+		/// </summary>
+		/// <returns></returns>
+		public override string GetTestSummary()
+		{
+			// Handle case where there aren't any session artifacts, for example with device starvation
+			if (SessionArtifacts == null)
+			{
+				return "NoSummary";
+			}
+
+			MarkdownBuilder ReportBuilder = new MarkdownBuilder();
+
+			// add header
+			ReportBuilder.Append(GetTestSummaryHeader());
+			ReportBuilder.Append("--------");
+
+			StringBuilder SB = new StringBuilder();
+
+			// Get any artifacts with failures
+			var FailureArtifacts = GetArtifactsWithFailures();
+
+			// Any with warnings (ensures)
+			var WarningArtifacts = SessionArtifacts.Where(A => A.LogSummary.Ensures.Count() > 0);
+
+			// combine artifacts into order as Failures, Warnings, Other
+			var AllArtifacts = FailureArtifacts.Union(WarningArtifacts);
+			AllArtifacts = AllArtifacts.Union(SessionArtifacts);
+
+			// Add a summary of each 
+			foreach ( var Artifact in AllArtifacts)
+			{
+				string Summary = "NoSummary";
+				int ExitCode = GetRoleSummary(Artifact, out Summary);
+
+				if (SB.Length > 0)
+				{
+					SB.AppendLine();
+				}
+				SB.Append(Summary);
+			}
+
+			ReportBuilder.Append(SB.ToString());
+
+			return ReportBuilder.ToString();
+		}
 	}
 }
\ No newline at end of file
diff --git a/Engine/Source/Programs/AutomationTool/Gauntlet/Unreal/Utils/Gauntlet.AutomationLogParser.cs b/Engine/Source/Programs/AutomationTool/Gauntlet/Unreal/Utils/Gauntlet.AutomationLogParser.cs
new file mode 100644
--- /dev/null
+++ b/Engine/Source/Programs/AutomationTool/Gauntlet/Unreal/Utils/Gauntlet.AutomationLogParser.cs
@@ -0,0 +1,123 @@
+﻿using System;
+using System.Collections.Generic;
+using System.Linq;
+using System.Text;
+using System.Text.RegularExpressions;
+using System.Threading.Tasks;
+
+namespace Gauntlet
+{
+	/// <summary>
+	/// Result for a single automation result
+	/// </summary>
+	public class AutomationTestResult
+	{
+		/// <summary>
+		/// Short friendly name of this test
+		/// </summary>
+		public string		DisplayName;
+
+		/// <summary>
+		/// Full test name
+		/// </summary>
+		public string		TestName;
+
+		/// <summary>
+		/// True if the test passed
+		/// </summary>
+		public bool			Passed;
+
+		/// <summary>
+		/// Events logged during the test. Should contain errors if the test failed
+		/// </summary>
+		public List<string> Events;
+
+		/// <summary>
+		/// Result contstructor
+		/// </summary>
+		/// <param name="InDisplayName"></param>
+		/// <param name="InTestName"></param>
+		/// <param name="bInPassed"></param>
+		public AutomationTestResult(string InDisplayName, string InTestName, bool bInPassed)
+		{
+			DisplayName = InDisplayName;
+			TestName = InTestName;
+			Passed = bInPassed;
+			Events = new List<string>();
+		}
+	};
+
+	/// <summary>
+	/// Helper class for parsing AutomationTest results from either an UnrealLogParser or log contents
+	/// </summary>
+	public class AutomationLogParser
+	{
+		protected UnrealLogParser Parser;
+
+		/// <summary>
+		/// Constructor that uses an existing log parser
+		/// </summary>
+		/// <param name="InParser"></param>
+		public AutomationLogParser(UnrealLogParser InParser)
+		{
+			Parser = InParser;
+		}
+
+		/// <summary>
+		/// Constructor that takes raw log contents
+		/// </summary>
+		/// <param name="InContents"></param>
+		public AutomationLogParser(string InContents)
+		{
+			Parser = new UnrealLogParser(InContents);
+		}
+		
+		/// <summary>
+		/// Returns all results found in our construction content.
+		/// </summary>
+		/// <returns></returns>
+		public IEnumerable<AutomationTestResult> GetResults()
+		{
+			// Find all automation results that succeeded/failed
+			// [2019.04.30-18.49.51:329][244]LogAutomationController: Display: Automation Test Succeeded (ST_PR04 - Project.Functional Tests./Game/Tests/Rendering/PlanarReflection.ST_PR04)
+			// [2019.04.30-18.49.51:330] [244] LogAutomationController: BeginEvents: Project.Functional Tests./Game/Tests/Rendering/PlanarReflection.ST_PR04
+			// [2019.04.30 - 18.49.51:331][244] LogAutomationController: Screenshot 'ST_PR04' was similar!  Global Difference = 0.001377, Max Local Difference = 0.037953
+			// [2019.04.30 - 18.49.51:332][244]LogAutomationController: EndEvents: Project.Functional Tests./Game/Tests/Rendering/PlanarReflection.ST_PR04
+			IEnumerable<Match> TestMatches = Parser.GetAllMatches(@"LogAutomationController.+Test\s+(Succeeded|Failed)\s+\((.+)\s+-\s+(.+)\)");
+
+			string[] AutomationChannel = Parser.GetLogChannel("AutomationController").ToArray();
+
+			// Convert these lines into results by parsing out the details and then adding the events
+			IEnumerable<AutomationTestResult> Results = TestMatches.Select(M =>
+			{
+				bool TestPassed = M.Groups[1].ToString().ToLower() == "succeeded" ? true : false;
+				string DisplayName = M.Groups[2].ToString();
+				string TestName = M.Groups[3].ToString();
+
+				AutomationTestResult Result = new AutomationTestResult(DisplayName, TestName, TestPassed);
+
+				string EventName = string.Format("BeginEvents: {0}", TestName);
+				int EventIndex = Array.FindIndex(AutomationChannel, S => S.Contains(EventName)) + 1;
+
+				if (EventIndex > 0)
+				{
+					while (EventIndex < AutomationChannel.Length)
+					{
+						string Event = AutomationChannel[EventIndex++];
+
+						if (Event.Contains("EndEvents"))
+						{
+							break;
+						}
+
+						Result.Events.Add(Event);
+					}
+				}
+
+				return Result;
+			});
+
+			return Results;
+		}
+	}
+}
\ No newline at end of file
diff --git a/Engine/Source/Programs/AutomationTool/Gauntlet/Unreal/Utils/Gauntlet.UnrealLogParser.cs b/Engine/Source/Programs/AutomationTool/Gauntlet/Unreal/Utils/Gauntlet.UnrealLogParser.cs
--- a/Engine/Source/Programs/AutomationTool/Gauntlet/Unreal/Utils/Gauntlet.UnrealLogParser.cs
+++ b/Engine/Source/Programs/AutomationTool/Gauntlet/Unreal/Utils/Gauntlet.UnrealLogParser.cs
@@ -210,29 +210,42 @@ namespace Gauntlet
 			return Info;
 		}
 
-
 		/// <summary>
 		/// Return all entries for the specified channel. E.g. "OrionGame" will
 		/// return all entries starting with LogOrionGame
 		/// </summary>
 		/// <param name="Channel"></param>
 		/// <returns></returns>
-		public IEnumerable<string> GetLogChannel(string Channel, bool ExactMatch=true)
+		public IEnumerable<string> GetLogChannels(IEnumerable<string> Channels, bool ExactMatch = true)
 		{
+			// expand Chan1, Chan2 into Log(?:Chan1|Chan2) for a non-capturing group
+			string Match = string.Format("Log(?:{0})", string.Join("|", Channels));
+
 			string Pattern;
 
 			if (ExactMatch)
 			{
-				Pattern = string.Format(@"(Log{0}:\s{{0,1}}.+)", Channel);
+				Pattern = string.Format(@"({0}:\s{{0,1}}.+)", Match);
 			}
 			else
 			{
-				Pattern = string.Format(@"(Log{0}.*:\s{{0,1}}.+)", Channel);
+				Pattern = string.Format(@"({0}.*:\s{{0,1}}.+)", Match);
 			}
 
 			return Regex.Matches(Content, Pattern).Cast<Match>().Select(M => M.Groups[1].ToString()).ToArray();
 		}
 
+		/// <summary>
+		/// Return all entries for the specified channel. E.g. "OrionGame" will
+		/// return all entries starting with LogOrionGame
+		/// </summary>
+		/// <param name="Channel"></param>
+		/// <returns></returns>
+		public IEnumerable<string> GetLogChannel(string Channel, bool ExactMatch=true)
+		{
+			return GetLogChannels(new string[] { Channel }, ExactMatch);
+		}
+
 		/// <summary>
 		/// Returns all warnings from the log
 		/// </summary>
