Fix for CrashReportTest errors causing a false positive

diff --git a/Engine/Source/Programs/AutomationTool/Gauntlet/Unreal/Base/Gauntlet.UnrealTestNode.cs b/Engine/Source/Programs/AutomationTool/Gauntlet/Unreal/Base/Gauntlet.UnrealTestNode.cs
--- a/Engine/Source/Programs/AutomationTool/Gauntlet/Unreal/Base/Gauntlet.UnrealTestNode.cs
+++ b/Engine/Source/Programs/AutomationTool/Gauntlet/Unreal/Base/Gauntlet.UnrealTestNode.cs
@@ -58,7 +58,7 @@ namespace Gauntlet
 				return new string[0];
 			}
 
-			return SessionArtifacts.Where(A => A.LogSummary.FatalError != null).Select(A => A.LogSummary.FatalError.Message);
+			return GetArtifactsWithFailures().Where(A => A.LogSummary.FatalError != null).Select(A => A.LogSummary.FatalError.Message);
 		}
 
 		// Begin UnrealTestNode properties and members
