Fixed missing copyright notice

diff --git a/Engine/Source/Programs/AutomationTool/Gauntlet/Unreal/Utils/Gauntlet.AutomationLogParser.cs b/Engine/Source/Programs/AutomationTool/Gauntlet/Unreal/Utils/Gauntlet.AutomationLogParser.cs
--- a/Engine/Source/Programs/AutomationTool/Gauntlet/Unreal/Utils/Gauntlet.AutomationLogParser.cs
+++ b/Engine/Source/Programs/AutomationTool/Gauntlet/Unreal/Utils/Gauntlet.AutomationLogParser.cs
@@ -1,4 +1,6 @@
-﻿using System;
+﻿// Copyright 1998-2019 Epic Games, Inc. All Rights Reserved.
+
+using System;
 using System.Collections.Generic;
 using System.Linq;
 using System.Text;
