By default only consider failed artifacts when looking for fatal errors.

diff --git a/Engine/Source/Programs/AutomationTool/Gauntlet/Unreal/Base/Gauntlet.UnrealTestNode.cs b/Engine/Source/Programs/AutomationTool/Gauntlet/Unreal/Base/Gauntlet.UnrealTestNode.cs
--- a/Engine/Source/Programs/AutomationTool/Gauntlet/Unreal/Base/Gauntlet.UnrealTestNode.cs
+++ b/Engine/Source/Programs/AutomationTool/Gauntlet/Unreal/Base/Gauntlet.UnrealTestNode.cs
@@ -58,7 +58,9 @@ namespace Gauntlet
 				return new string[0];
 			}
 
-			return GetArtifactsWithFailures().Where(A => A.LogSummary.FatalError != null).Select(A => A.LogSummary.FatalError.Message);
+			var FailedArtifacts = GetArtifactsWithFailures();
+
+			return FailedArtifacts.Where(A => A.LogSummary.FatalError != null).Select(A => A.LogSummary.FatalError.Message);
 		}
 
 		// Begin UnrealTestNode properties and members
