Updating general ExportMcpTemplates command to support overriding the default commandlet to use.

diff --git a/Engine/Source/Programs/AutomationTool/Scripts/Localisation.Automation.cs b/Engine/Source/Programs/AutomationTool/Scripts/Localisation.Automation.cs
--- a/Engine/Source/Programs/AutomationTool/Scripts/Localisation.Automation.cs
+++ b/Engine/Source/Programs/AutomationTool/Scripts/Localisation.Automation.cs
@@ -515,7 +515,7 @@ public class ExportMcpTemplates : BuildCommand
 		return Path.Combine(ProjectFile.Directory.FullName, "Content", "Backend");
 	}
 
-	public static void RunExportTemplates(FileReference ProjectFile, bool bCheckoutAndSubmit, bool bOnlyLoc, bool bbNoRobomerge, string OutputFolder = null)
+	public static void RunExportTemplates(FileReference ProjectFile, bool bCheckoutAndSubmit, bool bOnlyLoc, bool bbNoRobomerge, string CommandletOverride)
 	{
 		string EditorExe = "UE4Editor.exe";
 		EditorExe = HostPlatform.Current.GetUE4ExePath(EditorExe);
@@ -529,21 +529,6 @@ public class ExportMcpTemplates : BuildCommand
 		string FolderToGenerateIn = GameBackendFolder;
 
 		string Parameters = "-GenerateLoc";
-		if (!string.IsNullOrEmpty(OutputFolder))
-		{
-			FolderToGenerateIn = OutputFolder;
-
-			CommandUtils.DeleteDirectory_NoExceptions(OutputFolder);
-			CommandUtils.CopyDirectory_NoExceptions(GameBackendFolder, OutputFolder);
-
-			string[] CopiedFiles = Directory.GetFiles(OutputFolder, "*.*", SearchOption.AllDirectories);
-			foreach (string File in CopiedFiles)
-			{
-				CommandUtils.SetFileAttributes_NoExceptions(File, false);
-			}
-
-			Parameters += String.Format(" -OUT={0}", OutputFolder);
-		}
 
 		int WorkingCL = -1;
 		if (bCheckoutAndSubmit)
@@ -576,7 +561,8 @@ public class ExportMcpTemplates : BuildCommand
 			CommandUtils.P4.Edit(WorkingCL, FolderToGenerateIn + "/...");
 		}
 
-		CommandUtils.RunCommandlet(ProjectFile, EditorExe, "ExportTemplatesCommandlet", Parameters);
+		string Commandlet = string.IsNullOrWhiteSpace(CommandletOverride) ? "ExportTemplatesCommandlet" : CommandletOverride;
+		CommandUtils.RunCommandlet(ProjectFile, EditorExe, Commandlet, Parameters);
 
 		if (WorkingCL > 0)
 		{
@@ -585,13 +571,13 @@ public class ExportMcpTemplates : BuildCommand
 			if (bOnlyLoc)
 			{
 				// Revert all folders and files except GeneratedLoc.json
-				foreach (var DirPath in Directory.GetDirectories(FolderToGenerateIn))
+				foreach (string DirPath in Directory.GetDirectories(FolderToGenerateIn))
 				{
 					DirectoryInfo Dir = new DirectoryInfo(DirPath);
 					CommandUtils.P4.Revert(WorkingCL, FolderToGenerateIn + "/" + Dir.Name + "/...");
 				}
 
-				foreach (var FilePath in Directory.GetFiles(FolderToGenerateIn))
+				foreach (string FilePath in Directory.GetFiles(FolderToGenerateIn))
 				{
 					FileInfo File = new FileInfo(FilePath);
 					if (File.Name != "GeneratedLoc.json")
@@ -618,7 +604,8 @@ public class ExportMcpTemplates : BuildCommand
 		FileReference ProjectFile = new FileReference(CombinePaths(CmdEnv.LocalRoot, ProjectName, String.Format("{0}.uproject", ProjectName)));
 		bool bOnlyLoc = ParseParam("OnlyLoc");
 		bool bNoRobomerge = ParseParam("NoRobomerge");
-		RunExportTemplates(ProjectFile, true, bOnlyLoc, bNoRobomerge);
+		string CommandletOverride = ParseParamValue("Commandlet", null);
+		RunExportTemplates(ProjectFile, true, bOnlyLoc, bNoRobomerge, CommandletOverride);
 	}
 }
 
