Gauntlet: Add platform independent mapping to root drive paths
#rnx

diff --git a/Engine/Source/Programs/AutomationTool/Gauntlet/Unreal/Utils/Gauntlet.UnrealHelpers.cs b/Engine/Source/Programs/AutomationTool/Gauntlet/Unreal/Utils/Gauntlet.UnrealHelpers.cs
--- a/Engine/Source/Programs/AutomationTool/Gauntlet/Unreal/Utils/Gauntlet.UnrealHelpers.cs
+++ b/Engine/Source/Programs/AutomationTool/Gauntlet/Unreal/Utils/Gauntlet.UnrealHelpers.cs
@@ -360,4 +360,43 @@ namespace Gauntlet
 			return "";
 		}
 	}
+
+	/// <summary>
+	/// Automatically maps root drive to proper platform specific path
+	/// </summary>
+	public class EpicRoot
+	{
+		string PlatformPath;
+
+		public EpicRoot(string Path)
+		{
+			PlatformPath = Path;
+
+			if (BuildHostPlatform.Current.Platform == UnrealTargetPlatform.Mac)
+			{
+				if (!Path.Contains("P:"))
+				{
+					PlatformPath = Regex.Replace(Path, @"\\\\epicgames.net\\root", "/Volumes/epicgames.net/root", RegexOptions.IgnoreCase);
+				}
+				else 
+				{
+					PlatformPath = Regex.Replace(Path, "P:", "/Volumes/epicgames.net/root", RegexOptions.IgnoreCase);					
+				}				
+				
+				PlatformPath = PlatformPath.Replace(@"\", "/");
+			}
+
+		}
+
+		public static implicit operator string(EpicRoot Path)
+		{
+			return Path.PlatformPath;
+		}
+		public static implicit operator EpicRoot(string Path)
+		{
+			return new EpicRoot(Path);
+		}
+
+	}
+
 }
\ No newline at end of file
