Added copyright boilerplate to new source file

diff --git a/Engine/Source/Programs/AutomationTool/AutomationUtils/BundleUtils.cs b/Engine/Source/Programs/AutomationTool/AutomationUtils/BundleUtils.cs
--- a/Engine/Source/Programs/AutomationTool/AutomationUtils/BundleUtils.cs
+++ b/Engine/Source/Programs/AutomationTool/AutomationUtils/BundleUtils.cs
@@ -1,4 +1,5 @@
-﻿using System;
+﻿// Copyright 1998-2019 Epic Games, Inc. All Rights Reserved.
+using System;
 using System.Collections.Generic;
 using System.Linq;
 using System.Text;
