Do not output an additional "error:" line when a commandlet fails. This is almost always covered by other errors already in the log, and just creates additional noise.

diff --git a/Engine/Source/Programs/AutomationTool/AutomationUtils/CommandletUtils.cs b/Engine/Source/Programs/AutomationTool/AutomationUtils/CommandletUtils.cs
--- a/Engine/Source/Programs/AutomationTool/AutomationUtils/CommandletUtils.cs
+++ b/Engine/Source/Programs/AutomationTool/AutomationUtils/CommandletUtils.cs
@@ -419,7 +419,7 @@ namespace AutomationTool
 						ExitCodeDesc = String.Format(" (signal {0})", RunResult.ExitCode - 128);
 					}
 				}
-				throw new CommandletException(DestLogFile, RunResult.ExitCode, "Editor terminated with exit code {0}{1} while running {2}{3}; see log {4}", RunResult.ExitCode, ExitCodeDesc, Commandlet, (ProjectName == null)? "" : String.Format(" for {0}", ProjectName), DestLogFile);
+				throw new CommandletException(DestLogFile, RunResult.ExitCode, "Editor terminated with exit code {0}{1} while running {2}{3}; see log {4}", RunResult.ExitCode, ExitCodeDesc, Commandlet, (ProjectName == null)? "" : String.Format(" for {0}", ProjectName), DestLogFile) { OutputFormat = AutomationExceptionOutputFormat.Minimal };
 			}
 		}
 		
