Gauntlet: Adding hardware model constraints and worker job id
#rnx

diff --git a/Engine/Source/Programs/AutomationTool/Gauntlet/Framework/Devices/Gauntlet.DevicePool.cs b/Engine/Source/Programs/AutomationTool/Gauntlet/Framework/Devices/Gauntlet.DevicePool.cs
--- a/Engine/Source/Programs/AutomationTool/Gauntlet/Framework/Devices/Gauntlet.DevicePool.cs
+++ b/Engine/Source/Programs/AutomationTool/Gauntlet/Framework/Devices/Gauntlet.DevicePool.cs
@@ -43,6 +43,8 @@ namespace Gauntlet
 
 		public EPerfSpec PerfSpec;
 
+		public string Model = string.Empty;
+
 		public string Available;
 		
 		public bool RemoveOnShutdown;
@@ -58,13 +60,15 @@ namespace Gauntlet
 	/// </summary>
 	public class UnrealTargetConstraint : IEquatable<UnrealTargetConstraint>
 	{
-		public UnrealTargetPlatform? Platform;
-		public EPerfSpec PerfSpec;
+		public readonly UnrealTargetPlatform? Platform;
+		public readonly EPerfSpec PerfSpec;
+		public readonly string Model;
 
-		public UnrealTargetConstraint(UnrealTargetPlatform? Platform, EPerfSpec PerfSpec = EPerfSpec.Unspecified)
+		public UnrealTargetConstraint(UnrealTargetPlatform? Platform, EPerfSpec PerfSpec = EPerfSpec.Unspecified, string Model = null)
 		{
-			this.Platform = Platform;
+			this.Platform = Platform;			
 			this.PerfSpec = PerfSpec;
+			this.Model = Model == null ? string.Empty : Model;
 		}
 
 		/// <summary>
@@ -72,7 +76,7 @@ namespace Gauntlet
 		/// </summary>
 		public bool IsIdentity()
 		{
-			return PerfSpec == EPerfSpec.Unspecified;
+			return (PerfSpec == EPerfSpec.Unspecified) && (Model == string.Empty);
 		}
 
 
@@ -84,6 +88,11 @@ namespace Gauntlet
 			return Platform == Device.Platform && (IsIdentity() || this == DevicePool.Instance.GetConstraint(Device));
 		}
 
+		public bool Check(DeviceDefinition DeviceDef)
+		{
+			bool Match = (PerfSpec == EPerfSpec.Unspecified) ? DeviceDef.Model == Model : PerfSpec == DeviceDef.PerfSpec;
+			return Platform == DeviceDef.Platform && (IsIdentity() || Match );
+		}
 
 		public bool Equals(UnrealTargetConstraint Other)
 		{
@@ -93,7 +102,8 @@ namespace Gauntlet
 			}
 
 			if (ReferenceEquals(this, Other)) return true;
-			return Other.Platform == Platform && Other.PerfSpec == PerfSpec;
+
+			return Other.Platform == Platform && Other.PerfSpec == PerfSpec && Other.Model == Model;
 		}
 
 		public override bool Equals(object Obj)
@@ -125,7 +135,12 @@ namespace Gauntlet
 
 		public override string ToString()
 		{
-			return string.Format("{0}:{1}", Platform, PerfSpec);
+			if (PerfSpec == EPerfSpec.Unspecified && Model == string.Empty)
+			{
+				return string.Format("{0}", Platform);
+			}
+
+			return string.Format("{0}:{1}", Platform, PerfSpec == EPerfSpec.Unspecified ? Model : PerfSpec.ToString());
 		}
 
 		public override int GetHashCode()
@@ -569,6 +584,7 @@ namespace Gauntlet
 					Def.Name = Device.Name;
 					Def.Platform = DeviceMap.FirstOrDefault(Entry => Entry.Value == Device.Type).Key;
 					Def.DeviceData = Device.DeviceData;
+					Def.Model = string.Empty;
 
 					if (!String.IsNullOrEmpty(Device.PerfSpec) && !Enum.TryParse<EPerfSpec>(Device.PerfSpec, true, out Def.PerfSpec))
 					{
@@ -719,12 +735,12 @@ namespace Gauntlet
 						foreach (string DeviceRef in DevicesList)
 						{							
 
-							// check for <platform>:<address>:<port>. We pass the last two combined to the
-							// device constructor
+							// check for <platform>:<address>:<port>|<model>. We pass address:port to device constructor
 							Match M = Regex.Match(DeviceRef, @"(.+?):(.+)");
 
 							UnrealTargetPlatform DevicePlatform = DefaultPlatform;
 							string DeviceAddress = DeviceRef;
+							string Model = string.Empty;
 
 							// when using device service, skip adding local non-desktop devices to pool
 							bool IsDesktop = UnrealBuildTool.Utils.GetPlatformsInClass(UnrealPlatformClass.Desktop).Contains(DevicePlatform);							
@@ -739,14 +755,25 @@ namespace Gauntlet
 								{
 									throw new AutomationException("platform {0} is not a recognized device type", M.Groups[1].ToString());
 								}
+
 								DeviceAddress = M.Groups[2].ToString();
-							}
+
+								// parse device model
+								if (DeviceAddress.Contains("|"))
+								{
+									string[] Components = DeviceAddress.Split(new char[] { '|' });
+									DeviceAddress = Components[0];
+									Model = Components[1];
+								}
+
+							} 
 
 							Log.Info("Added device {0}:{1} to pool", DevicePlatform, DeviceAddress);
 							DeviceDefinition Def = new DeviceDefinition();
 							Def.Address = DeviceAddress;
 							Def.Name = DeviceAddress;
 							Def.Platform = DevicePlatform;
+							Def.Model = Model;
 							UnprovisionedDevices.Add(Def);
 						}
 					}
@@ -829,7 +856,7 @@ namespace Gauntlet
 				{
 					if (NewDevice != null)
 					{
-						RegisterDevice(NewDevice, new UnrealTargetConstraint(NewDevice.Platform, Def.PerfSpec));
+						RegisterDevice(NewDevice, new UnrealTargetConstraint(NewDevice.Platform, Def.PerfSpec, Def.Model));
 					}
 				}
 			}
@@ -861,7 +888,7 @@ namespace Gauntlet
 			lock (LockObject)
 			{
 				return AvailableDevices.Where(D => Validate == null ? Constraint.Check(D) : Validate(D)).Count() +
-					UnprovisionedDevices.Where(D => D.Platform == Constraint.Platform && Constraint.IsIdentity()).Count();
+					UnprovisionedDevices.Where(D => Constraint.Check(D)).Count();
 			}
 		}
 
@@ -874,7 +901,7 @@ namespace Gauntlet
 			lock (LockObject)
 			{
 				return AvailableDevices.Union(ReservedDevices).Where(D => Validate == null ? Constraint.Check(D) : Validate(D)).Count() +
-					UnprovisionedDevices.Where(D => D.Platform == Constraint.Platform && Constraint.IsIdentity()).Count();
+					UnprovisionedDevices.Where(D => Constraint.Check(D)).Count();
 			}
 		}
 
@@ -988,7 +1015,7 @@ namespace Gauntlet
 
 				// randomize the order of all devices that are of this platform
 				var MatchingProvisionedDevices = AvailableDevices.Where(D => Constraint.Check(D)).ToList();
-				var MatchingUnprovisionedDevices = UnprovisionedDevices.Where(D => D.Platform == Constraint.Platform && Constraint.IsIdentity()).ToList();
+				var MatchingUnprovisionedDevices = UnprovisionedDevices.Where(D => Constraint.Check(D)).ToList();
 
 				bool OutOfDevices = false;
 				bool ContinuePredicate = true;
diff --git a/Engine/Source/Programs/AutomationTool/Gauntlet/Unreal/Base/Gauntlet.UnrealTestContext.cs b/Engine/Source/Programs/AutomationTool/Gauntlet/Unreal/Base/Gauntlet.UnrealTestContext.cs
--- a/Engine/Source/Programs/AutomationTool/Gauntlet/Unreal/Base/Gauntlet.UnrealTestContext.cs
+++ b/Engine/Source/Programs/AutomationTool/Gauntlet/Unreal/Base/Gauntlet.UnrealTestContext.cs
@@ -364,6 +364,9 @@ namespace Gauntlet
 		// Begin ITestContext implementation
 		public UnrealBuildSource BuildInfo { get; private set; }
 
+		// Worker Job ID (generates unique node results and logs in parallel runs)
+		public string WorkerJobID;
+
 		/// <summary>
 		/// Global options for this test
 		/// </summary>
@@ -412,11 +415,17 @@ namespace Gauntlet
 		public string ToString(bool bWithServerType = false)
 		{
 			string Description = string.Format("{0}", RoleContext[UnrealTargetRole.Client]);
+
+			if (WorkerJobID != null)
+			{
+				Description += " " + WorkerJobID;
+			}
 				
 			if (bWithServerType)
 			{
 				Description += ", " +  RoleContext[UnrealTargetRole.Server].ToString();
 			}
+
 			return Description;
 		}
 	}
diff --git a/Engine/Source/Programs/AutomationTool/Gauntlet/Unreal/RunUnreal.cs b/Engine/Source/Programs/AutomationTool/Gauntlet/Unreal/RunUnreal.cs
--- a/Engine/Source/Programs/AutomationTool/Gauntlet/Unreal/RunUnreal.cs
+++ b/Engine/Source/Programs/AutomationTool/Gauntlet/Unreal/RunUnreal.cs
@@ -347,7 +347,15 @@ namespace Gauntlet
 					throw new AutomationException("Unable to convert perfspec '{0}' into an EPerfSpec", PerfSpec);
 				}
 
-				TestContext.Constraint = new UnrealTargetConstraint(UnrealPlatform, PerfSpec);
+				// parse hardware model
+				List<string> ModelArgs = CombinedParams.ParseValues("PerfModel", false);
+				string Model = ModelArgs.Count > 0 ? ModelArgs.Last() : string.Empty;
+
+				TestContext.Constraint = new UnrealTargetConstraint(UnrealPlatform, PerfSpec, Model);
+
+				// parse worker job id
+				List<string> WorkerJobIDArgs = CombinedParams.ParseValues("WorkerJobID", false);
+				TestContext.WorkerJobID = WorkerJobIDArgs.Count > 0 ? WorkerJobIDArgs.Last() : null;
 
 				TestContext.TestParams = CombinedParams;
 
