Throwing exception on failed shared cook installs. Cleaning up output on failure

diff --git a/Engine/Source/Programs/AutomationTool/Scripts/SharedCookedBuild.cs b/Engine/Source/Programs/AutomationTool/Scripts/SharedCookedBuild.cs
--- a/Engine/Source/Programs/AutomationTool/Scripts/SharedCookedBuild.cs
+++ b/Engine/Source/Programs/AutomationTool/Scripts/SharedCookedBuild.cs
@@ -104,7 +104,7 @@ public class SharedCookedBuild
 				ISharedCookedBuild Build = PlatformBuilds.First();
 				if (!Build.CopyBuild(InstallPath))
 				{
-					CommandUtils.LogError("Failed to copy shared build for {0} {1}", Build.Platform, Build.CL);
+					throw new AutomationException("Failed to copy shared build for {0} {1}", Build.Platform, Build.CL);
 				}
 			}
 		}
@@ -209,6 +209,7 @@ public class SharedCookedBuild
 		if (BestCL < 0)
 		{
 			CommandUtils.LogError("Could not locate valid shared cooked build for all target platforms");
+			CommandUtils.LogError("Current CL: {0}, Current Code CL: {1}", LocalSync.Changelist, LocalSync.CompatibleChangelist);
 		}
 
 		return CandidateBuilds.Where(x => x.CL == BestCL).ToList();
@@ -343,7 +344,7 @@ public class SharedCookedBuild
 				DirectoryReference.Delete(PlatformInstallPath, true);
 			}
 
-			IProcessResult Result = CommandUtils.Run(BPTI.FullName, string.Format("-Manifest={0} -OutputDir={1}", Manifest.FullName, PlatformInstallPath.FullName), null, CommandUtils.ERunOptions.Default);
+			IProcessResult Result = CommandUtils.Run(BPTI.FullName, string.Format("-Manifest={0} -OutputDir={1} -stdout -GenericConsoleOutput", Manifest.FullName, PlatformInstallPath.FullName), null, CommandUtils.ERunOptions.Default);
 			if (Result.ExitCode != 0)
 			{
 				CommandUtils.LogWarning("Failed to install manifest {0} to {1}", Manifest.FullName, PlatformInstallPath.FullName);
