Work around a problem with Mono being extremely slow when launched from UAT by running it through sh

diff --git a/Engine/Source/Programs/AutomationTool/AutomationUtils/MacHostPlatform.cs b/Engine/Source/Programs/AutomationTool/AutomationUtils/MacHostPlatform.cs
--- a/Engine/Source/Programs/AutomationTool/AutomationUtils/MacHostPlatform.cs
+++ b/Engine/Source/Programs/AutomationTool/AutomationUtils/MacHostPlatform.cs
@@ -118,8 +118,8 @@ namespace AutomationTool
 				else
 				{
 					// It's a C# app, so run it with Mono
-					CommandLine = "\"" + AppName + "\" " + (String.IsNullOrEmpty(CommandLine) ? "" : CommandLine);
-					AppName = "mono";
+					CommandLine = " -c 'mono \"" + AppName + "\" " + (String.IsNullOrEmpty(CommandLine) ? "" : CommandLine) + "'";
+					AppName = "sh";
 					Options &= ~CommandUtils.ERunOptions.AppMustExist;
 				}
 			}
