Back out changelist 6429753. Likely cause for IOS compile failure

diff --git a/Engine/Source/Programs/AutomationTool/AutomationUtils/MacHostPlatform.cs b/Engine/Source/Programs/AutomationTool/AutomationUtils/MacHostPlatform.cs
--- a/Engine/Source/Programs/AutomationTool/AutomationUtils/MacHostPlatform.cs
+++ b/Engine/Source/Programs/AutomationTool/AutomationUtils/MacHostPlatform.cs
@@ -118,8 +118,8 @@ namespace AutomationTool
 				else
 				{
 					// It's a C# app, so run it with Mono
-					CommandLine = " -c 'mono \"" + AppName + "\" " + (String.IsNullOrEmpty(CommandLine) ? "" : CommandLine) + "'";
-					AppName = "sh";
+					CommandLine = "\"" + AppName + "\" " + (String.IsNullOrEmpty(CommandLine) ? "" : CommandLine);
+					AppName = "mono";
 					Options &= ~CommandUtils.ERunOptions.AppMustExist;
 				}
 			}
