Gauntlet: Disable Vivox voice chat when running iOS tests, otherwise get blocking microphone permission request
#rnx

diff --git a/Engine/Source/Programs/AutomationTool/Gauntlet/Unreal/Game/EpicGame.TestConfig.cs b/Engine/Source/Programs/AutomationTool/Gauntlet/Unreal/Game/EpicGame.TestConfig.cs
--- a/Engine/Source/Programs/AutomationTool/Gauntlet/Unreal/Game/EpicGame.TestConfig.cs
+++ b/Engine/Source/Programs/AutomationTool/Gauntlet/Unreal/Game/EpicGame.TestConfig.cs
@@ -233,6 +233,13 @@ namespace EpicGame
 					Account UserAccount = AccountPool.Instance.ReserveAccount();
 					UserAccount.ApplyToConfig(AppConfig);
 				}
+
+				// turn off voice chat, otherwise will open blocking permission requests on mobile
+				if (ConfigRole.Platform == UnrealTargetPlatform.IOS)
+				{
+					AppConfig.CommandLine += " -ini:Engine:[VoiceChat.Vivox]:bEnabled=false";
+				}
+
 			}
 		}
 		
