Gauntlet: Don't attempt to reserve devices with model from service, reduce some device pool log severity
#rnx

diff --git a/Engine/Source/Programs/AutomationTool/Gauntlet/Framework/Devices/Gauntlet.DevicePool.cs b/Engine/Source/Programs/AutomationTool/Gauntlet/Framework/Devices/Gauntlet.DevicePool.cs
--- a/Engine/Source/Programs/AutomationTool/Gauntlet/Framework/Devices/Gauntlet.DevicePool.cs
+++ b/Engine/Source/Programs/AutomationTool/Gauntlet/Framework/Devices/Gauntlet.DevicePool.cs
@@ -543,7 +543,14 @@ namespace Gauntlet
                 if (!DeviceMap.ContainsKey(Entry.Key.Platform.Value))
 				{
 					// if an unsupported device, we can't reserve it
-					Log.Error("Unable to reserve service device of type: {0}", Entry.Key);
+					Log.Info("Unable to reserve service device of type: {0}", Entry.Key);
+					return false;
+				}
+
+				if (!string.IsNullOrEmpty(Entry.Key.Model))
+				{
+					// if specific device model, we can't currently reserve it from service
+					Log.Info("Unable to reserve service device of model: {0}", Entry.Key.Model);
 					return false;
 				}
 
@@ -629,8 +636,8 @@ namespace Gauntlet
 				return true;
 			}
 
-			Log.Error("Unable to reserve service devices:");
-			Devices.ForEach(Device => Log.Error("    Device: {0}", Device));
+			Log.Info("Unable to reserve service devices:");
+			Devices.ForEach(Device => Log.Info("    Device: {0}", Device));
 			return false;
 		}
 
