One more fix for Windows paths handling in MacHostPlatform.cs

diff --git a/Engine/Source/Programs/AutomationTool/AutomationUtils/MacHostPlatform.cs b/Engine/Source/Programs/AutomationTool/AutomationUtils/MacHostPlatform.cs
--- a/Engine/Source/Programs/AutomationTool/AutomationUtils/MacHostPlatform.cs
+++ b/Engine/Source/Programs/AutomationTool/AutomationUtils/MacHostPlatform.cs
@@ -118,7 +118,7 @@ namespace AutomationTool
 				else
 				{
 					// It's a C# app, so run it with Mono
-					CommandLine = " -c 'mono \"" + AppName + "\" " + (String.IsNullOrEmpty(CommandLine) ? "" : CommandLine) + "'";
+					CommandLine = " -c 'mono \"" + AppName + "\" " + (String.IsNullOrEmpty(CommandLine) ? "" : CommandLine.Replace("\\", "\\\\")) + "'";
 					AppName = "sh";
 					Options &= ~CommandUtils.ERunOptions.AppMustExist;
 				}
