Gauntlet: Use temp folder and include worker id for iOS IPA cache
#rnx

diff --git a/Engine/Source/Programs/AutomationTool/Gauntlet/Platform/IOS/Gauntlet.TargetDeviceIOS.cs b/Engine/Source/Programs/AutomationTool/Gauntlet/Platform/IOS/Gauntlet.TargetDeviceIOS.cs
--- a/Engine/Source/Programs/AutomationTool/Gauntlet/Platform/IOS/Gauntlet.TargetDeviceIOS.cs
+++ b/Engine/Source/Programs/AutomationTool/Gauntlet/Platform/IOS/Gauntlet.TargetDeviceIOS.cs
@@ -766,7 +766,13 @@ namespace Gauntlet
 		}
 		
 		// Gauntlet cache folder for tracking device/ipa state
-		string GauntletAppCache { get { return Path.Combine(Environment.GetFolderPath(Environment.SpecialFolder.Personal), ".gauntletappcache");} }
+		string GauntletAppCache
+		{
+			get
+			{	
+				return Path.Combine(Globals.TempDir, string.Format("IOSAppCache{0}", Globals.WorkerID == -1 ? "" : Globals.WorkerID.ToString()));
+			}
+		}
 
 		// path to locally extracted (and possibly resigned) app bundle
 		// Note: ios-deploy works with app bundles, which requires the IPA be unzipped for deployment (this will allow us to resign in the future as well)
