Only check for hacks in the engine directory.

diff --git a/Engine/Source/Programs/AutomationTool/Scripts/CheckForHacks.cs b/Engine/Source/Programs/AutomationTool/Scripts/CheckForHacks.cs
--- a/Engine/Source/Programs/AutomationTool/Scripts/CheckForHacks.cs
+++ b/Engine/Source/Programs/AutomationTool/Scripts/CheckForHacks.cs
@@ -76,7 +76,7 @@ namespace AutomationTool
 			List<FileInfo> FilesToCheck = new List<FileInfo>();
 			using (ThreadPoolWorkQueue Queue = new ThreadPoolWorkQueue())
 			{
-				DirectoryInfo BaseDir = new DirectoryInfo(RootDirectory.FullName);
+				DirectoryInfo BaseDir = new DirectoryInfo(EngineDirectory.FullName);
 				Queue.Enqueue(() => FindAllFiles(Queue, BaseDir, FilesToCheck));
 				Queue.Wait();
 			}
