Fix for discovering loose file shared cooked builds

diff --git a/Engine/Source/Programs/AutomationTool/Scripts/SharedCookedBuild.cs b/Engine/Source/Programs/AutomationTool/Scripts/SharedCookedBuild.cs
--- a/Engine/Source/Programs/AutomationTool/Scripts/SharedCookedBuild.cs
+++ b/Engine/Source/Programs/AutomationTool/Scripts/SharedCookedBuild.cs
@@ -249,7 +249,7 @@ public class SharedCookedBuild
 
 		// Search for all available builds
 		const string MetaDataFilename = "\\Metadata\\DevelopmentAssetRegistry.bin";
-		string BuildRule = Path + MetaDataFilename;
+		string BuildRule = SplitPath.Item2 + MetaDataFilename;
 		BuildRule = BuildRule.Replace("[BRANCHNAME]", LocalSync.BranchName);
 		BuildRule = BuildRule.Replace("[PLATFORM]", TargetPlatform);
 		string IncludeRule = BuildRule.Replace("[CL]", "*");
