Split Incremental / Regular EngineTests into groups.
Moved jobs that don't require rendering to general build pool.
Added option for including cooked platforms (currently disabled, need to remove copy from DevStreams.xml).
Added XB1/Switch tests (also currently disabled)
Moved reports and logs to new file structure and site on automation.epicgames.net
Added VerifyDDC option to Gauntlet job.

#tests preflighted (lots).

diff --git a/Engine/Source/Programs/AutomationTool/Gauntlet/EngineTest/Tests/EngineTest.RunAutomatedTests.cs b/Engine/Source/Programs/AutomationTool/Gauntlet/EngineTest/Tests/EngineTest.RunAutomatedTests.cs
--- a/Engine/Source/Programs/AutomationTool/Gauntlet/EngineTest/Tests/EngineTest.RunAutomatedTests.cs
+++ b/Engine/Source/Programs/AutomationTool/Gauntlet/EngineTest/Tests/EngineTest.RunAutomatedTests.cs
@@ -37,6 +37,12 @@ namespace EngineTest
 		[AutoParam]
 		public string ReportOutputURL = "";
 
+		/// <summary>
+		/// Validate DDC during tests
+		/// </summary>
+		[AutoParam]
+		public bool VerifyDDC = false;
+
 		/// <summary>
 		/// Used for having the editor and any client communicate
 		/// </summary>
@@ -54,6 +60,11 @@ namespace EngineTest
 				{
 					AppConfig.CommandLine += string.Format(" -NoWatchdog -stdout -FORCELOGFLUSH -CrashForUAT -unattended -buildmachine {0} -ExecCmds=\"Automation RunTests {1}; Quit;\"", ReportOutputArg, TestFilter);
 				}
+
+				if (VerifyDDC)
+				{
+					AppConfig.CommandLine += " -VerifyDDC";
+				}
 			}
 			else
 			{
