Adjust copy local settings for automation tool projects so we don't cause unnecessary builds.

#rnx

diff --git a/Engine/Source/Programs/AutomationTool/AllDesktop/AllDesktop.Automation.csproj b/Engine/Source/Programs/AutomationTool/AllDesktop/AllDesktop.Automation.csproj
--- a/Engine/Source/Programs/AutomationTool/AllDesktop/AllDesktop.Automation.csproj
+++ b/Engine/Source/Programs/AutomationTool/AllDesktop/AllDesktop.Automation.csproj
@@ -62,6 +62,7 @@
     <ProjectReference Include="..\..\DotNETCommon\DotNETUtilities\DotNETUtilities.csproj">
       <Project>{5d7d66e8-8c76-4af9-b3ec-2ef03421d730}</Project>
       <Name>DotNETUtilities</Name>
+      <Private>False</Private>
     </ProjectReference>
     <ProjectReference Include="..\..\UnrealBuildTool\UnrealBuildTool.csproj">
       <Project>{FD7C5E1A-CFE4-4FD5-A525-1EB1599A39AC}</Project>
diff --git a/Engine/Source/Programs/AutomationTool/BuildGraph/BuildGraph.Automation.csproj b/Engine/Source/Programs/AutomationTool/BuildGraph/BuildGraph.Automation.csproj
--- a/Engine/Source/Programs/AutomationTool/BuildGraph/BuildGraph.Automation.csproj
+++ b/Engine/Source/Programs/AutomationTool/BuildGraph/BuildGraph.Automation.csproj
@@ -107,6 +107,7 @@
     <ProjectReference Include="..\..\DotNETCommon\DotNETUtilities\DotNETUtilities.csproj">
       <Project>{5d7d66e8-8c76-4af9-b3ec-2ef03421d730}</Project>
       <Name>DotNETUtilities</Name>
+      <Private>False</Private>
     </ProjectReference>
     <ProjectReference Include="..\..\UnrealBuildTool\UnrealBuildTool.csproj">
       <Project>{fd7c5e1a-cfe4-4fd5-a525-1eb1599a39ac}</Project>
diff --git a/Engine/Source/Programs/AutomationTool/Gauntlet/Gauntlet.Automation.csproj b/Engine/Source/Programs/AutomationTool/Gauntlet/Gauntlet.Automation.csproj
--- a/Engine/Source/Programs/AutomationTool/Gauntlet/Gauntlet.Automation.csproj
+++ b/Engine/Source/Programs/AutomationTool/Gauntlet/Gauntlet.Automation.csproj
@@ -143,6 +143,7 @@
     <ProjectReference Include="..\..\DotNETCommon\DotNETUtilities\DotNETUtilities.csproj">
       <Project>{5d7d66e8-8c76-4af9-b3ec-2ef03421d730}</Project>
       <Name>DotNETUtilities</Name>
+      <Private>False</Private>
     </ProjectReference>
     <ProjectReference Include="..\Android\Android.Automation.csproj">
       <Project>{32e3a47e-359a-4f78-ae6c-703b87de0b97}</Project>
@@ -157,6 +158,7 @@
     <ProjectReference Include="..\Localization\Localization.Automation.csproj">
       <Project>{c2f959b1-1727-4598-9c32-3d9acce72750}</Project>
       <Name>Localization.Automation</Name>
+      <Private>False</Private>
     </ProjectReference>
     <ProjectReference Include="..\XLocLocalization\XLocLocalization.Automation.csproj">
       <Project>{7e2d6784-c9b7-40c6-a89b-0a8afddce60d}</Project>
@@ -166,6 +168,7 @@
     <ProjectReference Include="..\OneSkyLocalization\OneSkyLocalization.Automation.csproj">
       <Project>{39369368-5603-46da-baae-766bb3cb2b51}</Project>
       <Name>OneSkyLocalization.Automation</Name>
+      <Private>False</Private>
     </ProjectReference>
     <ProjectReference Include="..\Scripts\AutomationScripts.Automation.csproj">
       <Project>{8aa00d65-0954-4a27-ac0d-fb8b1106120f}</Project>
diff --git a/Engine/Source/Programs/AutomationTool/HTML5/HTML5.Automation.csproj b/Engine/Source/Programs/AutomationTool/HTML5/HTML5.Automation.csproj
--- a/Engine/Source/Programs/AutomationTool/HTML5/HTML5.Automation.csproj
+++ b/Engine/Source/Programs/AutomationTool/HTML5/HTML5.Automation.csproj
@@ -61,6 +61,7 @@
     <ProjectReference Include="..\..\DotNETCommon\DotNETUtilities\DotNETUtilities.csproj">
       <Project>{5d7d66e8-8c76-4af9-b3ec-2ef03421d730}</Project>
       <Name>DotNETUtilities</Name>
+      <Private>False</Private>
     </ProjectReference>
     <ProjectReference Include="..\..\UnrealBuildTool\UnrealBuildTool.csproj">
       <Project>{fd7c5e1a-cfe4-4fd5-a525-1eb1599a39ac}</Project>
diff --git a/Engine/Source/Programs/AutomationTool/Localization/Localization.Automation.csproj b/Engine/Source/Programs/AutomationTool/Localization/Localization.Automation.csproj
--- a/Engine/Source/Programs/AutomationTool/Localization/Localization.Automation.csproj
+++ b/Engine/Source/Programs/AutomationTool/Localization/Localization.Automation.csproj
@@ -56,6 +56,7 @@
     <ProjectReference Include="..\..\DotNETCommon\DotNETUtilities\DotNETUtilities.csproj">
       <Project>{5d7d66e8-8c76-4af9-b3ec-2ef03421d730}</Project>
       <Name>DotNETUtilities</Name>
+      <Private>False</Private>
     </ProjectReference>
     <ProjectReference Include="..\..\UnrealBuildTool\UnrealBuildTool.csproj">
       <Project>{fd7c5e1a-cfe4-4fd5-a525-1eb1599a39ac}</Project>
diff --git a/Engine/Source/Programs/AutomationTool/Scripts/AutomationScripts.Automation.csproj b/Engine/Source/Programs/AutomationTool/Scripts/AutomationScripts.Automation.csproj
--- a/Engine/Source/Programs/AutomationTool/Scripts/AutomationScripts.Automation.csproj
+++ b/Engine/Source/Programs/AutomationTool/Scripts/AutomationScripts.Automation.csproj
@@ -126,6 +126,7 @@
     <ProjectReference Include="..\..\DotNETCommon\DotNETUtilities\DotNETUtilities.csproj">
       <Project>{5d7d66e8-8c76-4af9-b3ec-2ef03421d730}</Project>
       <Name>DotNETUtilities</Name>
+      <Private>False</Private>
     </ProjectReference>
     <ProjectReference Include="..\..\UnrealBuildTool\UnrealBuildTool.csproj">
       <Project>{fd7c5e1a-cfe4-4fd5-a525-1eb1599a39ac}</Project>
diff --git a/Engine/Source/Programs/AutomationTool/XLocLocalization/XLocLocalization.Automation.csproj b/Engine/Source/Programs/AutomationTool/XLocLocalization/XLocLocalization.Automation.csproj
--- a/Engine/Source/Programs/AutomationTool/XLocLocalization/XLocLocalization.Automation.csproj
+++ b/Engine/Source/Programs/AutomationTool/XLocLocalization/XLocLocalization.Automation.csproj
@@ -59,6 +59,7 @@
     <ProjectReference Include="..\..\DotNETCommon\DotNETUtilities\DotNETUtilities.csproj">
       <Project>{5d7d66e8-8c76-4af9-b3ec-2ef03421d730}</Project>
       <Name>DotNETUtilities</Name>
+      <Private>False</Private>
     </ProjectReference>
     <ProjectReference Include="..\..\UnrealBuildTool\UnrealBuildTool.csproj">
       <Project>{fd7c5e1a-cfe4-4fd5-a525-1eb1599a39ac}</Project>
@@ -73,6 +74,7 @@
     <ProjectReference Include="..\Localization\Localization.Automation.csproj">
       <Project>{c2f959b1-1727-4598-9c32-3d9acce72750}</Project>
       <Name>Localization.Automation</Name>
+      <Private>False</Private>
     </ProjectReference>
   </ItemGroup>
   <ItemGroup />
