Adjust copy local settings for automation tool projects so we don't cause unnecessary builds.

#rnx

diff --git a/Engine/Source/Programs/AutomationTool/Android/Android.Automation.csproj b/Engine/Source/Programs/AutomationTool/Android/Android.Automation.csproj
--- a/Engine/Source/Programs/AutomationTool/Android/Android.Automation.csproj
+++ b/Engine/Source/Programs/AutomationTool/Android/Android.Automation.csproj
@@ -66,6 +66,7 @@
     <ProjectReference Include="..\..\DotNETCommon\DotNETUtilities\DotNETUtilities.csproj">
       <Project>{5d7d66e8-8c76-4af9-b3ec-2ef03421d730}</Project>
       <Name>DotNETUtilities</Name>
+      <Private>False</Private>
     </ProjectReference>
     <ProjectReference Include="..\..\UnrealBuildTool\UnrealBuildTool.csproj">
       <Project>{fd7c5e1a-cfe4-4fd5-a525-1eb1599a39ac}</Project>
diff --git a/Engine/Source/Programs/AutomationTool/IOS/IOS.Automation.csproj b/Engine/Source/Programs/AutomationTool/IOS/IOS.Automation.csproj
--- a/Engine/Source/Programs/AutomationTool/IOS/IOS.Automation.csproj
+++ b/Engine/Source/Programs/AutomationTool/IOS/IOS.Automation.csproj
@@ -69,6 +69,7 @@
     <ProjectReference Include="..\..\DotNETCommon\DotNETUtilities\DotNETUtilities.csproj">
       <Project>{5d7d66e8-8c76-4af9-b3ec-2ef03421d730}</Project>
       <Name>DotNETUtilities</Name>
+      <Private>False</Private>
     </ProjectReference>
     <ProjectReference Include="..\..\UnrealBuildTool\UnrealBuildTool.csproj">
       <Project>{fd7c5e1a-cfe4-4fd5-a525-1eb1599a39ac}</Project>
diff --git a/Engine/Source/Programs/AutomationTool/Linux/Linux.Automation.csproj b/Engine/Source/Programs/AutomationTool/Linux/Linux.Automation.csproj
--- a/Engine/Source/Programs/AutomationTool/Linux/Linux.Automation.csproj
+++ b/Engine/Source/Programs/AutomationTool/Linux/Linux.Automation.csproj
@@ -58,6 +58,7 @@
     <ProjectReference Include="..\..\DotNETCommon\DotNETUtilities\DotNETUtilities.csproj">
       <Project>{5d7d66e8-8c76-4af9-b3ec-2ef03421d730}</Project>
       <Name>DotNETUtilities</Name>
+      <Private>False</Private>
     </ProjectReference>
     <ProjectReference Include="..\..\UnrealBuildTool\UnrealBuildTool.csproj">
       <Project>{fd7c5e1a-cfe4-4fd5-a525-1eb1599a39ac}</Project>
diff --git a/Engine/Source/Programs/AutomationTool/Mac/Mac.Automation.csproj b/Engine/Source/Programs/AutomationTool/Mac/Mac.Automation.csproj
--- a/Engine/Source/Programs/AutomationTool/Mac/Mac.Automation.csproj
+++ b/Engine/Source/Programs/AutomationTool/Mac/Mac.Automation.csproj
@@ -58,6 +58,7 @@
     <ProjectReference Include="..\..\DotNETCommon\DotNETUtilities\DotNETUtilities.csproj">
       <Project>{5d7d66e8-8c76-4af9-b3ec-2ef03421d730}</Project>
       <Name>DotNETUtilities</Name>
+      <Private>False</Private>
     </ProjectReference>
     <ProjectReference Include="..\..\UnrealBuildTool\UnrealBuildTool.csproj">
       <Project>{fd7c5e1a-cfe4-4fd5-a525-1eb1599a39ac}</Project>
diff --git a/Engine/Source/Programs/AutomationTool/OneSkyLocalization/OneSkyLocalization.Automation.csproj b/Engine/Source/Programs/AutomationTool/OneSkyLocalization/OneSkyLocalization.Automation.csproj
--- a/Engine/Source/Programs/AutomationTool/OneSkyLocalization/OneSkyLocalization.Automation.csproj
+++ b/Engine/Source/Programs/AutomationTool/OneSkyLocalization/OneSkyLocalization.Automation.csproj
@@ -62,6 +62,7 @@
     <ProjectReference Include="..\..\DotNETCommon\DotNETUtilities\DotNETUtilities.csproj">
       <Project>{5d7d66e8-8c76-4af9-b3ec-2ef03421d730}</Project>
       <Name>DotNETUtilities</Name>
+      <Private>False</Private>
     </ProjectReference>
     <ProjectReference Include="..\..\UnrealBuildTool\UnrealBuildTool.csproj">
       <Project>{fd7c5e1a-cfe4-4fd5-a525-1eb1599a39ac}</Project>
@@ -76,6 +77,7 @@
     <ProjectReference Include="..\Localization\Localization.Automation.csproj">
       <Project>{c2f959b1-1727-4598-9c32-3d9acce72750}</Project>
       <Name>Localization.Automation</Name>
+      <Private>False</Private>
     </ProjectReference>
   </ItemGroup>
   <ItemGroup />
diff --git a/Engine/Source/Programs/AutomationTool/Win/Win.Automation.csproj b/Engine/Source/Programs/AutomationTool/Win/Win.Automation.csproj
--- a/Engine/Source/Programs/AutomationTool/Win/Win.Automation.csproj
+++ b/Engine/Source/Programs/AutomationTool/Win/Win.Automation.csproj
@@ -60,6 +60,7 @@
     <ProjectReference Include="..\..\DotNETCommon\DotNETUtilities\DotNETUtilities.csproj">
       <Project>{5d7d66e8-8c76-4af9-b3ec-2ef03421d730}</Project>
       <Name>DotNETUtilities</Name>
+      <Private>False</Private>
     </ProjectReference>
     <ProjectReference Include="..\..\UnrealBuildTool\UnrealBuildTool.csproj">
       <Project>{fd7c5e1a-cfe4-4fd5-a525-1eb1599a39ac}</Project>
@@ -74,6 +75,7 @@
     <ProjectReference Include="..\BuildGraph\BuildGraph.Automation.csproj">
       <Project>{a002361e-37f9-4124-af82-cf0d393cb928}</Project>
       <Name>BuildGraph.Automation</Name>
+      <Private>False</Private>
     </ProjectReference>
   </ItemGroup>
   <Import Project="$(MSBuildToolsPath)\Microsoft.CSharp.targets" />
