Adjust copy local settings for automation tool projects so we don't cause unnecessary builds.

#rnx

diff --git a/Engine/Source/Programs/AutomationTool/Lumin/Lumin.Automation.csproj b/Engine/Source/Programs/AutomationTool/Lumin/Lumin.Automation.csproj
--- a/Engine/Source/Programs/AutomationTool/Lumin/Lumin.Automation.csproj
+++ b/Engine/Source/Programs/AutomationTool/Lumin/Lumin.Automation.csproj
@@ -78,6 +78,7 @@
     <ProjectReference Include="..\..\DotNETCommon\DotNETUtilities\DotNETUtilities.csproj">
       <Project>{5D7D66E8-8C76-4AF9-B3EC-2EF03421D730}</Project>
       <Name>DotNETUtilities</Name>
+      <Private>False</Private>
     </ProjectReference>
   </ItemGroup>
   <Import Project="$(MSBuildToolsPath)\Microsoft.CSharp.targets" />
