Adjust copy local settings for automation tool projects so we don't cause unnecessary builds.

#rnx

diff --git a/Engine/Source/Programs/AutomationTool/TVOS/TVOS.Automation.csproj b/Engine/Source/Programs/AutomationTool/TVOS/TVOS.Automation.csproj
--- a/Engine/Source/Programs/AutomationTool/TVOS/TVOS.Automation.csproj
+++ b/Engine/Source/Programs/AutomationTool/TVOS/TVOS.Automation.csproj
@@ -60,6 +60,7 @@
     <ProjectReference Include="..\..\DotNETCommon\DotNETUtilities\DotNETUtilities.csproj">
       <Project>{5d7d66e8-8c76-4af9-b3ec-2ef03421d730}</Project>
       <Name>DotNETUtilities</Name>
+      <Private>False</Private>
     </ProjectReference>
     <ProjectReference Include="..\..\UnrealBuildTool\UnrealBuildTool.csproj">
       <Project>{fd7c5e1a-cfe4-4fd5-a525-1eb1599a39ac}</Project>
