Fix errors caused by failing to parse Gauntlet blacklist now that platforms are no longer an enum.

#rnx

diff --git a/Engine/Source/Programs/AutomationTool/Gauntlet/Framework/Base/Gauntlet.Blacklist.cs b/Engine/Source/Programs/AutomationTool/Gauntlet/Framework/Base/Gauntlet.Blacklist.cs
--- a/Engine/Source/Programs/AutomationTool/Gauntlet/Framework/Base/Gauntlet.Blacklist.cs
+++ b/Engine/Source/Programs/AutomationTool/Gauntlet/Framework/Base/Gauntlet.Blacklist.cs
@@ -20,13 +20,13 @@ namespace Gauntlet
 	{
 		public string TestName;
 
-		public UnrealTargetPlatform[] Platforms;
+		public string[] Platforms;
 
 		public string BranchName;
 		public BlacklistEntry()
 		{
 			TestName = "None";
-			Platforms = new UnrealTargetPlatform [] { };
+			Platforms = new string [] { };
 			BranchName = "None";
 		}
 
@@ -116,7 +116,7 @@ namespace Gauntlet
 			Entries = Entries.Where(E => E.BranchName == "*" || string.Equals(E.BranchName, NormalizedBranchName, StringComparison.OrdinalIgnoreCase));
 
 			// Filter by branch
-			Entries = Entries.Where(E => E.Platforms.Length == 0 || E.Platforms.Contains(InPlatform));
+			Entries = Entries.Where(E => E.Platforms.Length == 0 || E.Platforms.Contains(InPlatform.ToString()));
 
 			return Entries.Count() > 0;
 		}
