Gauntlet: Fix exception trying to parse an empty string as a platform.

diff --git a/Engine/Source/Programs/AutomationTool/Gauntlet/Unreal/Utils/Gauntlet.UnrealHelpers.cs b/Engine/Source/Programs/AutomationTool/Gauntlet/Unreal/Utils/Gauntlet.UnrealHelpers.cs
--- a/Engine/Source/Programs/AutomationTool/Gauntlet/Unreal/Utils/Gauntlet.UnrealHelpers.cs
+++ b/Engine/Source/Programs/AutomationTool/Gauntlet/Unreal/Utils/Gauntlet.UnrealHelpers.cs
@@ -228,7 +228,7 @@ namespace Gauntlet
 		internal class ConfigInfo
 		{
 			public UnrealTargetRole 				RoleType;
-			public UnrealTargetPlatform 		Platform;
+			public UnrealTargetPlatform? 		Platform;
 			public UnrealTargetConfiguration 	Configuration;
 			public bool							SharedBuild;
 
@@ -296,8 +296,11 @@ namespace Gauntlet
 					Config.Configuration = UnrealTargetConfiguration.Development;   // Development has no string
 				}
 
-				// this is not expected to be able to fail
-				Config.Platform = UnrealTargetPlatform.Parse(PlatformName);
+				UnrealTargetPlatform Platform;
+				if (ConfigType.Length > 0 && UnrealTargetPlatform.TryParse(ConfigType, out Platform))
+				{
+					Config.Platform = Platform;
+				}
 			}
 
 			return Config;
