- Added debugging information when failing to load a type from a compiled script assembly

diff --git a/Engine/Source/Programs/AutomationTool/AutomationUtils/ScriptCompiler.cs b/Engine/Source/Programs/AutomationTool/AutomationUtils/ScriptCompiler.cs
--- a/Engine/Source/Programs/AutomationTool/AutomationUtils/ScriptCompiler.cs
+++ b/Engine/Source/Programs/AutomationTool/AutomationUtils/ScriptCompiler.cs
@@ -111,6 +111,14 @@ namespace AutomationTool
 						}
 					}
 				}
+				catch (ReflectionTypeLoadException LoadEx)
+				{
+					foreach (var SubEx in LoadEx.LoaderExceptions)
+					{
+						Log.TraceWarning("Got type loader exception: {0}", SubEx.ToString());
+					}
+					throw new AutomationException("Failed to add commands from {0}. {1}", CompiledScripts, LoadEx);
+				}
 				catch (Exception Ex)
 				{
 					throw new AutomationException("Failed to add commands from {0}. {1}", CompiledScripts, Ex);
