Bug fix for regex restrictions on BuildGraph Options.

diff --git a/Engine/Source/Programs/AutomationTool/BuildGraph/Script.cs b/Engine/Source/Programs/AutomationTool/BuildGraph/Script.cs
--- a/Engine/Source/Programs/AutomationTool/BuildGraph/Script.cs
+++ b/Engine/Source/Programs/AutomationTool/BuildGraph/Script.cs
@@ -656,7 +656,7 @@ namespace AutomationTool
 						string Restrict = ReadAttribute(Element, "Restrict");
 						if(!String.IsNullOrEmpty(Restrict))
 						{
-							string Pattern = "^" + Restrict + "$";
+							string Pattern = "^(" + Restrict + ")$";
 							if(!Regex.IsMatch(Value, Pattern, RegexOptions.IgnoreCase))
 							{
 								LogError(Element, "'{0}' is not a valid value for '{1}' (required: '{2}')", Value, Name, Restrict);
