Modified Gauntlet's CopyDirectory function to allow provision of a pattern (e.g. *.so) or regular expresson to filter what files are copied.

Fixed issue where lazy copy wasn't working because destination info was only created in mirror mode.

Moved options for Gauntlet's CopyDirectory function to a structure. TODO - update the existing use cases and remove the older wrappers.

diff --git a/Engine/Source/Programs/AutomationTool/Gauntlet/Framework/Gauntlet.Utils.cs b/Engine/Source/Programs/AutomationTool/Gauntlet/Framework/Gauntlet.Utils.cs
--- a/Engine/Source/Programs/AutomationTool/Gauntlet/Framework/Gauntlet.Utils.cs
+++ b/Engine/Source/Programs/AutomationTool/Gauntlet/Framework/Gauntlet.Utils.cs
@@ -733,10 +733,62 @@ namespace Gauntlet
 			/// </summary>
 			public enum CopyOptions
 			{
-				Copy		= (1 << 0),		// Normal copy & combine/overwrite
-				Mirror		= (1 << 1),		// copy + remove files from dest if not in src
+				Copy = (1 << 0),        // Normal copy & combine/overwrite
+				Mirror = (1 << 1),      // copy + remove files from dest if not in src
 				Default = Copy
 			}
+
+			/// <summary>
+			/// Options that can be specified to the CopyDirectory function.
+			/// </summary>
+			public class CopyDirectoryOptions
+			{
+				public CopyOptions Mode { get; set; }
+
+				public int Retries { get; set; }
+
+				public Func<string, string> Transform { get; set; }
+
+				public string Pattern { get; set; }
+
+				public Regex Regex { get; set; }
+
+				public bool Recursive { get; set; }
+
+				public bool Verbose { get; set; }
+
+				public CopyDirectoryOptions()
+				{
+					Mode = CopyOptions.Copy;
+					Retries = 10;
+					Transform = delegate (string s)
+					{
+						return s;
+					};
+					Pattern = "*";
+					Recursive = true;
+					Regex = null;
+					Verbose = false;
+				}
+
+				/// <summary>
+				/// Returns true if the pattern indicates entire directory should be copied
+				/// </summary>
+				public bool IsDirectoryPattern
+				{
+					get
+					{
+						return 
+							Regex == null &&
+							(
+								string.IsNullOrEmpty(Pattern)
+								|| Pattern.Equals("*")
+								|| Pattern.Equals("*.*")
+								|| Pattern.Equals("...")
+							);
+					}
+				}
+			}
 			
 			/// <summary>
 			/// Convenience function that removes some of the more esoteric options
@@ -745,9 +797,27 @@ namespace Gauntlet
 			/// <param name="DestDirPath"></param>
 			/// <param name="Options"></param>
 			/// <param name="RetryCount"></param>
-			public static void CopyDirectory(string SourceDirPath, string DestDirPath, CopyOptions Options = CopyOptions.Default, int RetryCount = 5)
+			public static void CopyDirectory(string SourceDirPath, string DestDirPath, CopyOptions Mode = CopyOptions.Default, int RetryCount = 5)
 			{
-				CopyDirectory(SourceDirPath, DestDirPath, Options, delegate (string s) { return s; }, RetryCount);
+				CopyDirectory(SourceDirPath, DestDirPath, Mode, delegate (string s) { return s; }, RetryCount);
+			}
+
+			/// <summary>
+			/// Legacy convenience function that exposes transform & Retry count
+			/// </summary>
+			/// <param name="SourceDirPath"></param>
+			/// <param name="DestDirPath"></param>
+			/// <param name="Mode"></param>
+			/// <param name="Transform"></param>
+			/// <param name="RetryCount"></param>
+			public static void CopyDirectory(string SourceDirPath, string DestDirPath, CopyOptions Mode, Func<string, string> Transform, int RetryCount = 5)
+			{
+				CopyDirectoryOptions Options = new CopyDirectoryOptions();
+				Options.Retries = RetryCount;
+				Options.Mode = Mode;
+				Options.Transform = Transform;
+
+				CopyDirectory(SourceDirPath, DestDirPath, Options);
 			}
 
 			/// <summary>
@@ -757,7 +827,7 @@ namespace Gauntlet
 			/// <param name="SourcePath"></param>
 			/// <param name="DestPath"></param>
 			/// <param name="Verbose"></param>
-			public static void CopyDirectory(string SourceDirPath, string DestDirPath, CopyOptions Options, Func<string, string> Transform, int RetryCount = 5)
+			public static void CopyDirectory(string SourceDirPath, string DestDirPath, CopyDirectoryOptions Options)
 			{
 				DateTime StartTime = DateTime.Now;
 				
@@ -769,10 +839,35 @@ namespace Gauntlet
 					DestDir = Directory.CreateDirectory(DestDir.FullName);
 				}
 				
-				bool IsMirroring = (Options & CopyOptions.Mirror) == CopyOptions.Mirror;
+				bool IsMirroring = (Options.Mode & CopyOptions.Mirror) == CopyOptions.Mirror;
+
+				if (IsMirroring && !Options.IsDirectoryPattern)
+				{
+					Log.Warning("Can only use mirror with pattern that includes whole directories (e.g. '*')");
+					IsMirroring = false;
+				}
 
-				System.IO.FileInfo[] SourceFiles = SourceDir.GetFiles("*", SearchOption.AllDirectories);
-				System.IO.FileInfo[] DestFiles = null;
+				IEnumerable<FileInfo> SourceFiles = null;
+				FileInfo[] DestFiles = null;
+
+				// find all files. If a directory get them all, else use the pattern/regex
+				if (Options.IsDirectoryPattern)
+				{
+					SourceFiles = SourceDir.GetFiles("*", Options.Recursive ? SearchOption.AllDirectories : SearchOption.TopDirectoryOnly);
+				}
+				else
+				{
+					if (Options.Regex == null)
+					{
+						SourceFiles = SourceDir.GetFiles(Options.Pattern, Options.Recursive ? SearchOption.AllDirectories : SearchOption.TopDirectoryOnly);
+					}
+					else
+					{
+						SourceFiles = SourceDir.GetFiles("*", Options.Recursive ? SearchOption.AllDirectories : SearchOption.TopDirectoryOnly);
+
+						SourceFiles = SourceFiles.Where(F => Options.Regex.IsMatch(F.Name));
+					}
+				}
 
 				// Convert dest into a map of relative paths to absolute
 				Dictionary<string, System.IO.FileInfo> DestStructure = new Dictionary<string, System.IO.FileInfo>();
@@ -811,9 +906,26 @@ namespace Gauntlet
 						SourceFilePath = SourceFilePath.Substring(1);
 					}
 
-					string DestFilePath = Transform(SourceFilePath);
+					string DestFilePath = Options.Transform(SourceFilePath);
+
+					FileInfo DestInfo = null;
+
+					// We may have destination info if mirroring where we prebuild it all, if not
+					// grab it now
+					if (DestStructure.ContainsKey(DestFilePath))
+					{
+						DestInfo = DestStructure[DestFilePath];
+					}
+					else
+					{
+						string FullDestPath = Path.Combine(DestDir.FullName, DestFilePath);
+						if (File.Exists(FullDestPath))
+						{
+							DestInfo = new FileInfo(FullDestPath);
+						}
+					}
 
-					if (DestStructure.ContainsKey(DestFilePath) == false)
+					if (DestInfo == null)
 					{
 						// No copy in dest, add it to the list
 						CopyList.Add(SourceFilePath);
@@ -821,7 +933,6 @@ namespace Gauntlet
 					else
 					{
 						// Check the file is the same version
-						FileInfo DestInfo = DestStructure[DestFilePath];
 
 						// Difference in ticks. Even though we set the dest to the src there still appears to be minute
 						// differences in ticks. 1ms is 10k ticks...
@@ -833,6 +944,17 @@ namespace Gauntlet
 						{
 							CopyList.Add(SourceFilePath);
 						}
+						else
+						{
+							if (Options.Verbose)
+							{
+								Log.Info("Will skip copy to {0}. File up to date.", DestInfo.FullName);
+							}
+							else
+							{
+								Log.Verbose("Will skip copy to {0}. File up to date.", DestInfo.FullName);
+							}
+						}
 
 						// Remove it from the map
 						DestStructure.Remove(DestFilePath);
@@ -851,8 +973,15 @@ namespace Gauntlet
 					foreach (string RelativePath in DeletionList)
 					{
 						FileInfo DestInfo = new FileInfo(Path.Combine(DestDir.FullName, RelativePath));
-						
-						Log.Verbose("Deleting extra file {0}", DestInfo.FullName);
+
+						if (Options.Verbose)
+						{
+							Log.Info("Deleting extra file {0}", DestInfo.FullName);
+						}
+						else
+						{
+							Log.Verbose("Deleting extra file {0}", DestInfo.FullName);
+						}
 
 						try
 						{
@@ -877,7 +1006,14 @@ namespace Gauntlet
 						{
 							if (SubDir.GetFiles().Length == 0 && SubDir.GetDirectories().Length == 0)
 							{
-								Log.Verbose("Deleting empty dir {0}", SubDir.FullName);
+								if (Options.Verbose)
+								{
+									Log.Info("Deleting empty dir {0}", SubDir.FullName);
+								}
+								else
+								{
+									Log.Verbose("Deleting empty dir {0}", SubDir.FullName);
+								}
 
 								SubDir.Delete(true);
 							}
@@ -909,9 +1045,9 @@ namespace Gauntlet
 					// ensure path exists
 					string DestPath = Path.Combine(DestDir.FullName, RelativePath);
 
-					if (Transform != null)
+					if (Options.Transform != null)
 					{
-						DestPath = Transform(DestPath);
+						DestPath = Options.Transform(DestPath);
 					}
 
 					string SourcePath = Path.Combine(SourceDir.FullName, RelativePath);
@@ -934,9 +1070,9 @@ namespace Gauntlet
 
 					string DestFile = DestInfo.FullName;
 
-					if (Transform != null)
+					if (Options.Transform != null)
 					{
-						DestFile = Transform(DestFile);
+						DestFile = Options.Transform(DestFile);
 					}
 
 					int Tries = 0;
@@ -946,7 +1082,14 @@ namespace Gauntlet
 					{
 						try
 						{
-							Log.Verbose("Copying to {0}", DestFile);
+							if (Options.Verbose)
+							{
+								Log.Info("Copying to {0}", DestFile);
+							}
+							else
+							{
+								Log.Verbose("Copying to {0}", DestFile);
+							}
 
 							SrcInfo.CopyTo(DestFile, true);
 
@@ -958,9 +1101,9 @@ namespace Gauntlet
 						}
 						catch (Exception ex)
 						{
-							if (Tries++ < RetryCount)
+							if (Tries++ < Options.Retries)
 							{
-								Log.Info("Copy to {0} failed, retrying {1} of {2} in 30 secs..", DestFile, Tries, RetryCount);
+								Log.Info("Copy to {0} failed, retrying {1} of {2} in 30 secs..", DestFile, Tries, Options.Retries);
 								// todo - make param..
 								Thread.Sleep(30000);
 							}
@@ -979,7 +1122,14 @@ namespace Gauntlet
 				TimeSpan Duration = DateTime.Now - StartTime;
 				if (Duration.TotalSeconds > 10)
 				{
-					Log.Verbose("Copied Directory in {0}", Duration.ToString(@"mm\m\:ss\s"));
+					if (Options.Verbose)
+					{
+						Log.Info("Copied Directory in {0}", Duration.ToString(@"mm\m\:ss\s"));
+					}
+					else
+					{
+						Log.Verbose("Copied Directory in {0}", Duration.ToString(@"mm\m\:ss\s"));
+					}
 				}
 
 				// remove cancel handler
