Updated TargetDeviceAndroid so both main and patch OBB files are copied.
Made CopyFile funtion public so code using TargetDeviceAndroid can use it.

diff --git a/Engine/Source/Programs/AutomationTool/AutomationUtils/Automation.cs b/Engine/Source/Programs/AutomationTool/AutomationUtils/Automation.cs
--- a/Engine/Source/Programs/AutomationTool/AutomationUtils/Automation.cs
+++ b/Engine/Source/Programs/AutomationTool/AutomationUtils/Automation.cs
@@ -420,6 +420,11 @@ AutomationTool.exe [-verbose] [-compileonly] [-p4] Command0 [-Arg0 -Arg1 -Arg2 .
 
 		#region Main Program
 
+		/// <summary>
+		/// Compiler with all scripts
+		/// </summary>
+		public static ScriptCompiler Compiler { get; set; }
+
 		/// <summary>
 		/// Main method.
 		/// </summary>
@@ -478,7 +483,7 @@ AutomationTool.exe [-verbose] [-compileonly] [-p4] Command0 [-Arg0 -Arg1 -Arg2 .
 			ProjectUtils.CleanupFolders();
 
 			// Compile scripts.
-			ScriptCompiler Compiler = new ScriptCompiler();
+			Compiler = new ScriptCompiler();
 			using(TelemetryStopwatch ScriptCompileStopwatch = new TelemetryStopwatch("ScriptCompile"))
 			{
 				Compiler.FindAndCompileAllScripts(OutScriptsForProjectFileName, AdditionalScriptsFolders);
diff --git a/Engine/Source/Programs/AutomationTool/AutomationUtils/ScriptCompiler.cs b/Engine/Source/Programs/AutomationTool/AutomationUtils/ScriptCompiler.cs
--- a/Engine/Source/Programs/AutomationTool/AutomationUtils/ScriptCompiler.cs
+++ b/Engine/Source/Programs/AutomationTool/AutomationUtils/ScriptCompiler.cs
@@ -17,7 +17,7 @@ namespace AutomationTool
 	/// <summary>
 	/// Exception thrown by PreprocessScriptFile.
 	/// </summary>
-	class CompilationException : AutomationException
+	public class CompilationException : AutomationException
 	{
 		public CompilationException(string Filename, int StartLine, int StartColumn, int EndLine, int EndColumn, string Message, params string[] Args)
 			: base(String.Format("Compilation Failed.\n>{0}({1},{2},{3},{4}): error: {5}", Path.GetFullPath(Filename), StartLine + 1, StartColumn + 1, EndLine + 1, EndColumn + 1,
@@ -29,7 +29,7 @@ namespace AutomationTool
 	/// <summary>
 	/// Compiles and loads script assemblies.
 	/// </summary>
-	class ScriptCompiler
+	public class ScriptCompiler
 	{
 		#region Fields
 				
diff --git a/Engine/Source/Programs/AutomationTool/Gauntlet/Platform/Android/Gauntlet.TargetDeviceAndroid.cs b/Engine/Source/Programs/AutomationTool/Gauntlet/Platform/Android/Gauntlet.TargetDeviceAndroid.cs
--- a/Engine/Source/Programs/AutomationTool/Gauntlet/Platform/Android/Gauntlet.TargetDeviceAndroid.cs
+++ b/Engine/Source/Programs/AutomationTool/Gauntlet/Platform/Android/Gauntlet.TargetDeviceAndroid.cs
@@ -661,7 +661,7 @@ namespace Gauntlet
 			return AdbResult.ExitCode == 0;
 		}
 
-		protected bool CopyFileToDevice(string PackageName, string SourcePath, string DestPath, bool IgnoreDependencies = false)
+		public bool CopyFileToDevice(string PackageName, string SourcePath, string DestPath, bool IgnoreDependencies = false)
 		{
 			bool IsAPK = string.Equals(Path.GetExtension(SourcePath), ".apk", StringComparison.OrdinalIgnoreCase);
 
@@ -933,7 +933,7 @@ namespace Gauntlet
 					string DestFile = Path.GetFileName(DestPath);
 
 					// If we installed a new APK we need to change the package version
-					Match OBBMatch = Regex.Match(DestFile, @"main\.(\d+)\.com.*\.obb");
+					Match OBBMatch = Regex.Match(DestFile, @"\.(\d+)\.com.*\.obb");
 					if (OBBMatch.Success)
 					{
 						string NewFileName = DestFile.Replace(OBBMatch.Groups[1].ToString(), PackageVersion);
