Gauntlet: Automated permissions and log caching improvements for Android
#rnx

diff --git a/Engine/Source/Programs/AutomationTool/Gauntlet/Platform/Android/Gauntlet.TargetDeviceAndroid.cs b/Engine/Source/Programs/AutomationTool/Gauntlet/Platform/Android/Gauntlet.TargetDeviceAndroid.cs
--- a/Engine/Source/Programs/AutomationTool/Gauntlet/Platform/Android/Gauntlet.TargetDeviceAndroid.cs
+++ b/Engine/Source/Programs/AutomationTool/Gauntlet/Platform/Android/Gauntlet.TargetDeviceAndroid.cs
@@ -120,6 +120,8 @@ namespace Gauntlet
 			if (bHasExited)
 			{
 				ActivityExited = true;
+				// Make sure entire activity log has been captured
+				UpdateCachedLog(true);
 				Log.VeryVerbose("{0}: process exited, Activity running={1}, Activity in foreground={2} ", ToString(), bActivityPresent.ToString(), bActivityInForeground.ToString());
 			}
 
@@ -134,50 +136,58 @@ namespace Gauntlet
 		public bool WasKilled { get; protected set; }
 
 		/// <summary>
-		/// The output of the test activity, runs a shell command returning the full log from device per call (possibly over wifi)
-		/// result is cached, and updated at ActivityLogDelta frequency
+		/// The output of the test activity
 		/// </summary>
 		public string StdOut
 		{
 			get
-			{				
-				if (!String.IsNullOrEmpty(ActivityLogCached) && (ActivityLogTime == DateTime.MinValue || ((DateTime.UtcNow - ActivityLogTime) < ActivityLogDelta)))
-				{
-					return ActivityLogCached;
-				}
+			{
+				UpdateCachedLog();
+				return String.IsNullOrEmpty(ActivityLogCached) ? String.Empty : ActivityLogCached;
+			}
+		}
+		/// <summary>
+		/// Updates cached activity log by running a shell command returning the full log from device (possibly over wifi)
+		/// The result is cached and updated at ActivityLogDelta frequency
+		/// </summary>
+		private void UpdateCachedLog(bool ForceUpdate = false)
+		{
+			if (!ForceUpdate && (ActivityLogTime == DateTime.MinValue || ((DateTime.UtcNow - ActivityLogTime) < ActivityLogDelta)))
+			{
+				return;
+			}
 
-				if (Install.AndroidDevice != null && Install.AndroidDevice.Disposed)
-				{
-					return String.IsNullOrEmpty(ActivityLogCached) ? String.Empty : ActivityLogCached;
-				}
+			if (Install.AndroidDevice != null && Install.AndroidDevice.Disposed)
+			{
+				Log.Warning("Attempting to cache log using disposed Android device");
+				return;
+			}
 
-				ActivityLogTime = DateTime.UtcNow;
+			string GetLogCommand = string.Format("shell cat {0}/Logs/{1}.log", Install.AndroidDevice.DeviceArtifactPath, Install.Name);
+			IProcessResult LogQuery = Install.AndroidDevice.RunAdbDeviceCommand(GetLogCommand, true);
 
-				string GetLogCommand = string.Format("shell cat {0}/Logs/{1}.log", Install.AndroidDevice.DeviceArtifactPath, Install.Name);
-				IProcessResult LogQuery = Install.AndroidDevice.RunAdbDeviceCommand(GetLogCommand, true);
+			if (LogQuery.ExitCode != 0)
+			{
+				Log.VeryVerbose("Unable to query activity stdout on device {0}", Install.AndroidDevice.Name);
+			}
+			else
+			{
+				ActivityLogCached = LogQuery.Output;
+			}
 
-				if (LogQuery.ExitCode != 0)
-				{
-					Log.VeryVerbose("Unable to query activity stdout on device {0}", Install.AndroidDevice.Name);					
-				}
-				else
-				{
-					ActivityLogCached = LogQuery.Output;
-				}
-					
-				// the activity has exited, mark final log sentinel 
-				if (ActivityExited)
-				{
-					ActivityLogTime = DateTime.MinValue;
-				}
+			ActivityLogTime = DateTime.UtcNow;
 
-				return ActivityLogCached;
+			// the activity has exited, mark final log sentinel 
+			if (ActivityExited)
+			{
+				ActivityLogTime = DateTime.MinValue;
 			}
+
 		}
 
 		private static readonly TimeSpan ActivityLogDelta = TimeSpan.FromSeconds(15);
-		private DateTime ActivityLogTime = DateTime.UtcNow;
-		private string ActivityLogCached = "";
+		private DateTime ActivityLogTime = DateTime.UtcNow - ActivityLogDelta;
+		private string ActivityLogCached = string.Empty;
 		
 
 		public int WaitForExit()
@@ -989,6 +999,7 @@ namespace Gauntlet
 				AdbCommand = string.Format("push {0} {1}", TmpFile, CommandLineFilePath);
 				RunAdbDeviceCommand(AdbCommand);
 
+				EnablePermissions(Build.AndroidPackageName);
 
 				File.Delete(TmpFile);
 			}
@@ -1119,6 +1130,18 @@ namespace Gauntlet
 			string CommandLine = "shell svc power stayon " + (bAllowSleep ? "false" : "usb");
 			RunAdbDeviceCommand(CommandLine, true, false, true);
 		}
+		/// <summary>
+		/// Enable Android permissions which would otherwise block automation with permimssion requests
+		/// </summary>
+		public void EnablePermissions(string AndroidPackageName)
+		{
+			List<string> Permissions = new List<string>{ "WRITE_EXTERNAL_STORAGE", "GET_ACCOUNTS", "RECORD_AUDIO" };
+			Permissions.ForEach(Permission => {
+				string CommandLine = string.Format("shell pm grant {0} android.permission.{1}", AndroidPackageName, Permission);
+				Log.Verbose(string.Format("Enabling permission: {0} {1}", AndroidPackageName, Permission));
+				RunAdbDeviceCommand(CommandLine, true, false, true);
+			});
+		}
 
 		public void KillRunningProcess(string AndroidPackageName)
 		{
