Gauntlet: Speculative fix for killing Android adb server surfacing superfluous errors
#rnx

diff --git a/Engine/Source/Programs/AutomationTool/Gauntlet/Platform/Android/Gauntlet.TargetDeviceAndroid.cs b/Engine/Source/Programs/AutomationTool/Gauntlet/Platform/Android/Gauntlet.TargetDeviceAndroid.cs
--- a/Engine/Source/Programs/AutomationTool/Gauntlet/Platform/Android/Gauntlet.TargetDeviceAndroid.cs
+++ b/Engine/Source/Programs/AutomationTool/Gauntlet/Platform/Android/Gauntlet.TargetDeviceAndroid.cs
@@ -1281,10 +1281,7 @@ namespace Gauntlet
 
 					bUsingCustomKeys = true;
 
-					Log.Info("Running adb kill-server to refresh credentials");
-					TargetDeviceAndroid.RunAdbGlobalCommand("kill-server", bPauseErrorParsing: true);
-
-					Thread.Sleep(5000);
+					KillAdbServer();
 				}
 
 				InstanceCount++;
@@ -1292,6 +1289,19 @@ namespace Gauntlet
 
 		}
 
+		private static void KillAdbServer()
+		{			
+			using (new ScopedSuspendECErrorParsing())
+			{
+				Log.Info("Running adb kill-server to refresh credentials");
+				TargetDeviceAndroid.RunAdbGlobalCommand("kill-server");
+				// killing the adb server restarts it and can surface superfluous device errors
+				int SleepTime = CommandUtils.IsBuildMachine ? 15000 : 5000;
+				Thread.Sleep(SleepTime);
+			}
+
+		}
+
 		public static void RemoveInstance()
 		{
 			lock (Globals.MainLock)
@@ -1302,12 +1312,8 @@ namespace Gauntlet
 				if (InstanceCount == 0 && bUsingCustomKeys)
 				{
 					Reset();
-
-					Log.Info("Running adb kill-server to refresh credentials");
-					TargetDeviceAndroid.RunAdbGlobalCommand("kill-server", bPauseErrorParsing: true);
-					Thread.Sleep(2500);
+					KillAdbServer();
 				}
-
 			}
 		}
 
