Add the ability to declare additional artifact directories to copy over in addition to /saved/

diff --git a/Engine/Source/Programs/AutomationTool/Gauntlet/Platform/Android/Gauntlet.TargetDeviceAndroid.cs b/Engine/Source/Programs/AutomationTool/Gauntlet/Platform/Android/Gauntlet.TargetDeviceAndroid.cs
--- a/Engine/Source/Programs/AutomationTool/Gauntlet/Platform/Android/Gauntlet.TargetDeviceAndroid.cs
+++ b/Engine/Source/Programs/AutomationTool/Gauntlet/Platform/Android/Gauntlet.TargetDeviceAndroid.cs
@@ -393,9 +393,10 @@ namespace Gauntlet
         }
 
         public void PopulateDirectoryMappings(string ProjectDir)
-        {
-            LocalDirectoryMappings.Add(EIntendedBaseCopyDirectory.Binaries, Path.Combine(ProjectDir, "Binaries"));
-            LocalDirectoryMappings.Add(EIntendedBaseCopyDirectory.Config, Path.Combine(ProjectDir, "Config"));
+		{
+			LocalDirectoryMappings.Add(EIntendedBaseCopyDirectory.Build, Path.Combine(ProjectDir, "Build"));
+			LocalDirectoryMappings.Add(EIntendedBaseCopyDirectory.Binaries, Path.Combine(ProjectDir, "Binaries"));
+			LocalDirectoryMappings.Add(EIntendedBaseCopyDirectory.Config, Path.Combine(ProjectDir, "Config"));
             LocalDirectoryMappings.Add(EIntendedBaseCopyDirectory.Content, Path.Combine(ProjectDir, "Content"));
             LocalDirectoryMappings.Add(EIntendedBaseCopyDirectory.Demos, Path.Combine(ProjectDir, "Demos"));
             LocalDirectoryMappings.Add(EIntendedBaseCopyDirectory.Profiling, Path.Combine(ProjectDir, "Profiling"));
diff --git a/Engine/Source/Programs/AutomationTool/Gauntlet/Platform/IOS/Gauntlet.TargetDeviceIOS.cs b/Engine/Source/Programs/AutomationTool/Gauntlet/Platform/IOS/Gauntlet.TargetDeviceIOS.cs
--- a/Engine/Source/Programs/AutomationTool/Gauntlet/Platform/IOS/Gauntlet.TargetDeviceIOS.cs
+++ b/Engine/Source/Programs/AutomationTool/Gauntlet/Platform/IOS/Gauntlet.TargetDeviceIOS.cs
@@ -591,9 +591,10 @@ namespace Gauntlet
 		}
 
 		public void PopulateDirectoryMappings(string ProjectDir)
-        {
-            LocalDirectoryMappings.Add(EIntendedBaseCopyDirectory.Binaries, Path.Combine(ProjectDir, "Binaries"));
-            LocalDirectoryMappings.Add(EIntendedBaseCopyDirectory.Config, Path.Combine(ProjectDir, "Config"));
+		{
+			LocalDirectoryMappings.Add(EIntendedBaseCopyDirectory.Build, Path.Combine(ProjectDir, "Build"));
+			LocalDirectoryMappings.Add(EIntendedBaseCopyDirectory.Binaries, Path.Combine(ProjectDir, "Binaries"));
+			LocalDirectoryMappings.Add(EIntendedBaseCopyDirectory.Config, Path.Combine(ProjectDir, "Config"));
             LocalDirectoryMappings.Add(EIntendedBaseCopyDirectory.Content, Path.Combine(ProjectDir, "Content"));
             LocalDirectoryMappings.Add(EIntendedBaseCopyDirectory.Demos, Path.Combine(ProjectDir, "Demos"));
             LocalDirectoryMappings.Add(EIntendedBaseCopyDirectory.Profiling, Path.Combine(ProjectDir, "Profiling"));
diff --git a/Engine/Source/Programs/AutomationTool/Gauntlet/Platform/Windows/Gauntlet.TargetDeviceWindows.cs b/Engine/Source/Programs/AutomationTool/Gauntlet/Platform/Windows/Gauntlet.TargetDeviceWindows.cs
--- a/Engine/Source/Programs/AutomationTool/Gauntlet/Platform/Windows/Gauntlet.TargetDeviceWindows.cs
+++ b/Engine/Source/Programs/AutomationTool/Gauntlet/Platform/Windows/Gauntlet.TargetDeviceWindows.cs
@@ -260,9 +260,10 @@ namespace Gauntlet
 
         // We care about UserDir in windows as some of the roles may require files going into user instead of build dir.
         public void PopulateDirectoryMappings(string BasePath, string UserDir)
-        {
-            LocalDirectoryMappings.Add(EIntendedBaseCopyDirectory.Binaries, Path.Combine(BasePath, "Binaries"));
-            LocalDirectoryMappings.Add(EIntendedBaseCopyDirectory.Config, Path.Combine(BasePath, "Saved", "Config"));
+		{
+			LocalDirectoryMappings.Add(EIntendedBaseCopyDirectory.Build, Path.Combine(BasePath, "Build"));
+			LocalDirectoryMappings.Add(EIntendedBaseCopyDirectory.Binaries, Path.Combine(BasePath, "Binaries"));
+			LocalDirectoryMappings.Add(EIntendedBaseCopyDirectory.Config, Path.Combine(BasePath, "Saved", "Config"));
             LocalDirectoryMappings.Add(EIntendedBaseCopyDirectory.Content, Path.Combine(BasePath, "Content"));
             LocalDirectoryMappings.Add(EIntendedBaseCopyDirectory.Demos, Path.Combine(UserDir, "Saved", "Demos"));
             LocalDirectoryMappings.Add(EIntendedBaseCopyDirectory.Profiling, Path.Combine(BasePath, "Saved", "Profiling"));
diff --git a/Engine/Source/Programs/AutomationTool/Gauntlet/Unreal/Base/Gauntlet.UnrealSession.cs b/Engine/Source/Programs/AutomationTool/Gauntlet/Unreal/Base/Gauntlet.UnrealSession.cs
--- a/Engine/Source/Programs/AutomationTool/Gauntlet/Unreal/Base/Gauntlet.UnrealSession.cs
+++ b/Engine/Source/Programs/AutomationTool/Gauntlet/Unreal/Base/Gauntlet.UnrealSession.cs
@@ -69,6 +69,10 @@ namespace Gauntlet
 		/// </summary>
 		public List<UnrealFileToCopy> FilesToCopy;
 
+		/// <summary>
+		/// 
+		/// </summary>
+		public List<EIntendedBaseCopyDirectory> AdditionalArtifactDirectories;
 		/// <summary>
 		/// Role device configuration 
 		/// </summary>
@@ -1073,9 +1077,22 @@ namespace Gauntlet
 				{
 					Log.Info("Skipping archival of assets for dev build");
 				}
-			}			
-
+			}
 
+			foreach (EIntendedBaseCopyDirectory ArtifactDir in InRunningRole.Role.AdditionalArtifactDirectories)
+			{
+				if (InRunningRole.AppInstance.Device.GetPlatformDirectoryMappings().ContainsKey(ArtifactDir))
+				{
+					string SourcePath = InRunningRole.AppInstance.Device.GetPlatformDirectoryMappings()[ArtifactDir];
+					var DirToCopy = new DirectoryInfo(SourcePath);
+					if (DirToCopy.Exists)
+					{
+						// Grab the final dir name to copy everything into so everything's not just going into root artifact dir.
+						string IntendedCopyLocation = Path.Combine(InDestArtifactPath, DirToCopy.Name);
+						Utils.SystemHelpers.CopyDirectory(SourcePath, IntendedCopyLocation);
+					}
+				}
+			}
 			// TODO REMOVEME- this should go elsewhere, likely a util that can be called or inserted by relevant test nodes.
 			if (IsServer == false)
 			{
diff --git a/Engine/Source/Programs/AutomationTool/Gauntlet/Unreal/Base/Gauntlet.UnrealTestConfiguration.cs b/Engine/Source/Programs/AutomationTool/Gauntlet/Unreal/Base/Gauntlet.UnrealTestConfiguration.cs
--- a/Engine/Source/Programs/AutomationTool/Gauntlet/Unreal/Base/Gauntlet.UnrealTestConfiguration.cs
+++ b/Engine/Source/Programs/AutomationTool/Gauntlet/Unreal/Base/Gauntlet.UnrealTestConfiguration.cs
@@ -35,7 +35,8 @@ namespace Gauntlet
     /// </summary>
     public enum EIntendedBaseCopyDirectory
     {
-        Binaries,
+        Build,
+		Binaries,
         Config,
         Content,
         Demos,
@@ -71,6 +72,7 @@ namespace Gauntlet
 			ExplicitClientCommandLine = string.Empty;
 			Controllers = new List<string>();
             FilesToCopy = new List<UnrealFileToCopy>();
+			AdditionalArtifactDirectories = new List<EIntendedBaseCopyDirectory>();
             RoleType = ERoleModifier.None;
 		}
 
@@ -104,6 +106,11 @@ namespace Gauntlet
 
         public List<UnrealFileToCopy> FilesToCopy { get; set; }
 
+		/// <summary>
+		/// Additional directories to 
+		/// </summary>
+		public List<EIntendedBaseCopyDirectory> AdditionalArtifactDirectories { get; set; }
+
 		/// <summary>
 		/// A map value passed in per server in case a test needs multiple servers on different maps.
 		/// </summary>
diff --git a/Engine/Source/Programs/AutomationTool/Gauntlet/Unreal/Base/Gauntlet.UnrealTestNode.cs b/Engine/Source/Programs/AutomationTool/Gauntlet/Unreal/Base/Gauntlet.UnrealTestNode.cs
--- a/Engine/Source/Programs/AutomationTool/Gauntlet/Unreal/Base/Gauntlet.UnrealTestNode.cs
+++ b/Engine/Source/Programs/AutomationTool/Gauntlet/Unreal/Base/Gauntlet.UnrealTestNode.cs
@@ -397,6 +397,7 @@ namespace Gauntlet
 
 					// copy over relevant settings from test role
                     SessionRole.FilesToCopy = TestRole.FilesToCopy;
+					SessionRole.AdditionalArtifactDirectories = TestRole.AdditionalArtifactDirectories;
 					SessionRole.ConfigureDevice = TestRole.ConfigureDevice;
 					SessionRole.MapOverride = TestRole.MapOverride;
 
