fixed gauntlet test restarting

diff --git a/Engine/Source/Programs/AutomationTool/Gauntlet/Framework/Gauntlet.TestExecutor.cs b/Engine/Source/Programs/AutomationTool/Gauntlet/Framework/Gauntlet.TestExecutor.cs
--- a/Engine/Source/Programs/AutomationTool/Gauntlet/Framework/Gauntlet.TestExecutor.cs
+++ b/Engine/Source/Programs/AutomationTool/Gauntlet/Framework/Gauntlet.TestExecutor.cs
@@ -657,19 +657,24 @@ namespace Gauntlet
 					TestResult NodeResult = TestInfo.TestNode.GetTestResult();
 					TestInfo.FinalResult = (TestInfo.FinalResult != TestResult.Invalid) ? TestInfo.FinalResult : NodeResult;
 
+					bool bCanFinalizeTest = true;
 					if (TestInfo.FinalResult == TestResult.WantRetry)
 					{
 						Log.Info("{0} requested retry. Cleaning up old test and relaunching", TestInfo);
 
 						DateTime OriginalStartTime = TestInfo.PostStartTime;
 
-						TestInfo.TestNode.RestartTest();
-
-						// Mark us as still running
-						TestInfo.CancellationReason = "";
-						TestInfo.FinalResult = TestResult.Invalid;
+						bool bIsRestarted = TestInfo.TestNode.RestartTest();
+						if (bIsRestarted)
+						{
+							// Mark us as still running
+							TestInfo.CancellationReason = "";
+							TestInfo.FinalResult = TestResult.Invalid;
+							bCanFinalizeTest = false;
+						}
 					}
-					else
+
+					if (bCanFinalizeTest)
 					{
 						Log.Info("{0} result={1}", TestInfo, TestInfo.FinalResult);
 
@@ -677,16 +682,16 @@ namespace Gauntlet
 						{
 							ReportTestSummary(TestInfo);
 						}
-					}
 
-					// now cleanup
-					try
-					{
-						TestInfo.TestNode.CleanupTest();
-					}
-					catch (System.Exception ex)
-					{
-						Log.Error("Test {0} threw an exception while cleaning up. Ex: {1}", TestInfo.TestNode.Name, ex.Message);
+						// now cleanup
+						try
+						{
+							TestInfo.TestNode.CleanupTest();
+						}
+						catch (System.Exception ex)
+						{
+							Log.Error("Test {0} threw an exception while cleaning up. Ex: {1}", TestInfo.TestNode.Name, ex.Message);
+						}
 					}
 
 					Log.Info("****************************************************************");
diff --git a/Engine/Source/Programs/AutomationTool/Gauntlet/Unreal/Base/Gauntlet.UnrealTestNode.cs b/Engine/Source/Programs/AutomationTool/Gauntlet/Unreal/Base/Gauntlet.UnrealTestNode.cs
--- a/Engine/Source/Programs/AutomationTool/Gauntlet/Unreal/Base/Gauntlet.UnrealTestNode.cs
+++ b/Engine/Source/Programs/AutomationTool/Gauntlet/Unreal/Base/Gauntlet.UnrealTestNode.cs
@@ -540,7 +540,13 @@ namespace Gauntlet
 		{
 			TestInstance = UnrealApp.RestartSession();
 
-			return TestInstance != null;
+			bool bWasRestarted = (TestInstance != null);
+			if (bWasRestarted)
+			{
+				MarkTestStarted();
+			}
+
+			return bWasRestarted;
 		}
 
 		/// <summary>
