Fail PGO node when processing data fails
#rnx

diff --git a/Engine/Source/Programs/AutomationTool/Gauntlet/Unreal/Base/Gauntlet.UnrealPGONode.cs b/Engine/Source/Programs/AutomationTool/Gauntlet/Unreal/Base/Gauntlet.UnrealPGONode.cs
--- a/Engine/Source/Programs/AutomationTool/Gauntlet/Unreal/Base/Gauntlet.UnrealPGONode.cs
+++ b/Engine/Source/Programs/AutomationTool/Gauntlet/Unreal/Base/Gauntlet.UnrealPGONode.cs
@@ -41,6 +41,7 @@ namespace Gauntlet
 	{
 		protected string LocalOutputDirectory;
 		private IPGOPlatform PGOPlatform;
+		bool ProcessPGODataFailed = false;
 
 		public PGONode(UnrealTestContext InContext) : base(InContext)
 		{
@@ -95,6 +96,16 @@ namespace Gauntlet
 		protected string ScreenshotDirectory;
 		protected bool ResultsGathered = false;
 
+		protected override TestResult GetUnrealTestResult()
+		{
+			if (ProcessPGODataFailed)
+			{
+				return TestResult.Failed;
+			}
+
+			return base.GetUnrealTestResult();
+		}
+
 		public override void TickTest()
 		{
 			base.TickTest();
@@ -113,6 +124,7 @@ namespace Gauntlet
 					}
 					catch (Exception Ex)
 					{
+						ProcessPGODataFailed = true;
 						Log.Error("Error getting PGO results: {0}", Ex);
 					}
 
