Whitelist some more build products with duplicate names.

#rnx

diff --git a/Engine/Source/Programs/AutomationTool/BuildGraph/TempStorage.cs b/Engine/Source/Programs/AutomationTool/BuildGraph/TempStorage.cs
--- a/Engine/Source/Programs/AutomationTool/BuildGraph/TempStorage.cs
+++ b/Engine/Source/Programs/AutomationTool/BuildGraph/TempStorage.cs
@@ -1108,15 +1108,27 @@ namespace AutomationTool
 		public static bool IsDuplicateBuildProduct(FileReference LocalFile)
 		{
 			string FileName = LocalFile.GetFileName();
-			if(FileName.Equals("AgentInterface.dll", StringComparison.OrdinalIgnoreCase) || FileName.Equals("AgentInterface.pdb", StringComparison.OrdinalIgnoreCase))
+			if (FileName.Equals("AgentInterface.dll", StringComparison.OrdinalIgnoreCase) || FileName.Equals("AgentInterface.pdb", StringComparison.OrdinalIgnoreCase))
 			{
 				return true;
 			}
-			if(FileName.Equals("dxil.dll", StringComparison.OrdinalIgnoreCase))
+			if (FileName.Equals("dxil.dll", StringComparison.OrdinalIgnoreCase))
 			{
 				return true;
 			}
-			if(FileName.Equals("dxcompiler.dll", StringComparison.OrdinalIgnoreCase))
+			if (FileName.Equals("dxcompiler.dll", StringComparison.OrdinalIgnoreCase))
+			{
+				return true;
+			}
+			if (FileName.Equals("embree.2.14.0.dll", StringComparison.OrdinalIgnoreCase))
+			{
+				return true;
+			}
+			if (FileName.Equals("tbb.dll", StringComparison.OrdinalIgnoreCase))
+			{
+				return true;
+			}
+			if (FileName.Equals("tbbmalloc.dll", StringComparison.OrdinalIgnoreCase))
 			{
 				return true;
 			}
