Modified Mac volume path for users vs builders

diff --git a/Engine/Source/Programs/AutomationTool/Gauntlet/Unreal/Utils/Gauntlet.UnrealHelpers.cs b/Engine/Source/Programs/AutomationTool/Gauntlet/Unreal/Utils/Gauntlet.UnrealHelpers.cs
--- a/Engine/Source/Programs/AutomationTool/Gauntlet/Unreal/Utils/Gauntlet.UnrealHelpers.cs
+++ b/Engine/Source/Programs/AutomationTool/Gauntlet/Unreal/Utils/Gauntlet.UnrealHelpers.cs
@@ -375,16 +375,18 @@ namespace Gauntlet
 		{
 			PlatformPath = Path;
 
-			if (BuildHostPlatform.Current.Platform == UnrealTargetPlatform.Mac)
+			if (BuildHostPlatform.Current.Platform == UnrealTargetPlatform.Mac || BuildHostPlatform.Current.Platform == UnrealTargetPlatform.Linux)
 			{
+				string PosixMountPath = CommandUtils.IsBuildMachine ? "/Volumes/epicgames.net/root" : "/Volumes/root";
+
 				if (!Path.Contains("P:"))
 				{
-					PlatformPath = Regex.Replace(Path, @"\\\\epicgames.net\\root", "/Volumes/epicgames.net/root", RegexOptions.IgnoreCase);
+					PlatformPath = Regex.Replace(Path, @"\\\\epicgames.net\\root", PosixMountPath, RegexOptions.IgnoreCase);
 				}
-				else 
+				else
 				{
-					PlatformPath = Regex.Replace(Path, "P:", "/Volumes/epicgames.net/root", RegexOptions.IgnoreCase);					
-				}				
+					PlatformPath = Regex.Replace(Path, "P:", PosixMountPath, RegexOptions.IgnoreCase);
+				}
 				
 				PlatformPath = PlatformPath.Replace(@"\", "/");
 			}
