Gauntlet workers expect same network mount path as build machines on Mac
#rnx

diff --git a/Engine/Source/Programs/AutomationTool/Gauntlet/Unreal/Utils/Gauntlet.UnrealHelpers.cs b/Engine/Source/Programs/AutomationTool/Gauntlet/Unreal/Utils/Gauntlet.UnrealHelpers.cs
--- a/Engine/Source/Programs/AutomationTool/Gauntlet/Unreal/Utils/Gauntlet.UnrealHelpers.cs
+++ b/Engine/Source/Programs/AutomationTool/Gauntlet/Unreal/Utils/Gauntlet.UnrealHelpers.cs
@@ -377,7 +377,7 @@ namespace Gauntlet
 
 			if (BuildHostPlatform.Current.Platform == UnrealTargetPlatform.Mac || BuildHostPlatform.Current.Platform == UnrealTargetPlatform.Linux)
 			{
-				string PosixMountPath = CommandUtils.IsBuildMachine ? "/Volumes/epicgames.net/root" : "/Volumes/root";
+				string PosixMountPath = (CommandUtils.IsBuildMachine || Globals.WorkerID != -1) ? "/Volumes/epicgames.net/root" : "/Volumes/root";
 
 				if (!Path.Contains("P:"))
 				{
