Fix some duplicate build product errors on Mac.

#rnx

diff --git a/Engine/Source/Programs/AutomationTool/BuildGraph/TempStorage.cs b/Engine/Source/Programs/AutomationTool/BuildGraph/TempStorage.cs
--- a/Engine/Source/Programs/AutomationTool/BuildGraph/TempStorage.cs
+++ b/Engine/Source/Programs/AutomationTool/BuildGraph/TempStorage.cs
@@ -1120,15 +1120,15 @@ namespace AutomationTool
 			{
 				return true;
 			}
-			if (FileName.Equals("embree.2.14.0.dll", StringComparison.OrdinalIgnoreCase))
+			if (FileName.Equals("embree.2.14.0.dll", StringComparison.OrdinalIgnoreCase) || FileName.Equals("libembree.2.14.0.dylib", StringComparison.OrdinalIgnoreCase))
 			{
 				return true;
 			}
-			if (FileName.Equals("tbb.dll", StringComparison.OrdinalIgnoreCase))
+			if (FileName.Equals("tbb.dll", StringComparison.OrdinalIgnoreCase) || FileName.Equals("libtbb.dylib", StringComparison.OrdinalIgnoreCase))
 			{
 				return true;
 			}
-			if (FileName.Equals("tbbmalloc.dll", StringComparison.OrdinalIgnoreCase))
+			if (FileName.Equals("tbbmalloc.dll", StringComparison.OrdinalIgnoreCase) || FileName.Equals("libtbbmalloc.dylib", StringComparison.OrdinalIgnoreCase))
 			{
 				return true;
 			}
