Add some more padding to OBB main vs patch file packing when exceeding OBB limit

#android

diff --git a/Engine/Source/Programs/AutomationTool/Android/AndroidPlatform.Automation.cs b/Engine/Source/Programs/AutomationTool/Android/AndroidPlatform.Automation.cs
--- a/Engine/Source/Programs/AutomationTool/Android/AndroidPlatform.Automation.cs
+++ b/Engine/Source/Programs/AutomationTool/Android/AndroidPlatform.Automation.cs
@@ -433,8 +433,8 @@ public class AndroidPlatform : Platform
 				FilesToObb = new List<FileReference>();
 
 				// Collect the filesize and place into Obb or Patch list
-				Int64 MainObbSize = 22;		// EOCD without comment
-				Int64 PatchObbSize = 22;	// EOCD without comment
+				Int64 MainObbSize = 22 + 4096;		// EOCD without comment + padding
+				Int64 PatchObbSize = 22 + 4096;		// EOCD without comment + padding
 				foreach (FileReference FileRef in FilesForObb)
 				{
 					FileInfo LocalFileInfo = new FileInfo(FileRef.FullName);
