Gauntlet: Fix issue with case senstivity on perf model constraint
#rnx

diff --git a/Engine/Source/Programs/AutomationTool/Gauntlet/Framework/Devices/Gauntlet.DevicePool.cs b/Engine/Source/Programs/AutomationTool/Gauntlet/Framework/Devices/Gauntlet.DevicePool.cs
--- a/Engine/Source/Programs/AutomationTool/Gauntlet/Framework/Devices/Gauntlet.DevicePool.cs
+++ b/Engine/Source/Programs/AutomationTool/Gauntlet/Framework/Devices/Gauntlet.DevicePool.cs
@@ -90,8 +90,19 @@ namespace Gauntlet
 
 		public bool Check(DeviceDefinition DeviceDef)
 		{
-			bool Match = (PerfSpec == EPerfSpec.Unspecified) ? DeviceDef.Model == Model : PerfSpec == DeviceDef.PerfSpec;
-			return Platform == DeviceDef.Platform && (IsIdentity() || Match );
+			if (Platform != DeviceDef.Platform)
+			{
+				return false;
+			}
+
+			if (IsIdentity())
+			{
+				return true;
+			}
+
+			bool ModelMatch = Model == string.Empty ? true : Model.Equals(DeviceDef.Model, StringComparison.InvariantCultureIgnoreCase);
+			bool PerfMatch = (PerfSpec == EPerfSpec.Unspecified) ? true : PerfSpec == DeviceDef.PerfSpec;
+			return ModelMatch && PerfMatch;
 		}
 
 		public bool Equals(UnrealTargetConstraint Other)
@@ -103,7 +114,7 @@ namespace Gauntlet
 
 			if (ReferenceEquals(this, Other)) return true;
 
-			return Other.Platform == Platform && Other.PerfSpec == PerfSpec && Other.Model == Model;
+			return Other.Platform == Platform && Other.PerfSpec == PerfSpec && Other.Model.Equals(Model, StringComparison.InvariantCultureIgnoreCase);
 		}
 
 		public override bool Equals(object Obj)
