HTML5 - copytorelease

HTML5 - mac python: include dynamic lib path to build

HTML5 - fix deprecated emscripten function call in SDL_emscriptenaudio.c with latest toolchain

HTML5 - adjust AutomationTool for HTML5 to handle "alternative" project names

HTML5 - get launch onto current editor map working

#rnx

diff --git a/Engine/Source/Programs/AutomationTool/HTML5/HTML5Platform.Automation.cs b/Engine/Source/Programs/AutomationTool/HTML5/HTML5Platform.Automation.cs
--- a/Engine/Source/Programs/AutomationTool/HTML5/HTML5Platform.Automation.cs
+++ b/Engine/Source/Programs/AutomationTool/HTML5/HTML5Platform.Automation.cs
@@ -70,6 +70,11 @@ public class HTML5Platform : Platform
 		string UE4GameBasename = Path.Combine(PackagePath, _ProjectFilename);
 		string ProjectBasename = Path.Combine(PackagePath, Params.ShortProjectName + _ProjectNameExtra);
 
+		if (Params.IsCodeBasedProject && Params.HasClientCookedTargets)
+		{
+			UE4GameBasename = Params.ClientCookedTargets[0];
+		}
+
 		// ----------------------------------------
 		// packaging
 		if (HTMLPakAutomation.CanCreateMapPaks(Params))
@@ -176,18 +181,18 @@ public class HTML5Platform : Platform
 			BuildPath = LocalBuildPath;
 			TemplateFile = CombinePaths(BuildPath, "project_template.html");
 		}
-		GenerateFileFromTemplate(TemplateFile, ProjectBasename + ".html", Params, ConfigCache);
+		GenerateFileFromTemplate(TemplateFile, ProjectBasename + ".html", Params, SC, ConfigCache);
 
 		TemplateFile = CombinePaths(BuildPath, "project_template.js");
 		if ( File.Exists(TemplateFile) )
 		{
-			GenerateFileFromTemplate(TemplateFile, ProjectBasename + ".UE4.js", Params, ConfigCache);
+			GenerateFileFromTemplate(TemplateFile, ProjectBasename + ".UE4.js", Params, SC, ConfigCache);
 		}
 
 		TemplateFile = CombinePaths(BuildPath, "project_template.css");
 		if ( File.Exists(TemplateFile) )
 		{
-			GenerateFileFromTemplate(TemplateFile, ProjectBasename + ".css", Params, ConfigCache);
+			GenerateFileFromTemplate(TemplateFile, ProjectBasename + ".css", Params, SC, ConfigCache);
 		}
 
 
@@ -310,12 +315,12 @@ public class HTML5Platform : Platform
 		}
 	}
 
-	protected void GenerateFileFromTemplate(string InTemplateFile, string InOutputFile, ProjectParams Params, ConfigHierarchy ConfigCache)
+	protected void GenerateFileFromTemplate(string InTemplateFile, string InOutputFile, ProjectParams Params, DeploymentContext SC, ConfigHierarchy ConfigCache)
 	{
 		bool IsContentOnly = !Params.IsCodeBasedProject;
 		string ProjectConfiguration = Params.ClientConfigsToBuild[0].ToString();
 
-		string UE4GameName = IsContentOnly ? "UE4Game" : Params.ShortProjectName;
+		string UE4GameName = IsContentOnly ? "UE4Game" : (Params.HasClientCookedTargets ? Params.ClientCookedTargets[0] : Params.ShortProjectName);
 		string ProjectName = Params.ShortProjectName;
 		if (ProjectConfiguration != "Development")
 		{
@@ -368,6 +373,11 @@ public class HTML5Platform : Platform
 				if (LineStr.Contains("%UE4CMDLINE%"))
 				{
 					string ArgumentString = "'../../../" + Params.ShortProjectName + "/" + Params.ShortProjectName + ".uproject',";
+
+					if ( Params.MapToRun.Length > 0 )
+					{
+						ArgumentString += "'" + Params.MapToRun + "',";
+					}
 					ArgumentString += "'-stdout',"; // suppress double printing to console.log
 					LineStr = LineStr.Replace("%UE4CMDLINE%", ArgumentString);
 				}
@@ -493,6 +503,10 @@ public class HTML5Platform : Platform
 		string UE4GameBasename = Path.GetFileNameWithoutExtension(Params.GetProjectExeForPlatform(UnrealTargetPlatform.HTML5).ToString());
 		string ProjectBasename = Params.ShortProjectName;
 		string ProjectConfiguration = Params.ClientConfigsToBuild[0].ToString();
+		if (Params.IsCodeBasedProject && Params.HasClientCookedTargets)
+		{
+			UE4GameBasename = Params.ClientCookedTargets[0];
+		}
 		if (ProjectConfiguration != "Development")
 		{
 			UE4GameBasename += "-HTML5-" + ProjectConfiguration;
diff --git a/Engine/Source/Programs/UnrealBuildTool/Platform/HTML5/HTML5SDKInfo.cs b/Engine/Source/Programs/UnrealBuildTool/Platform/HTML5/HTML5SDKInfo.cs
--- a/Engine/Source/Programs/UnrealBuildTool/Platform/HTML5/HTML5SDKInfo.cs
+++ b/Engine/Source/Programs/UnrealBuildTool/Platform/HTML5/HTML5SDKInfo.cs
@@ -489,6 +489,16 @@ namespace UnrealBuildTool
 			return new FileReference(PYTHON);
 		}
 
+		/// <summary>
+		///
+		/// </summary>
+		/// <returns></returns>
+		public static string MacPythonLib() // UE-75402
+		{
+			string UE4PythonPath = FileReference.Combine(UnrealBuildTool.EngineDirectory, "Binaries", "ThirdParty", "Python").FullName;
+			return Path.Combine(UE4PythonPath, "Mac", "lib", "python2.7", "lib-dynload");
+		}
+
 		/// <summary>
 		///
 		/// </summary>
diff --git a/Engine/Source/Programs/UnrealBuildTool/Platform/HTML5/HTML5ToolChain.cs b/Engine/Source/Programs/UnrealBuildTool/Platform/HTML5/HTML5ToolChain.cs
--- a/Engine/Source/Programs/UnrealBuildTool/Platform/HTML5/HTML5ToolChain.cs
+++ b/Engine/Source/Programs/UnrealBuildTool/Platform/HTML5/HTML5ToolChain.cs
@@ -221,6 +221,10 @@ namespace UnrealBuildTool
 				// Packaging on Window needs this - zap any existing HOME environment variables to prevent any accidental pick ups
 				Environment.SetEnvironmentVariable("HOME", "");
 			}
+			if (BuildHostPlatform.Current.Platform == UnrealTargetPlatform.Mac)
+			{
+				Environment.SetEnvironmentVariable("LD_LIBRARY_PATH", HTML5SDKInfo.MacPythonLib()); // UE-75402
+			}
 			return Result;
 		}
 
