Generated SymbolizeCrashDump for Android was invalid (Contributed by [REDACTED_PERSON])

#PR #5959
#ue4
#android

diff --git a/Engine/Source/Programs/AutomationTool/Android/AndroidPlatform.Automation.cs b/Engine/Source/Programs/AutomationTool/Android/AndroidPlatform.Automation.cs
--- a/Engine/Source/Programs/AutomationTool/Android/AndroidPlatform.Automation.cs
+++ b/Engine/Source/Programs/AutomationTool/Android/AndroidPlatform.Automation.cs
@@ -845,7 +845,7 @@ public class AndroidPlatform : Platform
 						"setlocal",
 						"set NDK_ROOT=%ANDROID_NDK_ROOT%",
 						"if \"%ANDROID_NDK_ROOT%\"==\"\" set NDK_ROOT=\""+Environment.GetEnvironmentVariable("ANDROID_NDK_ROOT")+"\"",
-						"set NDKSTACK=%NDK_ROOT%\ndk-stack.cmd",
+						"set NDKSTACK=%NDK_ROOT%\\ndk-stack.cmd",
 						"",
 						"%NDKSTACK% -sym "+GetFinalSymbolizedSODirectory(ApkName, SC, Architecture, GPUArchitecture)+" -dump \"%1\" > "+ Params.ShortProjectName+"_SymbolizedCallStackOutput.txt",
 						"",
