Fix for HoloLens preventing UBT from compiling in Linux

diff --git a/Engine/Source/Programs/AutomationTool/AutomationUtils/LinuxHostPlatform.cs b/Engine/Source/Programs/AutomationTool/AutomationUtils/LinuxHostPlatform.cs
--- a/Engine/Source/Programs/AutomationTool/AutomationUtils/LinuxHostPlatform.cs
+++ b/Engine/Source/Programs/AutomationTool/AutomationUtils/LinuxHostPlatform.cs
@@ -82,6 +82,7 @@ namespace AutomationTool
 				CommandLine = (String.IsNullOrEmpty(CommandLine) ? "" : CommandLine) + " /verbosity:quiet /nologo";
 				// Pass #define MONO to all the automation scripts (see XboxOne)
 				CommandLine += " /p:DefineConstants=MONO";
+				CommandLine += " /p:DefineConstants=__MonoCS__";
 				// Some projects have TargetFrameworkProfile=Client which causes warnings on Linux
 				// so force it to empty.
 				CommandLine += " /p:TargetFrameworkProfile=";
