Properly handle ; character in command line params when launching UBT from UAT scripts

diff --git a/Engine/Source/Programs/AutomationTool/AutomationUtils/MacHostPlatform.cs b/Engine/Source/Programs/AutomationTool/AutomationUtils/MacHostPlatform.cs
--- a/Engine/Source/Programs/AutomationTool/AutomationUtils/MacHostPlatform.cs
+++ b/Engine/Source/Programs/AutomationTool/AutomationUtils/MacHostPlatform.cs
@@ -118,7 +118,7 @@ namespace AutomationTool
 				else
 				{
 					// It's a C# app, so run it with Mono
-					CommandLine = " -c 'mono \"" + AppName + "\" " + (String.IsNullOrEmpty(CommandLine) ? "" : CommandLine.Replace("\\", "\\\\")) + "'";
+					CommandLine = " -c 'mono \"" + AppName + "\" " + (String.IsNullOrEmpty(CommandLine) ? "" : CommandLine.Replace("\\", "\\\\").Replace(";", "\\\\;")) + "'";
 					AppName = "sh";
 					Options &= ~CommandUtils.ERunOptions.AppMustExist;
 				}
