Gauntlet: sleep thread once Android activity has exited to allow for log to flush (speculative fix)
#rnx

diff --git a/Engine/Source/Programs/AutomationTool/Gauntlet/Platform/Android/Gauntlet.TargetDeviceAndroid.cs b/Engine/Source/Programs/AutomationTool/Gauntlet/Platform/Android/Gauntlet.TargetDeviceAndroid.cs
--- a/Engine/Source/Programs/AutomationTool/Gauntlet/Platform/Android/Gauntlet.TargetDeviceAndroid.cs
+++ b/Engine/Source/Programs/AutomationTool/Gauntlet/Platform/Android/Gauntlet.TargetDeviceAndroid.cs
@@ -120,7 +120,8 @@ namespace Gauntlet
 			if (bHasExited)
 			{
 				ActivityExited = true;
-				// Make sure entire activity log has been captured
+				// The activity has exited, make sure entire activity log has been captured, sleep to allow time for the log to flush
+				Thread.Sleep(5000);
 				UpdateCachedLog(true);
 				Log.VeryVerbose("{0}: process exited, Activity running={1}, Activity in foreground={2} ", ToString(), bActivityPresent.ToString(), bActivityInForeground.ToString());
 			}
