Temporarily reverted the optimization that makes UAT launch UBT via a shell script, as it's causing issues that are not easy to solve quickly (UE-76351, UE-76600)

diff --git a/Engine/Source/Programs/AutomationTool/AutomationUtils/MacHostPlatform.cs b/Engine/Source/Programs/AutomationTool/AutomationUtils/MacHostPlatform.cs
--- a/Engine/Source/Programs/AutomationTool/AutomationUtils/MacHostPlatform.cs
+++ b/Engine/Source/Programs/AutomationTool/AutomationUtils/MacHostPlatform.cs
@@ -118,8 +118,8 @@ namespace AutomationTool
 				else
 				{
 					// It's a C# app, so run it with Mono
-					CommandLine = " -c 'mono \"" + AppName + "\" " + (String.IsNullOrEmpty(CommandLine) ? "" : CommandLine.Replace("\\", "\\\\").Replace(";", "\\\\;")) + "'";
-					AppName = "sh";
+					CommandLine = "\"" + AppName + "\" " + (String.IsNullOrEmpty(CommandLine) ? "" : CommandLine);
+					AppName = "mono";
 					Options &= ~CommandUtils.ERunOptions.AppMustExist;
 				}
 			}
