Gauntlet: Add support for handling report creation errors

diff --git a/Engine/Source/Programs/AutomationTool/Gauntlet/Unreal/Base/Gauntlet.UnrealTestNode.cs b/Engine/Source/Programs/AutomationTool/Gauntlet/Unreal/Base/Gauntlet.UnrealTestNode.cs
--- a/Engine/Source/Programs/AutomationTool/Gauntlet/Unreal/Base/Gauntlet.UnrealTestNode.cs
+++ b/Engine/Source/Programs/AutomationTool/Gauntlet/Unreal/Base/Gauntlet.UnrealTestNode.cs
@@ -14,6 +14,7 @@ using System.Security.Cryptography;
 
 namespace Gauntlet
 {
+
 	public abstract class UnrealTestNode<TConfigClass> : BaseTest, IDisposable
 		where TConfigClass : UnrealTestConfiguration, new()
 	{
@@ -635,32 +636,73 @@ namespace Gauntlet
 				Log.Warning("Failed to release devices. {0}", Ex);
 			}
 
+			string Message = string.Empty;
+
 			try
 			{
 				CreateReport(GetTestResult(), Context, Context.BuildInfo, SessionArtifacts, ArtifactPath);
 			}
 			catch (Exception Ex)
 			{
+				CreateReportFailed = true;
+				Message = Ex.Message;				
 				Log.Warning("Failed to save completion report. {0}", Ex);
 			}
 
+			if (CreateReportFailed)
+			{
+				try
+				{
+					HandleCreateReportFailure(Context, Message);
+				}
+				catch (Exception Ex)
+				{
+					Log.Warning("Failed to handle completion report failure. {0}", Ex);
+				}
+			}
+
 			try
 			{
 				SubmitToDashboard(GetTestResult(), Context, Context.BuildInfo, SessionArtifacts, ArtifactPath);
 			}
 			catch (Exception Ex)
-			{
+			{				
 				Log.Warning("Failed to submit results to dashboard. {0}", Ex);
 			}
 		}
 
+		/// <summary>
+		/// Called when report creation fails, by default logs warning with failure message
+		/// </summary>
+		protected virtual void HandleCreateReportFailure(UnrealTestContext Context, string Message = "")
+		{
+			if (string.IsNullOrEmpty(Message))
+			{
+				Message = string.Format("See Gauntlet.log for details");
+			}
+
+			if (Globals.IsWorker)
+			{
+				// log for worker to parse context
+				Log.Info("GauntletWorker:CreateReportFailure:{0}", Context.WorkerJobID);
+			}
+
+			Log.Warning("CreateReport Failed: {0}", Message);
+		}
+
+		/// <summary>
+		/// Whether creating the test report failed
+		/// </summary>
+		public bool CreateReportFailed { get; protected set; }
+		 
+
 		/// <summary>
 		/// Optional function that is called on test completion and gives an opportunity to create a report
 		/// </summary>
 		/// <param name="Result"></param>
 		/// <param name="Contex"></param>
 		/// <param name="Build"></param>
-		public virtual void CreateReport(TestResult Result, UnrealTestContext Contex, UnrealBuildSource Build, IEnumerable<UnrealRoleArtifacts> Artifacts, string ArtifactPath)
+		public virtual void CreateReport(TestResult Result, UnrealTestContext Context, UnrealBuildSource Build, IEnumerable<UnrealRoleArtifacts> Artifacts, string ArtifactPath)
 		{
 		}
 
