#UE4 Create the plugin manifest after directories have been remapped

diff --git a/Engine/Source/Programs/AutomationTool/Scripts/CopyBuildToStagingDirectory.Automation.cs b/Engine/Source/Programs/AutomationTool/Scripts/CopyBuildToStagingDirectory.Automation.cs
--- a/Engine/Source/Programs/AutomationTool/Scripts/CopyBuildToStagingDirectory.Automation.cs
+++ b/Engine/Source/Programs/AutomationTool/Scripts/CopyBuildToStagingDirectory.Automation.cs
@@ -402,6 +402,7 @@ public partial class Project : CommandUtils
 			// stage all the previously staged files
 			SC.StageFiles(StagedFileType.NonUFS, DirectoryReference.Combine(SC.ProjectRoot, "Saved", "SharedIterativeBuild", SC.CookPlatform, "Staged"), StageFilesSearch.AllDirectories, StagedDirectoryReference.Root); // remap to the root directory
 		}
+		bool bCreatePluginManifest = false;
 		if (Params.HasDLCName)
 		{
 			// Making a plugin
@@ -516,10 +517,7 @@ public partial class Project : CommandUtils
 				}
 			}
 
-			if (Params.UsePak(SC.StageTargetPlatform))
-			{
-				CreatePluginManifest(SC, SC.FilesToStage.UFSFiles, StagedFileType.UFS, Params.DLCFile.GetFileNameWithoutAnyExtensions());
-			}
+			bCreatePluginManifest = true;
 		}
 		else
 		{
@@ -825,18 +823,7 @@ public partial class Project : CommandUtils
 					}
 				}
 
-				// Generate a plugin manifest if we're using a pak file and not creating a mod. Mods can be enumerated independently by users copying them into the Mods directory.
-				if (Params.UsePak(SC.StageTargetPlatform))
-				{
-					if (Params.HasDLCName)
-					{
-						CreatePluginManifest(SC, SC.FilesToStage.NonUFSFiles, StagedFileType.NonUFS, Params.DLCFile.GetFileNameWithoutExtension());
-					}
-					else
-					{
-						CreatePluginManifest(SC, SC.FilesToStage.UFSFiles, StagedFileType.UFS, Params.ShortProjectName);
-					}
-				}
+				bCreatePluginManifest = true;
 			}
 			else
 			{
@@ -887,6 +874,20 @@ public partial class Project : CommandUtils
 		SC.FilesToStage.NonUFSDebugFiles = SC.FilesToStage.NonUFSDebugFiles.ToDictionary(x => ApplyDirectoryRemap(SC, x.Key), x => x.Value);
 		SC.FilesToStage.NonUFSSystemFiles = SC.FilesToStage.NonUFSSystemFiles.ToDictionary(x => ApplyDirectoryRemap(SC, x.Key), x => x.Value);
 
+		// Create plugin manifests after the directory mappings
+		if (bCreatePluginManifest && Params.UsePak(SC.StageTargetPlatform))
+		{
+			// Generate a plugin manifest if we're using a pak file and not creating a mod. Mods can be enumerated independently by users copying them into the Mods directory.
+			if (Params.HasDLCName)
+			{
+				CreatePluginManifest(SC, SC.FilesToStage.UFSFiles, StagedFileType.UFS, Params.DLCFile.GetFileNameWithoutAnyExtensions());
+			}
+			else
+			{
+				CreatePluginManifest(SC, SC.FilesToStage.UFSFiles, StagedFileType.UFS, Params.ShortProjectName);
+			}
+		}
+
 		// Make all the filenames lowercase
 		if (SC.StageTargetPlatform.DeployLowerCaseFilenames())
 		{
