Added helper UAT script to create a p4config file. Solves issue in Mac with tools that issue P4 commands directly.

diff --git a/Engine/Source/Programs/AutomationTool/AutomationUtils/AutomationUtils.Automation.csproj b/Engine/Source/Programs/AutomationTool/AutomationUtils/AutomationUtils.Automation.csproj
--- a/Engine/Source/Programs/AutomationTool/AutomationUtils/AutomationUtils.Automation.csproj
+++ b/Engine/Source/Programs/AutomationTool/AutomationUtils/AutomationUtils.Automation.csproj
@@ -92,6 +92,7 @@
     <Compile Include="UE4Build.cs" />
     <Compile Include="Utils.cs" />
     <Compile Include="WindowsHostPlatform.cs" />
+    <Compile Include="P4WriteConfig.cs" />
   </ItemGroup>
   <ItemGroup>
     <ProjectReference Include="..\..\DotNETCommon\DotNETUtilities\DotNETUtilities.csproj">
diff --git a/Engine/Source/Programs/AutomationTool/AutomationUtils/P4WriteConfig.cs b/Engine/Source/Programs/AutomationTool/AutomationUtils/P4WriteConfig.cs
new file mode 100644
--- /dev/null
+++ b/Engine/Source/Programs/AutomationTool/AutomationUtils/P4WriteConfig.cs
@@ -0,0 +1,111 @@
+using System;
+using System.Text;
+using System.IO;
+using Tools.DotNETCommon;
+using System.Collections.Generic;
+using System.Linq;
+
+namespace AutomationTool
+{
+	[Help("Auto-detects P4 settings based on the current path and creates a p4config file with the relevant settings.")]
+	[Help("setignore", "Adds a P4IGNORE to the default file (Engine/Extras/Perforce/p4ignore)")]
+	[Help("path=<path>", "Write to a path other than the current directory")]
+	[Help("p4port=<server:port>", "Optional hint/override of the server to use during lookup")]
+	[Help("p4user=<server:port>", "Optional hint/override of the username to use during lookup")]
+	public class P4WriteConfig: BuildCommand
+	{
+		public override ExitCode Execute()
+		{
+			Log.TraceInformation("Setting up Perforce environment.");
+
+			// User can specify these to help auto detection
+			string Port = ParseParamValue("p4port", "");
+			string User = ParseParamValue("p4user", "");
+
+			bool SetIgnore = ParseParam("setignore");
+			bool ListOnly = ParseParam("listonly");
+
+			string OutputPath = ParseParamValue("path", "");
+
+			// apply any hints
+			if (!string.IsNullOrEmpty(Port))
+			{
+				Environment.SetEnvironmentVariable(EnvVarNames.P4Port, Port);
+			}
+
+			if (!string.IsNullOrEmpty(User))
+			{
+				Environment.SetEnvironmentVariable(EnvVarNames.User, User);
+			}
+
+			// try to init P4
+			try
+			{
+				CommandUtils.InitP4Environment();
+				CommandUtils.InitDefaultP4Connection();
+			}
+			catch (Exception Ex)
+			{ 		
+				Log.TraceError("Unable to find matching Perforce info. If the below does not help try P4CreateConfig -p4port=<server:port> and -p4user=<username> to supply more info");
+				Log.TraceError("{0}", Ex.Message);
+				return ExitCode.Error_Arguments;
+			}
+
+			// store all our settings
+			Dictionary<string, string> P4Config = new Dictionary<string, string>();
+
+			P4Config["P4PORT"] = P4Env.ServerAndPort;
+			P4Config["P4USER"] = P4Env.User;
+			P4Config["P4CLIENT"] = P4Env.Client;
+
+			if (SetIgnore)
+			{
+				string IgnorePath = Path.Combine(CommandUtils.EngineDirectory.ToString(), "Extras", "Perforce", "p4ignore");
+				P4Config["P4IGNORE"] = IgnorePath;
+			}
+
+			string P4Settings = string.Join("\n", P4Config.Keys.Select(K => string.Format("{0}={1}", K, P4Config[K])));
+
+			string DefaultPath = Environment.CurrentDirectory;
+
+			if (!string.IsNullOrEmpty(OutputPath))
+			{
+				if (!DirectoryExists(OutputPath) && !FileExists(OutputPath))
+				{
+					throw new AutomationException("Path {0} does not exist.", OutputPath);
+				}
+				DefaultPath = OutputPath;
+			}
+			else
+			{
+				OutputPath = Environment.CurrentDirectory;
+			}
+
+			if (!File.Exists(OutputPath))
+			{
+				OutputPath = Path.Combine(OutputPath, "p4config.txt");
+			}
+
+			Console.WriteLine("***\nWriting_\n{0}\nto - {1}\n***", P4Settings, OutputPath);
+
+			if (!ListOnly)
+			{
+				File.WriteAllText(OutputPath, P4Settings);
+
+				string OutputFile = Path.GetFileName(OutputPath);
+
+				Log.TraceInformation("Wrote P4 settings to {0}", OutputPath);
+
+				P4.P4(string.Format("set P4CONFIG={0}", OutputFile));
+				Log.TraceInformation("set P4CONFIG={0}", OutputFile);
+			}
+			else
+			{
+				Log.TraceInformation("Skipped write");
+			}
+
+
+			return ExitCode.Success;
+		}		
+	}
+}
